Moving Mover, MoverExamples, MoverTests out of Restricted engine folder

#virtualized

diff --git a/Engine/Plugins/Experimental/Mover/Config/BaseMover.ini b/Engine/Plugins/Experimental/Mover/Config/BaseMover.ini
new file mode 100644
--- /dev/null
+++ b/Engine/Plugins/Experimental/Mover/Config/BaseMover.ini
@@ -0,0 +1,8 @@
+[CoreRedirects]
++ClassRedirects=(OldName="/Script/Mover.KinematicMoverComponent",NewName="/Script/Mover.MoverComponent")
++StructRedirects=(OldName="/Script/Mover.KinematicMoverInputCmdContext",NewName="/Script/Mover.MoverInputCmdContext")
++StructRedirects=(OldName="/Script/Mover.KinematicMoverSyncState",NewName="/Script/Mover.MoverSyncState")
++StructRedirects=(OldName="/Script/Mover.KinematicMoverAuxStateContext",NewName="/Script/Mover.MoverAuxStateContext")
++StructRedirects=(OldName="/Script/Mover.KinematicMoverTickStartData",NewName="/Script/Mover.MoverTickStartData")
++StructRedirects=(OldName="/Script/Mover.KinematicMoverTickEndData",NewName="/Script/Mover.MoverTickEndData")
+
diff --git a/Engine/Plugins/Experimental/Mover/README.md b/Engine/Plugins/Experimental/Mover/README.md
new file mode 100644
--- /dev/null
+++ b/Engine/Plugins/Experimental/Mover/README.md
@@ -0,0 +1,11 @@
+# Mover Plugin
+
+Mover is an Unreal Engine plugin to support movement of actors with rollback networking, using the Network Prediction Plugin. This plugin is the potential successor to Character Movement Component. The goal is to allow gameplay developers to focus on crafting motion without having to be experts in networking.
+
+The Mover plugin is Experimental. Many features are incomplete or missing. APIs and data formats are subject to change at any time.
+
+## Getting Started
+
+## Examples
+
+## Known Issues
\ No newline at end of file
diff --git a/Engine/Plugins/Experimental/Mover/Source/Mover/Private/Backends/MoverNetworkPhysicsLiaison.cpp b/Engine/Plugins/Experimental/Mover/Source/Mover/Private/Backends/MoverNetworkPhysicsLiaison.cpp
new file mode 100644
--- /dev/null
+++ b/Engine/Plugins/Experimental/Mover/Source/Mover/Private/Backends/MoverNetworkPhysicsLiaison.cpp
@@ -0,0 +1,1045 @@
+// Copyright Epic Games, Inc. All Rights Reserved.
+
+#include "Backends/MoverNetworkPhysicsLiaison.h"
+
+#include "Chaos/Character/CharacterGroundConstraint.h"
+#include "Chaos/PhysicsObject.h"
+#include "Chaos/PhysicsObjectInternalInterface.h"
+#include "GameFramework/Actor.h"
+#include "MovementModeStateMachine.h"
+#include "Kinematic/Settings/CommonLegacyMovementSettings.h"
+#include "PBDRigidsSolver.h"
+#include "Components/ShapeComponent.h"
+#include "Physics/Experimental/PhysScene_Chaos.h"
+#include "PhysicsMover/Modes/PhysicsDrivenFallingMode.h"
+#include "PhysicsMover/Modes/PhysicsDrivenWalkingMode.h"
+#include "PhysicsMover/PhysicsMovementUtils.h"
+#include "PhysicsMover/PhysicsMoverManager.h"
+#include "PhysicsProxy/CharacterGroundConstraintProxy.h"
+
+//////////////////////////////////////////////////////////////////////////
+
+extern FPhysicsDrivenMotionDebugParams GPhysicsDrivenMotionDebugParams;
+
+//////////////////////////////////////////////////////////////////////////
+// FNetworkPhysicsMoverInputs
+
+void FNetworkPhysicsMoverInputs::ApplyData(UActorComponent* NetworkComponent) const
+{
+	if (NetworkComponent)
+	{
+		if (UMoverNetworkPhysicsLiaisonComponent* LiaisonComp = Cast<UMoverNetworkPhysicsLiaisonComponent>(NetworkComponent))
+		{
+			LiaisonComp->SetCurrentInputData(InputCmdContext);
+		}
+	}
+}
+
+void FNetworkPhysicsMoverInputs::BuildData(const UActorComponent* NetworkComponent)
+{
+	if (NetworkComponent)
+	{
+		if (const UMoverNetworkPhysicsLiaisonComponent* LiaisonComp = Cast<const UMoverNetworkPhysicsLiaisonComponent>(NetworkComponent))
+		{
+			LiaisonComp->GetCurrentInputData(InputCmdContext);
+		}
+	}
+}
+
+bool FNetworkPhysicsMoverInputs::NetSerialize(FArchive& Ar, class UPackageMap* PackageMap, bool& bOutSuccess)
+{
+	FNetworkPhysicsData::SerializeFrames(Ar);
+
+	if (PackageMap)
+	{
+		InputCmdContext.NetSerialize(FNetSerializeParams(Ar));
+		bOutSuccess = true;
+	}
+	else
+	{
+		bOutSuccess = false;
+	}
+	
+	return bOutSuccess;
+}
+
+void FNetworkPhysicsMoverInputs::InterpolateData(const FNetworkPhysicsData& MinData, const FNetworkPhysicsData& MaxData)
+{
+	const FNetworkPhysicsMoverInputs& MinDataInput = static_cast<const FNetworkPhysicsMoverInputs&>(MinData);
+	const FNetworkPhysicsMoverInputs& MaxDataInput = static_cast<const FNetworkPhysicsMoverInputs&>(MaxData);
+
+	const float LerpFactor = (LocalFrame - MinDataInput.LocalFrame) / (MaxDataInput.LocalFrame - MinDataInput.LocalFrame);
+
+	const FKinematicDefaultInputs* MinInput = MinDataInput.InputCmdContext.InputCollection.FindDataByType<FKinematicDefaultInputs>();
+	const FKinematicDefaultInputs* MaxInput = MaxDataInput.InputCmdContext.InputCollection.FindDataByType<FKinematicDefaultInputs>();
+
+	FKinematicDefaultInputs& LocalInput = InputCmdContext.InputCollection.FindOrAddMutableDataByType<FKinematicDefaultInputs>();
+
+	if (MinInput && MaxInput)
+	{
+		// Note, this ignores movement base as this is not used by the physics mover
+		const FKinematicDefaultInputs* ClosestInputs = LerpFactor < 0.5f ? MinInput : MaxInput;
+		LocalInput.bIsJumpJustPressed = ClosestInputs->bIsJumpJustPressed;
+		LocalInput.bIsJumpPressed = ClosestInputs->bIsJumpPressed;
+		LocalInput.SuggestedMovementMode = ClosestInputs->SuggestedMovementMode;
+
+		LocalInput.SetMoveInput(ClosestInputs->GetMoveInputType(), FMath::Lerp(MinInput->GetMoveInput(), MaxInput->GetMoveInput(), LerpFactor));
+		LocalInput.OrientationIntent = FMath::Lerp(MinInput->OrientationIntent, MaxInput->OrientationIntent, LerpFactor);
+		LocalInput.ControlRotation = FMath::Lerp(MinInput->ControlRotation, MaxInput->ControlRotation, LerpFactor);
+
+	}
+	else if (MinInput)
+	{
+		LocalInput = *MinInput;
+	}
+	else if (MaxInput)
+	{
+		LocalInput = *MaxInput;
+	}
+}
+
+void FNetworkPhysicsMoverInputs::MergeData(const FNetworkPhysicsData& FromData)
+{
+	const FNetworkPhysicsMoverInputs& FromDataInput = static_cast<const FNetworkPhysicsMoverInputs&>(FromData);
+
+	if (const FKinematicDefaultInputs* FromInput = FromDataInput.InputCmdContext.InputCollection.FindDataByType<FKinematicDefaultInputs>())
+	{
+		FKinematicDefaultInputs& LocalInputs = InputCmdContext.InputCollection.FindOrAddMutableDataByType<FKinematicDefaultInputs>();
+
+		LocalInputs.bIsJumpJustPressed |= FromInput->bIsJumpJustPressed;
+		LocalInputs.bIsJumpPressed |= FromInput->bIsJumpPressed;
+	}
+}
+
+//////////////////////////////////////////////////////////////////////////
+// FNetworkPhysicsMoverState
+
+void FNetworkPhysicsMoverState::ApplyData(UActorComponent* NetworkComponent) const
+{
+	if (NetworkComponent)
+	{
+		if (UMoverNetworkPhysicsLiaisonComponent* LiaisonComp = Cast<UMoverNetworkPhysicsLiaisonComponent>(NetworkComponent))
+		{
+			LiaisonComp->SetCurrentStateData(SyncStateContext, AuxStateContext);
+		}
+	}
+}
+
+void FNetworkPhysicsMoverState::BuildData(const UActorComponent* NetworkComponent)
+{
+	if (NetworkComponent)
+	{
+		if (const UMoverNetworkPhysicsLiaisonComponent* LiaisonComp = Cast<const UMoverNetworkPhysicsLiaisonComponent>(NetworkComponent))
+		{
+			LiaisonComp->GetCurrentStateData(SyncStateContext, AuxStateContext);
+		}
+	}
+}
+
+bool FNetworkPhysicsMoverState::NetSerialize(FArchive& Ar, class UPackageMap* PackageMap, bool& bOutSuccess)
+{
+	FNetworkPhysicsData::SerializeFrames(Ar);
+
+	if (PackageMap)
+	{
+		FNetSerializeParams Params(Ar);
+		SyncStateContext.NetSerialize(Params);
+		AuxStateContext.NetSerialize(Params);
+		bOutSuccess = true;
+	}
+	else
+	{
+		bOutSuccess = false;
+	}
+
+	return bOutSuccess;
+}
+
+void FNetworkPhysicsMoverState::InterpolateData(const FNetworkPhysicsData& MinData, const FNetworkPhysicsData& MaxData)
+{
+	const FNetworkPhysicsMoverState& MinState = static_cast<const FNetworkPhysicsMoverState&>(MinData);
+	const FNetworkPhysicsMoverState& MaxState = static_cast<const FNetworkPhysicsMoverState&>(MaxData);
+
+	const float LerpFactor = (LocalFrame - MinState.LocalFrame) / (MaxState.LocalFrame - MinState.LocalFrame);
+	SyncStateContext.Interpolate(&MinState.SyncStateContext, &MaxState.SyncStateContext, LerpFactor);
+	AuxStateContext.Interpolate(&MinState.AuxStateContext, &MaxState.AuxStateContext, LerpFactor);
+}
+
+//////////////////////////////////////////////////////////////////////////
+
+void UMoverNetworkPhysicsLiaisonComponent::GetCurrentInputData(OUT FMoverInputCmdContext& InputCmd) const
+{
+	InputCmd = NetInputCmd;
+}
+
+void UMoverNetworkPhysicsLiaisonComponent::GetCurrentStateData(OUT FMoverSyncState& SyncState, OUT FMoverAuxStateContext& AuxState) const
+{
+	SyncState = NetSyncState;
+	AuxState = NetAuxState;
+}
+
+void UMoverNetworkPhysicsLiaisonComponent::SetCurrentInputData(const FMoverInputCmdContext& InputCmd)
+{
+	NetInputCmd = InputCmd;
+}
+
+void UMoverNetworkPhysicsLiaisonComponent::SetCurrentStateData(const FMoverSyncState& SyncState, const FMoverAuxStateContext& AuxState)
+{
+	NetSyncState = SyncState;
+	NetAuxState = AuxState;
+}
+
+//////////////////////////////////////////////////////////////////////////
+// UMoverNetworkPhysicsLiaisonComponent
+
+UMoverNetworkPhysicsLiaisonComponent::UMoverNetworkPhysicsLiaisonComponent()
+{
+	PrimaryComponentTick.TickGroup = TG_PrePhysics;
+	PrimaryComponentTick.bCanEverTick = true;
+
+	bWantsInitializeComponent = true;
+	bAutoActivate = true;
+	SetIsReplicatedByDefault(true);
+
+	// Network physics relies on movement being replicated
+	if (AActor* MyActor = GetOwner())
+	{
+		MyActor->SetReplicatingMovement(true);
+		MyActor->SetReplicateMovement(true);
+	}
+
+	static const FName NetworkPhysicsComponentName(TEXT("PC_NetworkPhysicsComponent"));
+
+	NetworkPhysicsComponent = CreateDefaultSubobject<UNetworkPhysicsComponent, UNetworkPhysicsComponent>(NetworkPhysicsComponentName);
+	NetworkPhysicsComponent->SetNetAddressable(); // Make DSO components net addressable
+	NetworkPhysicsComponent->SetIsReplicated(true);
+}
+
+//////////////////////////////////////////////////////////////////////////
+//  UMoverNetworkPhysicsLiaisonComponent IMoverBackendLiaisonInterface
+
+float UMoverNetworkPhysicsLiaisonComponent::GetCurrentSimTimeMs()
+{
+	if (FPhysScene* Scene = GetWorld()->GetPhysicsScene())
+	{
+		if (Chaos::FPhysicsSolver* Solver = Scene->GetSolver())
+		{
+			return Solver->GetAsyncDeltaTime() * GetCurrentSimFrame() * 1000.0f;
+		}
+	}
+
+	return 0.0f;
+}
+
+int32 UMoverNetworkPhysicsLiaisonComponent::GetCurrentSimFrame()
+{
+	if (FPhysScene* Scene = GetWorld()->GetPhysicsScene())
+	{
+		if (Chaos::FPhysicsSolver* Solver = Scene->GetSolver())
+		{
+			int32 Offset = 0;
+			if (NetworkPhysicsComponent && !NetworkPhysicsComponent->HasServerWorld())
+			{
+				if (APlayerController* PC = NetworkPhysicsComponent->GetPlayerController())
+				{
+					Offset = PC->GetNetworkPhysicsTickOffset();
+				}
+			}
+
+			return Solver->GetCurrentFrame() + Offset;
+		}
+	}
+
+	return 0;
+}
+
+//////////////////////////////////////////////////////////////////////////
+// UMoverNetworkPhysicsLiaisonComponent UObject interface
+
+void UMoverNetworkPhysicsLiaisonComponent::OnRegister()
+{
+	Super::OnRegister();
+
+	if (!NetworkPhysicsComponent->IsRegistered())
+	{
+		NetworkPhysicsComponent->RegisterComponent();
+	}
+
+	// Need to set this here as physics creation requires access to MoverComp
+	
+	if ( (MoverComp = GetOwner()->FindComponentByClass<UMoverComponent>()) != nullptr )
+	{
+	
+		CommonMovementSettings = MoverComp->FindSharedSettings<UCommonLegacyMovementSettings>();
+		check(CommonMovementSettings);
+
+		if (MoverComp->UpdatedCompAsPrimitive)
+		{
+			MoverComp->UpdatedCompAsPrimitive->OnComponentPhysicsStateChanged.AddUniqueDynamic(this, &ThisClass::OnComponentPhysicsStateChanged);
+		}
+	}
+}
+
+void UMoverNetworkPhysicsLiaisonComponent::OnUnregister()
+{
+	if (MoverComp)
+	{
+		if (MoverComp->UpdatedCompAsPrimitive)
+		{
+			MoverComp->UpdatedCompAsPrimitive->OnComponentPhysicsStateChanged.RemoveDynamic(this, &ThisClass::OnComponentPhysicsStateChanged);
+		}
+	}
+
+	Super::OnUnregister();
+}
+
+void UMoverNetworkPhysicsLiaisonComponent::SetupConstraint()
+{
+	UWorld* World = GetWorld();
+	if (World && World->IsGameWorld())
+	{
+		if (FPhysScene_Chaos* PhysScene = World->GetPhysicsScene())
+		{
+			if (Chaos::FPhysicsSolver* Solver = PhysScene->GetSolver())
+			{
+				if (MoverComp && MoverComp->UpdatedCompAsPrimitive)
+				{
+					if (FBodyInstance* BI = MoverComp->UpdatedCompAsPrimitive->GetBodyInstance())
+					{
+						if (Chaos::FSingleParticlePhysicsProxy* CharacterProxy = BI->ActorHandle)
+						{
+							// Create and register the constraint
+							Constraint = MakeUnique<Chaos::FCharacterGroundConstraint>();
+							Constraint->Init(CharacterProxy);
+							Solver->RegisterObject(Constraint.Get());
+
+							// Set the common settings from the initial aux data
+							// The rest get set every frame depending on the current movement mode
+							Constraint->SetCosMaxWalkableSlopeAngle(CommonMovementSettings->MaxWalkSlopeCosine);
+							Constraint->SetVerticalAxis(MoverComp->GetUpDirection());
+
+							// Enable Physics Simulation
+							MoverComp->UpdatedCompAsPrimitive->SetSimulatePhysics(true);
+							
+							// Turn off sleeping
+							Chaos::FRigidBodyHandle_External& PhysicsBody = CharacterProxy->GetGameThreadAPI();
+							PhysicsBody.SetSleepType(Chaos::ESleepType::NeverSleep);
+						}
+					}
+				}
+			}
+		}
+	}
+}
+
+void UMoverNetworkPhysicsLiaisonComponent::DestroyConstraint()
+{
+	UWorld* World = GetWorld();
+	if (World && World->IsGameWorld() && HasValidPhysicsState())
+	{
+		if (FPhysScene_Chaos* PhysScene = World->GetPhysicsScene())
+		{
+			if (Chaos::FPhysicsSolver* Solver = PhysScene->GetSolver())
+			{
+				// Note: Proxy gets destroyed when the constraint is deregistered and that deletes the constraint
+				Solver->UnregisterObject(Constraint.Release());
+			}
+		}
+	}
+}
+
+void UMoverNetworkPhysicsLiaisonComponent::OnComponentPhysicsStateChanged(UPrimitiveComponent* ChangedComponent, EComponentPhysicsStateChange StateChange)
+{
+	if (StateChange == EComponentPhysicsStateChange::Destroyed)
+	{
+		DestroyConstraint();
+	}
+	else if (StateChange == EComponentPhysicsStateChange::Created)
+	{
+		SetupConstraint();
+
+		if (MoverComp && MoverComp->ModeFSM)
+		{
+			MoverComp->ModeFSM->SetModeImmediately(MoverComp->StartingMovementMode);
+		}
+	}
+}
+
+void UMoverNetworkPhysicsLiaisonComponent::InitializeComponent()
+{
+	Super::InitializeComponent();
+
+	NetworkPhysicsComponent->InitializeComponent();
+
+	if (ensureAlwaysMsgf(MoverComp, TEXT("UMoverNetworkPhysicsLiaisonComponent on actor %s failed to find associated Mover component. This actor's movement will not be simulated. Verify its setup."), *GetNameSafe(GetOwner())))
+	{
+		MoverComp->InitMoverSimulation();
+		MoverComp->ModeFSM->SetModeImmediately(MoverComp->StartingMovementMode);
+	}
+}
+
+void UMoverNetworkPhysicsLiaisonComponent::UninitializeComponent()
+{
+	NetworkPhysicsComponent->UninitializeComponent();
+	Super::UninitializeComponent();
+}
+
+bool UMoverNetworkPhysicsLiaisonComponent::ShouldCreatePhysicsState() const
+{
+	if (!IsRegistered() || IsBeingDestroyed())
+	{
+		return false;
+	}
+
+	UWorld* World = GetWorld();
+	if (World->IsGameWorld())
+	{
+		FPhysScene* PhysScene = World->GetPhysicsScene();
+
+		if (PhysScene && CanCreatePhysics())
+		{
+			return true;
+		}
+	}
+
+	return false;
+}
+
+bool UMoverNetworkPhysicsLiaisonComponent::HasValidPhysicsState() const
+{
+	return Constraint.IsValid() && Constraint->IsValid();
+}
+
+bool UMoverNetworkPhysicsLiaisonComponent::HasValidState() const
+{
+	return HasValidPhysicsState() && MoverComp && MoverComp->UpdatedCompAsPrimitive && MoverComp->UpdatedComponent
+		&& MoverComp->ModeFSM->IsValidLowLevel() && MoverComp->SimBlackboard->IsValidLowLevel();
+}
+
+void UMoverNetworkPhysicsLiaisonComponent::OnCreatePhysicsState()
+{
+	Super::OnCreatePhysicsState();
+
+	SetupConstraint();
+
+	UWorld* World = GetWorld();
+	if (World && World->IsGameWorld())
+	{
+		if (FPhysScene_Chaos* PhysScene = World->GetPhysicsScene())
+		{
+			if (Chaos::FPhysicsSolver* Solver = PhysScene->GetSolver())
+			{
+				if (MoverComp && MoverComp->UpdatedCompAsPrimitive)
+				{
+					if (FBodyInstance* BI = MoverComp->UpdatedCompAsPrimitive->GetBodyInstance())
+					{
+						if (Chaos::FSingleParticlePhysicsProxy* CharacterProxy = BI->ActorHandle)
+						{
+							// Register network data for recording and rewind/resim
+							if (NetworkPhysicsComponent)
+							{
+								NetworkPhysicsComponent->CreateDataHistory<FNetworkPhysicsMoverTraits>(this);
+							}
+
+							// Register with the physics mover manager
+							if (UPhysicsMoverManager* Manager = World->GetSubsystem<UPhysicsMoverManager>())
+							{
+								Manager->RegisterPhysicsMoverComponent(this);
+							}
+						}
+					}
+				}
+			}
+		}
+	}
+}
+
+void UMoverNetworkPhysicsLiaisonComponent::OnDestroyPhysicsState()
+{
+	UWorld* World = GetWorld();
+	if (World && World->IsGameWorld() && HasValidPhysicsState())
+	{
+		if (NetworkPhysicsComponent)
+		{
+			NetworkPhysicsComponent->RemoveDataHistory();
+		}
+
+		DestroyConstraint();
+
+		if (UPhysicsMoverManager* Manager = World->GetSubsystem<UPhysicsMoverManager>())
+		{
+			Manager->UnregisterPhysicsMoverComponent(this);
+		}
+	}
+	
+	Super::OnDestroyPhysicsState();
+}
+
+bool UMoverNetworkPhysicsLiaisonComponent::CanCreatePhysics() const
+{
+	check(GetOwner());
+	FString ActorName = GetOwner()->GetName();
+
+	if (!IsValid(MoverComp->UpdatedComponent))
+	{
+		UE_LOG(LogMover, Warning, TEXT("Can't create physics %s (%s). UpdatedComponent is not set."), *ActorName, *GetPathName());
+		return false;
+	}
+
+	if (!IsValid(MoverComp->UpdatedCompAsPrimitive))
+	{
+		UE_LOG(LogMover, Warning, TEXT("Can't create physics %s (%s). UpdatedComponent is not a PrimitiveComponent."), *ActorName, *GetPathName());
+		return false;
+	}
+
+	return true;
+}
+
+//////////////////////////////////////////////////////////////////////////
+
+Chaos::FUniqueIdx UMoverNetworkPhysicsLiaisonComponent::GetUniqueIdx() const
+{
+	if (MoverComp && MoverComp->UpdatedCompAsPrimitive)
+	{
+		if (FBodyInstance* BI = MoverComp->UpdatedCompAsPrimitive->GetBodyInstance())
+		{
+			if (FPhysicsActorHandle ActorHandle = BI->ActorHandle)
+			{
+				return ActorHandle->GetGameThreadAPI().UniqueIdx();
+			}
+		}
+	}
+
+	return Chaos::FUniqueIdx();
+}
+
+void UMoverNetworkPhysicsLiaisonComponent::UpdateConstraintSettings(const FMoverAuxStateContext& AuxState)
+{
+	if (HasValidState())
+	{
+		Constraint->SetVerticalAxis(MoverComp->GetUpDirection());
+		Constraint->SetCosMaxWalkableSlopeAngle(CommonMovementSettings->MaxWalkSlopeCosine);
+
+		const UBaseMovementMode* CurrentMode = MoverComp->ModeFSM->GetCurrentMode();
+		if (CurrentMode && CurrentMode->Implements<UPhysicsCharacterMovementModeInterface>())
+		{
+			const IPhysicsCharacterMovementModeInterface* PhysicsMode = CastChecked<IPhysicsCharacterMovementModeInterface>(CurrentMode);
+			PhysicsMode->UpdateConstraintSettings(*Constraint);
+		}
+	}
+}
+
+FMoverTimeStep UMoverNetworkPhysicsLiaisonComponent::GetCurrentMoverTimeStep() const
+{
+	FMoverTimeStep TimeStep;
+
+	if (FPhysScene* Scene = GetWorld()->GetPhysicsScene())
+	{
+		if (Chaos::FPhysicsSolver* Solver = Scene->GetSolver())
+		{
+			int32 Offset = 0;
+			if (NetworkPhysicsComponent && !NetworkPhysicsComponent->HasServerWorld())
+			{
+				if (APlayerController* PC = NetworkPhysicsComponent->GetPlayerController())
+				{
+					Offset = PC->GetNetworkPhysicsTickOffset();
+				}
+			}
+
+			TimeStep.ServerFrame = Solver->GetCurrentFrame() + Offset;
+			TimeStep.StepMs = Solver->GetAsyncDeltaTime() * 1000.0f;
+			TimeStep.BaseSimTimeMs = TimeStep.ServerFrame * TimeStep.StepMs;
+			TimeStep.bIsResimulating = Solver->GetEvolution()->IsResimming();
+		}
+	}
+
+	return TimeStep;
+}
+
+void UMoverNetworkPhysicsLiaisonComponent::GameThread_ProduceInput(float DeltaSeconds, OUT FPhysicsMoverAsyncInput& Input)
+{
+	if (HasValidState())
+	{
+		Input.MoverIdx = GetUniqueIdx();
+		Input.MoverSimulation = this;
+		
+		// Produce input
+		if (!MoverComp || !NetworkPhysicsComponent)
+		{
+			return;
+		}
+
+		if (!NetworkPhysicsComponent->HasServerWorld())
+		{
+			APlayerController* PlayerController = NetworkPhysicsComponent->GetPlayerController();
+			if (PlayerController && PlayerController->IsLocalController())
+			{
+				if (!bCachedInputIsValid)
+				{
+					const int DeltaTimeMS = FMath::TruncToInt(DeltaSeconds * 1000.0f);
+
+					MoverComp->ProduceInput(DeltaTimeMS, &Input.InputCmd);
+
+					// We only want to consume one input per physics frame
+					// so if there is already a valid cached input we use that.
+					// Input is set invalid when the async output is consumed
+					bCachedInputIsValid = true;
+				}
+				else
+				{
+					Input.InputCmd = MoverComp->CachedLastProducedInputCmd;
+				}
+			}
+		}
+
+		Input.SyncState.SyncStateCollection.Empty();
+		if (MoverComp->bHasValidCachedState)
+		{
+			Input.SyncState = MoverComp->CachedLastSyncState;
+			Input.AuxState = MoverComp->CachedLastAuxState;
+		}
+		else
+		{
+			FMoverDefaultSyncState& Default = Input.SyncState.SyncStateCollection.FindOrAddMutableDataByType<FMoverDefaultSyncState>();
+			Default.MovementMode = KinematicModeNames::Falling;
+			Input.AuxState = FMoverAuxStateContext();
+		}
+
+		// This is required so that the physics thread can have a copy of the data to access
+		NetInputCmd = Input.InputCmd;
+		NetSyncState = Input.SyncState;
+		NetAuxState = Input.AuxState;
+
+		if (MoverComp->bHasValidLastProducedInput)
+		{
+			MoverComp->OnPreSimulationTick.Broadcast(GetCurrentMoverTimeStep(), MoverComp->CachedLastProducedInputCmd);
+		}
+		else
+		{
+			MoverComp->OnPreSimulationTick.Broadcast(GetCurrentMoverTimeStep(), FMoverInputCmdContext());
+		}
+	}
+}
+
+void UMoverNetworkPhysicsLiaisonComponent::GameThread_ConsumeOutput(const FPhysicsMoverAsyncOutput& Output)
+{
+	if (Output.bIsValid && MoverComp)
+	{
+		if (const FMoverDefaultSyncState* OutputSyncState = Output.SyncState.SyncStateCollection.FindDataByType<FMoverDefaultSyncState>())
+		{
+			// If landed, broadcast OnLanded
+			if (MoverComp->bHasValidCachedState)
+			{
+				if (const FMoverDefaultSyncState* CachedSyncState = MoverComp->CachedLastSyncState.SyncStateCollection.FindDataByType<FMoverDefaultSyncState>())
+				{
+					if ((CachedSyncState->MovementMode == KinematicModeNames::Falling) && (OutputSyncState->MovementMode == KinematicModeNames::Walking))
+					{
+						if (UFallingMode* FallingMode = MoverComp->FindMode_Mutable<UFallingMode>())
+						{
+							FHitResult HitResult;
+							MoverComp->TryGetFloorCheckHitResult(HitResult);
+							FallingMode->OnLanded.Broadcast(OutputSyncState->MovementMode, HitResult);
+						}
+					}
+				}
+			}
+
+			// TODO: Consider moving to a util function
+			MoverComp->CachedLastSyncState = Output.SyncState;
+			MoverComp->CachedLastAuxState = Output.AuxState;
+			MoverComp->CachedLastSimTickTimeStep.BaseSimTimeMs = GetCurrentSimTimeMs();
+			MoverComp->CachedLastSimTickTimeStep.ServerFrame = GetCurrentSimFrame();
+			MoverComp->bHasValidCachedState = true;
+		}
+
+		FMoverTimeStep TimeStep;
+		if (FPhysScene* PhysScene = GetWorld()->GetPhysicsScene())
+		{
+			if (Chaos::FPhysicsSolver* Solver = PhysScene->GetSolver())
+			{
+				TimeStep.StepMs = Solver->GetAsyncDeltaTime() * 1000.0f;
+				TimeStep.BaseSimTimeMs = Solver->GetSolverTime();
+			}
+		}
+		MoverComp->OnPostSimulationTick.Broadcast(TimeStep);
+
+		bCachedInputIsValid = false;
+
+		UpdateConstraintSettings(Output.AuxState);
+	}
+}
+
+void UMoverNetworkPhysicsLiaisonComponent::AsyncPhysics_ProcessInputs(int32 PhysicsStep, float DeltaTime, const FPhysicsMoverAsyncInput& Input) const
+{
+	// Override input data unless player is local client
+	if (FPhysScene* PhysScene = GetWorld()->GetPhysicsScene())
+	{
+		if (NetworkPhysicsComponent)
+		{
+			if (NetworkPhysicsComponent->HasServerWorld())
+			{
+				// Server remote player
+				GetCurrentInputData(Input.InputCmd);
+			}
+			else
+			{
+				GetCurrentStateData(Input.SyncState, Input.AuxState);
+
+				bool bIsSolverResim = false;
+				if (Chaos::FPhysicsSolver* Solver = PhysScene->GetSolver())
+				{
+					bIsSolverResim = Solver->GetEvolution()->IsResimming();
+				}
+
+				bool bLocalPlayer = false;
+				APlayerController* PlayerController = NetworkPhysicsComponent->GetPlayerController();
+				if (PlayerController && PlayerController->IsLocalController())
+				{
+					bLocalPlayer = true;
+				}
+
+				if (!bLocalPlayer || bIsSolverResim)
+				{
+					GetCurrentInputData(Input.InputCmd);
+				}
+			}
+		}
+		
+	}
+}
+
+void UMoverNetworkPhysicsLiaisonComponent::AsyncPhysics_OnPreSimulate(const FPhysicsMoverSimulationTickParams& TickParams, const FPhysicsMoverAsyncInput& Input, OUT FPhysicsMoverAsyncOutput& Output) const
+{
+	// Sync/aux state should carry over to the next sim frame by default unless something modifies it
+	Output.SyncState = Input.SyncState;
+	Output.AuxState = Input.AuxState;
+
+	if (!HasValidState() || !Input.InputCmd.InputCollection.FindDataByType<FKinematicDefaultInputs>())
+	{
+		return;
+	}
+
+	// Exit if Physics is not enabled on the CollisionShape
+	const UShapeComponent* CollisionShape = Cast<const UShapeComponent>(MoverComp->UpdatedComponent);
+	if (!CollisionShape || !CollisionShape->IsSimulatingPhysics())
+	{
+		return;
+	}
+	
+	Chaos::FCharacterGroundConstraintHandle* ConstraintHandle = Constraint->GetProxy<Chaos::FCharacterGroundConstraintProxy>()->GetPhysicsThreadAPI();
+	if (!ConstraintHandle || !ConstraintHandle->IsEnabled() || !ConstraintHandle->GetCharacterParticle())
+	{
+		return;
+	}
+
+	// Update sync state from physics
+	Chaos::FPBDRigidParticleHandle* CharacterParticle = ConstraintHandle->GetCharacterParticle()->CastToRigidParticle();
+	if (!CharacterParticle || CharacterParticle->Disabled())
+	{
+		return;
+	}
+
+	FMoverDefaultSyncState& SyncState = Input.SyncState.SyncStateCollection.FindOrAddMutableDataByType<FMoverDefaultSyncState>();
+
+	const IPhysicsCharacterMovementModeInterface* PhysicsMode = Cast<const IPhysicsCharacterMovementModeInterface>(MoverComp->ModeFSM->FindMovementMode(SyncState.MovementMode));
+	if (!PhysicsMode)
+	{
+		UE_LOG(LogMover, Verbose, TEXT("Attempting to run non-physics movement mode %s in physics mover update."), *SyncState.MovementMode.ToString());
+		return;
+	}
+
+	// Make the sync state velocity relative to the ground if walking
+	FVector LocalGroundVelocity = FVector::ZeroVector;
+	if (SyncState.MovementMode == KinematicModeNames::Walking)
+	{
+		if (const UMoverBlackboard* Blackboard = MoverComp->GetSimBlackboard())
+		{
+			FFloorCheckResult LastFloorResult;
+			if (Blackboard->TryGet(KinematicBlackboard::LastFloorResult, LastFloorResult))
+			{
+				LocalGroundVelocity = UPhysicsMovementUtils::ComputeGroundVelocityFromHitResult(CharacterParticle->X(), LastFloorResult.HitResult, TickParams.DeltaTimeSeconds);
+				LocalGroundVelocity -= LocalGroundVelocity.ProjectOnToNormal(LastFloorResult.HitResult.ImpactNormal);
+			}
+		}
+	}
+	SyncState.SetTransforms_WorldSpace(CharacterParticle->X(), FRotator(CharacterParticle->R()), CharacterParticle->V() - LocalGroundVelocity);
+
+	///////////////////////////////////////////////////////////////////////////////////////////////////////////////////////////
+	// Update the simulation
+	//Input.MoverSimulation->SimulationTick(TickParams, Input.SimInput, SimOutput);
+
+	FMoverTickStartData TickStartData(Input.InputCmd, Input.SyncState, Input.AuxState);
+	FMoverTickEndData TickEndData;
+	FFloorCheckResult FloorResult;
+
+	FMoverTimeStep TimeStep = GetCurrentMoverTimeStep();
+
+	bool bHasRolledBack = false; // TODO
+
+	//-------------------------------------------------------------------------------------------
+	// Copied from KinematicMoverComponent::SimulationTick
+
+	if (bHasRolledBack)
+	{
+		MoverComp->ProcessFirstSimTickAfterRollback(TimeStep);
+	}
+
+	// Sync/aux state should carry over to the next sim frame by default unless something modifies it
+	TickEndData.SyncState = TickStartData.SyncState;
+	TickEndData.AuxState = TickStartData.AuxState;
+
+	if (MoverComp->ModeFSM->IsValidLowLevel() && MoverComp->SimBlackboard->IsValidLowLevel())
+	{
+		const FMoverDefaultSyncState* StartingSyncState = TickStartData.SyncState.SyncStateCollection.FindDataByType<FMoverDefaultSyncState>();
+		check(StartingSyncState);
+
+		FKinematicDefaultInputs* InputCmd = TickStartData.InputCmd.InputCollection.FindMutableDataByType<FKinematicDefaultInputs>();
+		FMoverDefaultSyncState& OutputSyncState = TickEndData.SyncState.SyncStateCollection.FindOrAddMutableDataByType<FMoverDefaultSyncState>();
+
+		if (InputCmd && !InputCmd->SuggestedMovementMode.IsNone())
+		{
+			MoverComp->QueueNextMode(InputCmd->SuggestedMovementMode);
+		}
+
+		// Tick the actual simulation. This is where the proposed moves are queried and executed, affecting change to the moving actor's gameplay state and captured in the output sim state
+		MoverComp->ModeFSM->OnSimulationTick(MoverComp->UpdatedComponent, MoverComp->UpdatedCompAsPrimitive, MoverComp->SimBlackboard.Get(), TickStartData, TimeStep, TickEndData);
+
+		const FName MovementModeAfterTick = MoverComp->ModeFSM->GetCurrentModeName();
+		OutputSyncState.MovementMode = MovementModeAfterTick;
+	}
+
+	//-------------------------------------------------------------------------------------------
+
+	MoverComp->SimBlackboard->TryGet(KinematicBlackboard::LastFloorResult, Output.FloorResult);
+
+	Output.SyncState = TickEndData.SyncState;
+	Output.AuxState = TickEndData.AuxState;
+
+	/////////////////////////////////////////////////////////////////////////////////////////////////////////////////////////////
+
+	// Update physics constraint from output sync state and aux state
+	const FMoverDefaultSyncState* OutputSyncState = Output.SyncState.SyncStateCollection.FindDataByType<FMoverDefaultSyncState>();
+	if (!ensure(OutputSyncState))
+	{
+		return;
+	}
+
+	FVector TargetDeltaPos = OutputSyncState->GetLocation_WorldSpace() - CharacterParticle->X();
+
+	if (TargetDeltaPos.SizeSquared2D() > GPhysicsDrivenMotionDebugParams.TeleportThreshold * GPhysicsDrivenMotionDebugParams.TeleportThreshold)
+	{
+		TeleportParticle(CharacterParticle, OutputSyncState->GetLocation_WorldSpace(), OutputSyncState->GetOrientation_WorldSpace().Quaternion());
+	}
+
+	// Add back the ground velocity that was subtracted to but the movement velocity in local space
+	FVector TargetVelocity = OutputSyncState->GetVelocity_WorldSpace() + LocalGroundVelocity;
+
+	// Landed so add the new ground velocity
+	if ((OutputSyncState->MovementMode == KinematicModeNames::Walking) && (SyncState.MovementMode != KinematicModeNames::Walking))
+	{
+		if (const UPhysicsDrivenWalkingMode* WalkingMode = Cast<UPhysicsDrivenWalkingMode>(MoverComp->FindMovementMode(UPhysicsDrivenWalkingMode::StaticClass())))
+		{
+			LocalGroundVelocity = UPhysicsMovementUtils::ComputeGroundVelocityFromHitResult(CharacterParticle->X(), Output.FloorResult.HitResult, TickParams.DeltaTimeSeconds);
+			LocalGroundVelocity -= LocalGroundVelocity.ProjectOnToNormal(Output.FloorResult.HitResult.ImpactNormal);
+			TargetVelocity += WalkingMode->FractionalVelocityToTarget * LocalGroundVelocity;
+		}
+	}
+
+	CharacterParticle->SetV(TargetVelocity);
+
+	// Note: Output sync state does not have a target angular velocity so
+	// use the target orientation
+	FRotator DeltaRotation = OutputSyncState->GetOrientation_WorldSpace() - FRotator(CharacterParticle->R());
+	FRotator Winding, Remainder;
+	DeltaRotation.GetWindingAndRemainder(Winding, Remainder);
+	float TargetDeltaFacing = FMath::DegreesToRadians(Remainder.Yaw);
+	if (TickParams.DeltaTimeSeconds > UE_SMALL_NUMBER)
+	{
+		CharacterParticle->SetW((TargetDeltaFacing / TickParams.DeltaTimeSeconds) * Chaos::FVec3::ZAxisVector);
+	}
+
+	// Update the constraint data based on the floor result
+	if (Output.FloorResult.bBlockingHit)
+	{
+		// Set the ground particle on the constraint
+		Chaos::FGeometryParticleHandle* GroundParticle = nullptr;
+
+		if (IPhysicsComponent* PhysicsComp = Cast<IPhysicsComponent>(Output.FloorResult.HitResult.Component))
+		{
+			if (Chaos::FPhysicsObjectHandle PhysicsObject = PhysicsComp->GetPhysicsObjectById(Output.FloorResult.HitResult.Item))
+			{
+				Chaos::FReadPhysicsObjectInterface_Internal Interface = Chaos::FPhysicsObjectInternalInterface::GetRead();
+				GroundParticle = Interface.GetParticle(PhysicsObject);
+
+				// Wake the ground particle if it is sleeping
+				WakeParticleIfSleeping(GroundParticle);
+			}
+		}
+		ConstraintHandle->SetGroundParticle(GroundParticle);
+
+		// Set the max walkable slope angle using any override from the hit component
+		float WalkableSlopeCosine = ConstraintHandle->GetSettings().CosMaxWalkableSlopeAngle;
+		if (Output.FloorResult.HitResult.Component != nullptr)
+		{
+			const FWalkableSlopeOverride& SlopeOverride = Output.FloorResult.HitResult.Component->GetWalkableSlopeOverride();
+			WalkableSlopeCosine = SlopeOverride.ModifyWalkableFloorZ(WalkableSlopeCosine);
+		}
+
+		if (!Output.FloorResult.bWalkableFloor)
+		{
+			WalkableSlopeCosine = 2.0f;
+		}
+
+		ConstraintHandle->SetData({
+			Output.FloorResult.HitResult.ImpactNormal,
+			TargetDeltaPos,
+			TargetDeltaFacing,
+			Output.FloorResult.FloorDist,
+			WalkableSlopeCosine
+			});
+	}
+	else
+	{
+		ConstraintHandle->SetData({
+			Chaos::FVec3::ZAxisVector,
+			Chaos::FVec3::ZeroVector,
+			0.0,
+			1.0e10,
+			0.5f
+			});
+	}
+
+	Output.bIsValid = true;
+}
+
+void UMoverNetworkPhysicsLiaisonComponent::TeleportParticle(Chaos::FGeometryParticleHandle* Particle, const FVector& Position, const FQuat& Rotation) const
+{
+	if (!Particle)
+	{
+		return;
+	}
+
+	if (UWorld* World = GetWorld())
+	{
+		if (FPhysScene* Scene = World->GetPhysicsScene())
+		{
+			if (Chaos::FPhysicsSolver* Solver = Scene->GetSolver())
+			{
+				if (Solver->GetEvolution())
+				{
+					Solver->GetEvolution()->SetParticleTransform(Particle, Position, Rotation, true);
+				}
+			}
+		}
+	}
+}
+
+void UMoverNetworkPhysicsLiaisonComponent::WakeParticleIfSleeping(Chaos::FGeometryParticleHandle* Particle) const
+{
+	if (Particle)
+	{
+		Chaos::FPBDRigidParticleHandle* Rigid = Particle->CastToRigidParticle();
+		if (Rigid && (Rigid->ObjectState() == Chaos::EObjectStateType::Sleeping))
+		{
+			if (UWorld* World = GetWorld())
+			{
+				if (FPhysScene* Scene = World->GetPhysicsScene())
+				{
+					if (Chaos::FPhysicsSolver* Solver = Scene->GetSolver())
+					{
+						if (Solver->GetEvolution())
+						{
+							Solver->GetEvolution()->SetParticleObjectState(Rigid, Chaos::EObjectStateType::Dynamic);
+						}
+					}
+				}
+			}
+		}
+	}
+}
+
+void UMoverNetworkPhysicsLiaisonComponent::ASyncPhysics_OnContactModification(const FPhysicsMoverAsyncInput& Input, Chaos::FCollisionContactModifier& Modifier) const
+{
+	if (!HasValidState())
+	{
+		return;
+	}
+
+	Chaos::FCharacterGroundConstraintHandle* ConstraintHandle = Constraint->GetProxy<Chaos::FCharacterGroundConstraintProxy>()->GetPhysicsThreadAPI();
+	if (!ConstraintHandle || !ConstraintHandle->IsEnabled() || !ConstraintHandle->GetCharacterParticle())
+	{
+		return;
+	}
+
+	Chaos::FPBDRigidParticleHandle* CharacterParticle = ConstraintHandle->GetCharacterParticle()->CastToRigidParticle();
+	if (!CharacterParticle || CharacterParticle->Disabled())
+	{
+		return;
+	}
+
+	const Chaos::FGeometryParticleHandle* GroundParticle = ConstraintHandle->GetGroundParticle();
+	if (!GroundParticle)
+	{
+		return;
+	}
+
+	if (const FMoverDefaultSyncState* SyncState = Input.SyncState.SyncStateCollection.FindDataByType<FMoverDefaultSyncState>())
+	{
+		float PawnHalfHeight;
+		float PawnRadius;
+		MoverComp->UpdatedComponent->CalcBoundingCylinder(PawnRadius, PawnHalfHeight);
+
+		const float CharacterHeight = CharacterParticle->X().Z;
+		const float EndCapHeight = CharacterHeight - PawnHalfHeight + PawnRadius;
+
+		const float CosThetaMax = 0.97f;
+
+		float MinContactHeightStepUps = CharacterHeight + 1.0e10f;
+		if (SyncState->MovementMode == KinematicModeNames::Walking)
+		{
+			if (const UPhysicsDrivenWalkingMode* WalkingMode = Cast<UPhysicsDrivenWalkingMode>(MoverComp->FindMode_Mutable<UPhysicsDrivenWalkingMode>()))
+			{
+				// Contacts on the character capsule below the MaxStepHeight tend to snag the character when stepping
+				// up or down, so disable them
+
+				if (const UCommonLegacyMovementSettings* Settings = MoverComp->FindSharedSettings<UCommonLegacyMovementSettings>())
+				{
+					const float StepDistance = FMath::Abs(WalkingMode->TargetHeight - ConstraintHandle->GetData().GroundDistance);
+					if (StepDistance >= GPhysicsDrivenMotionDebugParams.MinStepUpDistance)
+					{
+						MinContactHeightStepUps = CharacterHeight - WalkingMode->TargetHeight + Settings->MaxStepHeight;
+					}
+				}
+			}
+		}
+
+		for (Chaos::FContactPairModifier& PairModifier : Modifier.GetContacts(CharacterParticle))
+		{
+			const int32 CharacterIdx = CharacterParticle == PairModifier.GetParticlePair()[0] ? 0 : 1;
+			const int32 OtherIdx = CharacterIdx == 0 ? 1 : 0;
+
+			if (GroundParticle == PairModifier.GetParticlePair()[OtherIdx])
+			{
+				for (int32 Idx = 0; Idx < PairModifier.GetNumContacts(); ++Idx)
+				{
+					Chaos::FVec3 Point0, Point1;
+					PairModifier.GetWorldContactLocations(Idx, Point0, Point1);
+					Chaos::FVec3 CharacterPoint = CharacterIdx == 0 ? Point0 : Point1;
+
+					Chaos::FVec3 ContactNormal = PairModifier.GetWorldNormal(Idx);
+					if ((ContactNormal.Z > CosThetaMax) && CharacterPoint.Z < EndCapHeight)
+					{
+						// Disable any nearly vertical contact with the end cap of the capsule
+						// This will be handled by the character ground constraint
+						PairModifier.SetContactPointDisabled(Idx);
+					}
+					else if (CharacterPoint.Z < MinContactHeightStepUps)
+					{
+						// In the case of steps ups disable all contacts below the max step height
+						PairModifier.SetContactPointDisabled(Idx);
+					}
+				}
+			}
+		}
+	}
+}
diff --git a/Engine/Plugins/Experimental/Mover/Source/Mover/Private/LayeredMove.cpp b/Engine/Plugins/Experimental/Mover/Source/Mover/Private/LayeredMove.cpp
new file mode 100644
--- /dev/null
+++ b/Engine/Plugins/Experimental/Mover/Source/Mover/Private/LayeredMove.cpp
@@ -0,0 +1,558 @@
+// Copyright Epic Games, Inc. All Rights Reserved.
+
+#include "LayeredMove.h"
+#include "MoverLog.h"
+#include "MoverComponent.h"
+#include "MoverSimulationTypes.h"
+
+const float LayeredMove_InvalidTime = -UE_BIG_NUMBER;
+
+void FLayeredMoveFinishVelocitySettings::NetSerialize(FArchive& Ar)
+{
+	uint8 bHasFinishVelocitySettings = Ar.IsSaving() ? 0 : (FinishVelocityMode == ELayeredMoveFinishVelocityMode::MaintainLastRootMotionVelocity);
+	Ar.SerializeBits(&bHasFinishVelocitySettings, 1);
+
+	if (bHasFinishVelocitySettings)
+	{
+		uint8 FinishVelocityModeAsU8 = (uint8)(FinishVelocityMode);
+		Ar << FinishVelocityModeAsU8;
+		FinishVelocityMode = (ELayeredMoveFinishVelocityMode)FinishVelocityModeAsU8;
+
+		if (FinishVelocityMode == ELayeredMoveFinishVelocityMode::SetVelocity)
+		{
+			Ar << SetVelocity;
+		}
+		else if (FinishVelocityMode == ELayeredMoveFinishVelocityMode::ClampVelocity)
+		{
+			Ar << ClampVelocity;
+		}
+	}
+}
+
+FLayeredMoveBase::FLayeredMoveBase() :
+	MixMode(EMoveMixMode::AdditiveVelocity),
+	DurationMs(-1.f),
+	StartSimTimeMs(LayeredMove_InvalidTime)
+{
+}
+
+
+void FLayeredMoveBase::StartMove(const UMoverComponent* MoverComp, UMoverBlackboard* SimBlackboard, float CurrentSimTimeMs)
+{
+	StartSimTimeMs = CurrentSimTimeMs;
+	OnStart(MoverComp, SimBlackboard);
+}
+
+bool FLayeredMoveBase::IsFinished(float CurrentSimTimeMs) const
+{
+	const bool bHasStarted = (StartSimTimeMs >= 0.f);
+	const bool bTimeExpired = bHasStarted && (DurationMs > 0.f) && (StartSimTimeMs + DurationMs <= CurrentSimTimeMs);
+	const bool bDidTickOnceAndExpire = bHasStarted && (DurationMs == 0.f);
+
+	return bTimeExpired || bDidTickOnceAndExpire;
+}
+
+void FLayeredMoveBase::EndMove(const UMoverComponent* MoverComp, UMoverBlackboard* SimBlackboard, float CurrentSimTimeMs)
+{
+	OnEnd(MoverComp, SimBlackboard, CurrentSimTimeMs);
+}
+
+FLayeredMoveBase* FLayeredMoveBase::Clone() const
+{
+	// If child classes don't override this, saved moves will not work
+	checkf(false, TEXT("FLayeredMoveBase::Clone() being called erroneously. This should always be overridden in child classes!"));
+	return nullptr;
+}
+
+
+void FLayeredMoveBase::NetSerialize(FArchive& Ar)
+{
+	uint8 MixModeAsU8 = (uint8)MixMode;
+	Ar << MixModeAsU8;
+	MixMode = (EMoveMixMode)MixModeAsU8;
+
+	Ar << DurationMs;
+	Ar << StartSimTimeMs;
+
+	FinishVelocitySettings.NetSerialize(Ar);
+}
+
+
+UScriptStruct* FLayeredMoveBase::GetScriptStruct() const
+{
+	return FLayeredMoveBase::StaticStruct();
+}
+
+
+FString FLayeredMoveBase::ToSimpleString() const
+{
+	return GetScriptStruct()->GetName();
+}
+
+
+FLayeredMoveGroup::FLayeredMoveGroup()
+	: bApplyResidualVelocity(false)
+	, ResidualVelocity(FVector::Zero())
+	, ResidualClamping(-1.f)
+{
+}
+
+
+void FLayeredMoveGroup::QueueLayeredMove(TSharedPtr<FLayeredMoveBase> Move)
+{
+	if (ensure(Move.IsValid()))
+	{
+		QueuedLayeredMoves.Add(Move);
+		UE_LOG(LogMover, VeryVerbose, TEXT("LayeredMove queued move (%s)"), *Move->ToSimpleString());
+	}
+}
+
+bool FLayeredMoveGroup::DoGenerateMove(const FMoverTickStartData& StartState, const FMoverTimeStep& TimeStep, const UMoverComponent* MoverComp, UMoverBlackboard* SimBlackboard, FProposedMove& OutMove)
+{
+	const float SimStartTimeMs		= TimeStep.BaseSimTimeMs;
+	const float SimTimeAfterTickMs	= SimStartTimeMs + TimeStep.StepMs;
+
+	FlushMoveArrays(MoverComp, SimBlackboard, SimStartTimeMs);
+
+	FProposedMove CumulativeMoveStep;
+	CumulativeMoveStep.MixMode = EMoveMixMode::AdditiveVelocity;
+
+	bool bDidAccumulateAnyMoves = false;
+
+	// Tick and accumulate all active moves
+	// Gather all proposed moves and distill this into a final movement report. May include separate additive vs override moves.
+	// TODO: may want to sort by priority or other factors
+	for (TSharedPtr<FLayeredMoveBase>& ActiveMove : ActiveLayeredMoves)
+	{
+		FProposedMove MoveStep;
+		
+		if (ActiveMove->GenerateMove(StartState, TimeStep, MoverComp, SimBlackboard, MoveStep) &&
+			CumulativeMoveStep.MixMode < EMoveMixMode::OverrideAll)
+		{
+			bDidAccumulateAnyMoves = true;
+
+			if (CumulativeMoveStep.PreferredMode != MoveStep.PreferredMode &&
+				!CumulativeMoveStep.PreferredMode.IsNone() && !MoveStep.PreferredMode.IsNone())
+			{
+				UE_LOG(LogMover, Warning, TEXT("Multiple LayeredMoves are conflicting with preferred moves. %s will override %s"),
+					*MoveStep.PreferredMode.ToString(), *CumulativeMoveStep.PreferredMode.ToString());
+			}
+
+			if (MoveStep.bHasDirIntent && CumulativeMoveStep.MixMode != EMoveMixMode::OverrideAll)
+			{
+				if (CumulativeMoveStep.bHasDirIntent)
+				{
+					UE_LOG(LogMover, Warning, TEXT("Multiple LayeredMoves are setting direction intent and only one direction intent will be used."));
+				}
+				
+				CumulativeMoveStep.bHasDirIntent = MoveStep.bHasDirIntent;
+				CumulativeMoveStep.DirectionIntent = MoveStep.DirectionIntent;
+			}
+
+			if (ActiveMove->MixMode == EMoveMixMode::OverrideVelocity)
+			{
+				if (CumulativeMoveStep.MixMode == EMoveMixMode::OverrideVelocity || CumulativeMoveStep.MixMode == EMoveMixMode::OverrideAll)
+				{
+					UE_LOG(LogMover, Warning, TEXT("Multiple LayeredMoves with Override mix mode are active simultaneously. Only one will take effect."));
+				}
+
+				if (!MoveStep.PreferredMode.IsNone())
+				{
+					CumulativeMoveStep.PreferredMode = MoveStep.PreferredMode;
+				}
+				
+				CumulativeMoveStep.MixMode = EMoveMixMode::OverrideVelocity;
+				CumulativeMoveStep.LinearVelocity  = MoveStep.LinearVelocity;
+				CumulativeMoveStep.AngularVelocity = MoveStep.AngularVelocity;
+			}
+			else if (ActiveMove->MixMode == EMoveMixMode::AdditiveVelocity)
+			{
+				if (CumulativeMoveStep.MixMode != EMoveMixMode::OverrideVelocity && CumulativeMoveStep.MixMode != EMoveMixMode::OverrideAll)
+				{
+					if (!MoveStep.PreferredMode.IsNone())
+					{
+						CumulativeMoveStep.PreferredMode = MoveStep.PreferredMode;
+					}
+
+					CumulativeMoveStep.PreferredMode = MoveStep.PreferredMode;
+					CumulativeMoveStep.LinearVelocity += MoveStep.LinearVelocity;
+					CumulativeMoveStep.AngularVelocity += MoveStep.AngularVelocity;
+				}
+			}
+			else if (ActiveMove->MixMode == EMoveMixMode::OverrideAll)
+			{
+				if (CumulativeMoveStep.MixMode == EMoveMixMode::OverrideVelocity || CumulativeMoveStep.MixMode == EMoveMixMode::OverrideAll)
+				{
+					UE_LOG(LogMover, Warning, TEXT("Multiple LayeredMoves with Override mix mode are active simultaneously. Only one will take effect."));
+				}
+				
+				CumulativeMoveStep = MoveStep;
+			}
+			else
+			{
+				check(0);	// unhandled case
+			}
+		}
+	}
+
+	if (bDidAccumulateAnyMoves)
+	{
+		OutMove = CumulativeMoveStep;
+	}
+
+	return bDidAccumulateAnyMoves;
+}
+
+
+void FLayeredMoveGroup::NetSerialize(FArchive& Ar, uint8 MaxNumMovesToSerialize/* = MAX_uint8*/)
+{
+	// TODO: Warn if some sources will be dropped
+	const uint8 NumActiveMovesToSerialize = FMath::Min<int32>(ActiveLayeredMoves.Num(), MaxNumMovesToSerialize);
+	const uint8 NumQueuedMovesToSerialize = NumActiveMovesToSerialize < MaxNumMovesToSerialize ? MaxNumMovesToSerialize - NumActiveMovesToSerialize : 0;
+	NetSerializeLayeredMovesArray(Ar, ActiveLayeredMoves, NumActiveMovesToSerialize);
+	NetSerializeLayeredMovesArray(Ar, QueuedLayeredMoves, NumQueuedMovesToSerialize);
+}
+
+
+FLayeredMoveGroup& FLayeredMoveGroup::operator=(const FLayeredMoveGroup& Other)
+{
+	// Perform deep copy of this Group
+	if (this != &Other)
+	{
+		// Deep copy active moves
+		ActiveLayeredMoves.Empty(Other.ActiveLayeredMoves.Num());
+		for (int i = 0; i < Other.ActiveLayeredMoves.Num(); ++i)
+		{
+			if (Other.ActiveLayeredMoves[i].IsValid())
+			{
+				FLayeredMoveBase* CopyOfSourcePtr = Other.ActiveLayeredMoves[i]->Clone();
+				ActiveLayeredMoves.Add(TSharedPtr<FLayeredMoveBase>(CopyOfSourcePtr));
+			}
+			else
+			{
+				UE_LOG(LogMover, Warning, TEXT("FLayeredMoveGroup::operator= trying to copy invalid Other Layered Move"));
+			}
+		}
+
+		// Deep copy queued moves
+		QueuedLayeredMoves.Empty(Other.QueuedLayeredMoves.Num());
+		for (int i = 0; i < Other.QueuedLayeredMoves.Num(); ++i)
+		{
+			if (Other.QueuedLayeredMoves[i].IsValid())
+			{
+				FLayeredMoveBase* CopyOfSourcePtr = Other.QueuedLayeredMoves[i]->Clone();
+				QueuedLayeredMoves.Add(TSharedPtr<FLayeredMoveBase>(CopyOfSourcePtr));
+			}
+			else
+			{
+				UE_LOG(LogMover, Warning, TEXT("FLayeredMoveGroup::operator= trying to copy invalid Other Layered Move"));
+			}
+		}
+	}
+
+	return *this;
+}
+
+bool FLayeredMoveGroup::operator==(const FLayeredMoveGroup& Other) const
+{
+	// Deep move-by-move comparison
+	if (ActiveLayeredMoves.Num() != Other.ActiveLayeredMoves.Num())
+	{
+		return false;
+	}
+	if (QueuedLayeredMoves.Num() != Other.QueuedLayeredMoves.Num())
+	{
+		return false;
+	}
+
+
+	for (int32 i = 0; i < ActiveLayeredMoves.Num(); ++i)
+	{
+		if (ActiveLayeredMoves[i].IsValid() == Other.ActiveLayeredMoves[i].IsValid())
+		{
+			if (ActiveLayeredMoves[i].IsValid())
+			{
+				// TODO: Implement deep equality checks
+				// 				if (!ActiveLayeredMoves[i]->MatchesAndHasSameState(Other.ActiveLayeredMoves[i].Get()))
+				// 				{
+				// 					return false; // They're valid and don't match/have same state
+				// 				}
+			}
+		}
+		else
+		{
+			return false; // Mismatch in validity
+		}
+	}
+	for (int32 i = 0; i < QueuedLayeredMoves.Num(); ++i)
+	{
+		if (QueuedLayeredMoves[i].IsValid() == Other.QueuedLayeredMoves[i].IsValid())
+		{
+			if (QueuedLayeredMoves[i].IsValid())
+			{
+				// TODO: Implement deep equality checks
+				// 				if (!QueuedLayeredMoves[i]->MatchesAndHasSameState(Other.QueuedLayeredMoves[i].Get()))
+				// 				{
+				// 					return false; // They're valid and don't match/have same state
+				// 				}
+			}
+		}
+		else
+		{
+			return false; // Mismatch in validity
+		}
+	}
+	return true;
+}
+
+bool FLayeredMoveGroup::operator!=(const FLayeredMoveGroup& Other) const
+{
+	return !(FLayeredMoveGroup::operator==(Other));
+}
+
+void FLayeredMoveGroup::AddStructReferencedObjects(FReferenceCollector& Collector) const
+{
+	for (const TSharedPtr<FLayeredMoveBase>& LayeredMove : ActiveLayeredMoves)
+	{
+		if (LayeredMove.IsValid())
+		{
+			LayeredMove->AddReferencedObjects(Collector);
+		}
+	}
+
+	for (const TSharedPtr<FLayeredMoveBase>& LayeredMove : QueuedLayeredMoves)
+	{
+		if (LayeredMove.IsValid())
+		{
+			LayeredMove->AddReferencedObjects(Collector);
+		}
+	}
+}
+
+FString FLayeredMoveGroup::ToSimpleString() const
+{
+	return FString::Printf(TEXT("FLayeredMoveGroup. Active: %i Queued: %i"), ActiveLayeredMoves.Num(), QueuedLayeredMoves.Num());
+}
+
+TArray<TSharedPtr<FLayeredMoveBase>>::TConstIterator FLayeredMoveGroup::GetActiveMovesIterator() const
+{
+	return ActiveLayeredMoves.CreateConstIterator();
+}
+
+TArray<TSharedPtr<FLayeredMoveBase>>::TConstIterator FLayeredMoveGroup::GetQueuedMovesIterator() const
+{
+	return QueuedLayeredMoves.CreateConstIterator();
+}
+
+void FLayeredMoveGroup::FlushMoveArrays(const UMoverComponent* MoverComp, UMoverBlackboard* SimBlackboard, float CurrentSimTimeMs)
+{
+	bool bResidualVelocityOverridden = false;
+	bool bClampVelocityOverridden = false;
+	// Remove any finished moves
+	ActiveLayeredMoves.RemoveAll([MoverComp, SimBlackboard, CurrentSimTimeMs, &bResidualVelocityOverridden, &bClampVelocityOverridden, this]
+		(const TSharedPtr<FLayeredMoveBase>& Move)
+		{
+			if (Move.IsValid())
+			{
+				if (Move->IsFinished(CurrentSimTimeMs))
+				{
+					GatherResidualVelocitySettings(Move, bResidualVelocityOverridden, bClampVelocityOverridden);
+					Move->EndMove(MoverComp, SimBlackboard, CurrentSimTimeMs);
+					return true;
+				}
+			}
+			else
+			{
+				return true;	
+			}
+
+			return false;
+		});
+
+	// Make any queued moves active
+	for (TSharedPtr<FLayeredMoveBase>& QueuedMove : QueuedLayeredMoves)
+	{
+		ActiveLayeredMoves.Add(QueuedMove);
+		QueuedMove->StartMove(MoverComp, SimBlackboard, CurrentSimTimeMs);
+	}
+
+	QueuedLayeredMoves.Empty();
+}
+
+void FLayeredMoveGroup::GatherResidualVelocitySettings(const TSharedPtr<FLayeredMoveBase>& Move, bool& bResidualVelocityOverridden, bool& bClampVelocityOverridden)
+{
+	if (Move->FinishVelocitySettings.FinishVelocityMode == ELayeredMoveFinishVelocityMode::SetVelocity)
+	{
+		if (Move->MixMode == EMoveMixMode::OverrideVelocity)
+		{
+			if (bResidualVelocityOverridden)
+			{
+				UE_LOG(LogMover, Log, TEXT("Multiple LayeredMove residual settings have a MixMode that overrides. Only one will take effect."));
+			}
+
+			bResidualVelocityOverridden = true;
+			ResidualVelocity = Move->FinishVelocitySettings.SetVelocity;
+		}
+		else if (Move->MixMode == EMoveMixMode::AdditiveVelocity && !bResidualVelocityOverridden)
+		{
+			ResidualVelocity += Move->FinishVelocitySettings.SetVelocity;
+		}
+		else if (Move->MixMode == EMoveMixMode::OverrideAll)
+		{
+			if (bResidualVelocityOverridden)
+			{
+				UE_LOG(LogMover, Log, TEXT("Multiple LayeredMove residual settings have a MixMode that overrides. Only one will take effect."));
+			}
+
+			bResidualVelocityOverridden = true;
+			ResidualVelocity = Move->FinishVelocitySettings.SetVelocity;
+		}
+		else
+		{
+			check(0);	// unhandled case
+		}
+		bApplyResidualVelocity = true;
+	}
+	else if (Move->FinishVelocitySettings.FinishVelocityMode == ELayeredMoveFinishVelocityMode::ClampVelocity)
+	{
+		if (Move->MixMode == EMoveMixMode::OverrideVelocity)
+		{
+			if (bClampVelocityOverridden)
+			{
+				UE_LOG(LogMover, Log, TEXT("Multiple LayeredMove residual settings have a MixMode that overrides. Only one will take effect."));
+			}
+
+			bClampVelocityOverridden = true;
+			ResidualClamping = Move->FinishVelocitySettings.ClampVelocity;
+		}
+		else if (Move->MixMode == EMoveMixMode::AdditiveVelocity && !bClampVelocityOverridden)
+		{
+			if (ResidualClamping < 0)
+			{
+				ResidualClamping = Move->FinishVelocitySettings.ClampVelocity;
+			}
+			// No way to really add clamping so we instead apply it if it was smaller
+			else if (ResidualClamping > Move->FinishVelocitySettings.ClampVelocity)
+			{
+				ResidualClamping = Move->FinishVelocitySettings.ClampVelocity;
+			}
+		}
+		else if (Move->MixMode == EMoveMixMode::OverrideAll)
+		{
+			if (bClampVelocityOverridden)
+			{
+				UE_LOG(LogMover, Log, TEXT("Multiple LayeredMove residual settings have a MixMode that overrides. Only one will take effect."));
+			}
+
+			bClampVelocityOverridden = true;
+			ResidualClamping = Move->FinishVelocitySettings.ClampVelocity;
+		}
+		else
+		{
+			check(0);	// unhandled case
+		}
+	}
+}
+
+struct FLayeredMoveDeleter
+{
+	FORCEINLINE void operator()(FLayeredMoveBase* Object) const
+	{
+		check(Object);
+		UScriptStruct* ScriptStruct = Object->GetScriptStruct();
+		check(ScriptStruct);
+		ScriptStruct->DestroyStruct(Object);
+		FMemory::Free(Object);
+	}
+};
+
+
+/* static */ void FLayeredMoveGroup::NetSerializeLayeredMovesArray(FArchive& Ar, TArray< TSharedPtr<FLayeredMoveBase> >& LayeredMovesArray, uint8 MaxNumLayeredMovesToSerialize /*=MAX_uint8*/)
+{
+	uint8 NumMovesToSerialize;
+	if (Ar.IsSaving())
+	{
+		UE_CLOG(LayeredMovesArray.Num() > MaxNumLayeredMovesToSerialize, LogMover, Warning, TEXT("Too many Layered Moves (%d!) to net serialize. Clamping to %d"),
+			LayeredMovesArray.Num(), MaxNumLayeredMovesToSerialize);
+
+		NumMovesToSerialize = FMath::Min<int32>(LayeredMovesArray.Num(), MaxNumLayeredMovesToSerialize);
+	}
+
+	Ar << NumMovesToSerialize;
+
+	if (Ar.IsLoading())
+	{
+		LayeredMovesArray.SetNumZeroed(NumMovesToSerialize);
+	}
+
+	for (int32 i = 0; i < NumMovesToSerialize && !Ar.IsError(); ++i)
+	{
+		TCheckedObjPtr<UScriptStruct> ScriptStruct = LayeredMovesArray[i].IsValid() ? LayeredMovesArray[i]->GetScriptStruct() : nullptr;
+		UScriptStruct* ScriptStructLocal = ScriptStruct.Get();
+		Ar << ScriptStruct;
+
+		if (ScriptStruct.IsValid())
+		{
+			// Restrict replication to derived classes of FLayeredMoveBase for security reasons:
+			// If FLayeredMoveGroup is replicated through a Server RPC, we need to prevent clients from sending us
+			// arbitrary ScriptStructs due to the allocation/reliance on GetCppStructOps below which could trigger a server crash
+			// for invalid structs. All provided sources are direct children of FLayeredMoveBase and we never expect to have deep hierarchies
+			// so this should not be too costly
+			bool bIsDerivedFromBase = false;
+			UStruct* CurrentSuperStruct = ScriptStruct->GetSuperStruct();
+			while (CurrentSuperStruct)
+			{
+				if (CurrentSuperStruct == FLayeredMoveBase::StaticStruct())
+				{
+					bIsDerivedFromBase = true;
+					break;
+				}
+				CurrentSuperStruct = CurrentSuperStruct->GetSuperStruct();
+			}
+
+			if (bIsDerivedFromBase)
+			{
+				if (Ar.IsLoading())
+				{
+					if (LayeredMovesArray[i].IsValid() && ScriptStructLocal == ScriptStruct.Get())
+					{
+						// What we have locally is the same type as we're being serialized into, so we don't need to
+						// reallocate - just use existing structure
+					}
+					else
+					{
+						// For now, just reset/reallocate the data when loading.
+						// Longer term if we want to generalize this and use it for property replication, we should support
+						// only reallocating when necessary
+						FLayeredMoveBase* NewMove = (FLayeredMoveBase*)FMemory::Malloc(ScriptStruct->GetCppStructOps()->GetSize());
+						ScriptStruct->InitializeStruct(NewMove);
+
+						LayeredMovesArray[i] = TSharedPtr<FLayeredMoveBase>(NewMove, FLayeredMoveDeleter());
+					}
+				}
+
+				LayeredMovesArray[i]->NetSerialize(Ar);
+			}
+			else
+			{
+				UE_LOG(LogMover, Error, TEXT("FLayeredMoveGroup::NetSerialize: ScriptStruct not derived from FLayeredMoveBase attempted to serialize."));
+				Ar.SetError();
+				break;
+			}
+		}
+		else if (ScriptStruct.IsError())
+		{
+			UE_LOG(LogMover, Error, TEXT("FLayeredMoveGroup::NetSerialize: Invalid ScriptStruct serialized."));
+			Ar.SetError();
+			break;
+		}
+	}
+}
+
+void FLayeredMoveGroup::ResetResidualVelocity()
+{
+	bApplyResidualVelocity = false;
+	ResidualVelocity = FVector::ZeroVector;
+	ResidualClamping = -1.f;
+}
+
diff --git a/Engine/Plugins/Experimental/Mover/Source/Mover/Private/MoveLibrary/AirMovementUtils.cpp b/Engine/Plugins/Experimental/Mover/Source/Mover/Private/MoveLibrary/AirMovementUtils.cpp
new file mode 100644
--- /dev/null
+++ b/Engine/Plugins/Experimental/Mover/Source/Mover/Private/MoveLibrary/AirMovementUtils.cpp
@@ -0,0 +1,133 @@
+﻿// Copyright Epic Games, Inc. All Rights Reserved.
+
+#include "MoveLibrary/AirMovementUtils.h"
+
+#include "MoverComponent.h"
+#include "MoveLibrary/FloorQueryUtils.h"
+#include "MoveLibrary/MovementUtils.h"
+
+FProposedMove UAirMovementUtils::ComputeControlledFreeMove(const FFreeMoveParams& InParams)
+{
+	FProposedMove OutMove;
+
+	const FPlane MovementPlane(FVector::ZeroVector, FVector::UpVector);
+
+	OutMove.DirectionIntent = UMovementUtils::ComputeDirectionIntent(InParams.MoveInput, InParams.MoveInputType);
+	OutMove.bHasDirIntent = !OutMove.DirectionIntent.IsNearlyZero();
+
+	FComputeVelocityParams ComputeVelocityParams;
+	ComputeVelocityParams.DeltaSeconds = InParams.DeltaSeconds;
+	ComputeVelocityParams.InitialVelocity = InParams.PriorVelocity;
+	ComputeVelocityParams.MoveDirectionIntent = InParams.MoveInput;
+	ComputeVelocityParams.MaxSpeed = InParams.MaxSpeed;
+	ComputeVelocityParams.TurningBoost = InParams.TurningBoost;
+	ComputeVelocityParams.Deceleration = InParams.Deceleration;
+	ComputeVelocityParams.Acceleration = InParams.Acceleration;
+	
+	OutMove.LinearVelocity = UMovementUtils::ComputeVelocity(ComputeVelocityParams);
+	OutMove.MovePlaneVelocity = UMovementUtils::ConstrainToPlane(OutMove.LinearVelocity, MovementPlane, false);
+
+	// JAH TODO: this is where we can perform turning, based on aux settings. For now, just snap to the intended final orientation.
+	FVector IntendedFacingDir = InParams.OrientationIntent.RotateVector(FVector::ForwardVector).GetSafeNormal();
+	OutMove.AngularVelocity = UMovementUtils::ComputeAngularVelocity(InParams.PriorOrientation, IntendedFacingDir.ToOrientationRotator(), InParams.DeltaSeconds, InParams.TurningRate);
+
+	return OutMove;
+}
+
+bool UAirMovementUtils::IsValidLandingSpot(USceneComponent* UpdatedComponent, UPrimitiveComponent* UpdatedPrimitive, const FVector& Location, const FHitResult& Hit, float FloorSweepDistance, float WalkableFloorZ, FFloorCheckResult& OutFloorResult)
+{
+	OutFloorResult.Clear();
+
+	if (!Hit.bBlockingHit)
+	{
+		return false;
+	}
+
+	if (Hit.bStartPenetrating)
+	{
+		return false;
+	}
+
+	// Reject unwalkable floor normals.
+	if (!UFloorQueryUtils::IsHitSurfaceWalkable(Hit, WalkableFloorZ))
+	{
+		return false;
+	}
+
+	// Make sure floor test passes here.
+	UFloorQueryUtils::FindFloor(UpdatedComponent, UpdatedPrimitive, 
+		FloorSweepDistance, WalkableFloorZ,
+		Location, OutFloorResult);
+
+	if (!OutFloorResult.IsWalkableFloor())
+	{
+		return false;
+	}
+
+	return true;
+}
+
+float UAirMovementUtils::TryMoveToFallAlongSurface(USceneComponent* UpdatedComponent, UPrimitiveComponent* UpdatedPrimitive, UMoverComponent* MoverComponent, const FVector& Delta, float PctOfDeltaToMove, const FQuat Rotation, const FVector& Normal, FHitResult& Hit, bool bHandleImpact, float FloorSweepDistance, float MaxWalkSlopeCosine, FFloorCheckResult& OutFloorResult, FMovementRecord& MoveRecord)
+{
+	OutFloorResult.Clear();
+
+	if (!Hit.bBlockingHit)
+	{
+		return 0.f;
+	}
+
+	float PctOfTimeUsed = 0.f;
+	const FVector OldHitNormal = Normal;
+
+	FVector SlideDelta = UMovementUtils::ComputeSlideDelta(Delta, PctOfDeltaToMove, Normal, Hit);
+
+	if ((SlideDelta | Delta) > 0.f)
+	{
+		// First sliding attempt along surface
+		UMovementUtils::TrySafeMoveUpdatedComponent(UpdatedComponent, UpdatedPrimitive, SlideDelta, Rotation, true, Hit, ETeleportType::None, MoveRecord);
+
+		PctOfTimeUsed = Hit.Time;
+		if (Hit.IsValidBlockingHit())
+		{
+			// Notify first impact
+			if (MoverComponent && bHandleImpact)
+			{
+				FMoverOnImpactParams ImpactParams(NAME_None, Hit, SlideDelta);
+				MoverComponent->HandleImpact(ImpactParams);
+			}
+
+			// Check if we landed
+			if (!IsValidLandingSpot(UpdatedComponent, UpdatedPrimitive, UpdatedPrimitive->GetComponentLocation(),
+				Hit, FloorSweepDistance, MaxWalkSlopeCosine, OutFloorResult))
+			{
+				// We've hit another surface during our first move, so let's try to slide along both of them together
+
+				// Compute new slide normal when hitting multiple surfaces.
+				SlideDelta = UMovementUtils::ComputeTwoWallAdjustedDelta(SlideDelta, Hit, OldHitNormal);
+
+				// Only proceed if the new direction is of significant length and not in reverse of original attempted move.
+				if (!SlideDelta.IsNearlyZero(UE::MoverUtils::SMALL_MOVE_DISTANCE) && (SlideDelta | Delta) > 0.f)
+				{
+					// Perform second move, taking 2 walls into account
+					UMovementUtils::TrySafeMoveUpdatedComponent(UpdatedComponent, UpdatedPrimitive, SlideDelta, Rotation, true, Hit, ETeleportType::None, MoveRecord);
+					PctOfTimeUsed += (Hit.Time * (1.f - PctOfTimeUsed));
+
+					// Notify second impact
+					if (MoverComponent && bHandleImpact && Hit.bBlockingHit)
+					{
+						FMoverOnImpactParams ImpactParams(NAME_None, Hit, SlideDelta);
+						MoverComponent->HandleImpact(ImpactParams);
+					}
+
+					// Check if we've landed, to acquire floor result
+					IsValidLandingSpot(UpdatedComponent, UpdatedPrimitive, UpdatedPrimitive->GetComponentLocation(),
+						Hit, FloorSweepDistance, MaxWalkSlopeCosine, OutFloorResult);
+				}
+			}
+		}
+
+		return FMath::Clamp(PctOfTimeUsed, 0.f, 1.f);
+	}
+
+	return 0.f;
+}
diff --git a/Engine/Plugins/Experimental/Mover/Source/Mover/Private/MoveLibrary/BasedMovementUtils.cpp b/Engine/Plugins/Experimental/Mover/Source/Mover/Private/MoveLibrary/BasedMovementUtils.cpp
new file mode 100644
--- /dev/null
+++ b/Engine/Plugins/Experimental/Mover/Source/Mover/Private/MoveLibrary/BasedMovementUtils.cpp
@@ -0,0 +1,287 @@
+// Copyright Epic Games, Inc. All Rights Reserved.
+
+#include "MoveLibrary/BasedMovementUtils.h"
+#include "MoveLibrary/MovementUtils.h"
+#include "MoveLibrary/FloorQueryUtils.h"
+#include "Components/PrimitiveComponent.h"
+#include "MoverLog.h"
+#include "Kismet/KismetMathLibrary.h"
+
+#include UE_INLINE_GENERATED_CPP_BY_NAME(BasedMovementUtils)
+
+void FRelativeBaseInfo::Clear()
+{
+	MovementBase = nullptr;
+	BoneName = NAME_None;
+	Location = FVector::ZeroVector;
+	Rotation = FQuat::Identity;
+}
+
+bool FRelativeBaseInfo::HasRelativeInfo() const
+{
+	return MovementBase != nullptr;
+}
+
+bool FRelativeBaseInfo::UsesSameBase(const FRelativeBaseInfo& Other) const
+{
+	return UsesSameBase(Other.MovementBase, Other.BoneName);
+}
+
+bool FRelativeBaseInfo::UsesSameBase(const UPrimitiveComponent* OtherComp, FName OtherBoneName) const
+{
+	return HasRelativeInfo()
+		&& (MovementBase == OtherComp)
+		&& (BoneName == OtherBoneName);
+}
+
+void FRelativeBaseInfo::SetFromFloorResult(const FFloorCheckResult& FloorTestResult)
+{
+	bool bDidSucceed = false;
+
+	if (FloorTestResult.bWalkableFloor)
+	{
+		MovementBase = FloorTestResult.HitResult.GetComponent();
+
+		if (MovementBase)
+		{
+			BoneName = FloorTestResult.HitResult.BoneName;
+			bDidSucceed  = UBasedMovementUtils::GetMovementBaseTransform(MovementBase, BoneName, /*out*/Location, /*out*/Rotation);
+			bDidSucceed &= UBasedMovementUtils::TransformWorldLocationToBased(MovementBase, BoneName, FloorTestResult.HitResult.ImpactPoint, /*out*/ContactLocalPosition);
+		}
+	}
+
+	if (!bDidSucceed)
+	{
+		Clear();
+	}
+}
+
+void FRelativeBaseInfo::SetFromComponent(UPrimitiveComponent* InRelativeComp, FName InBoneName)
+{
+	bool bDidSucceed = false;
+
+	MovementBase = InRelativeComp;
+
+	if (MovementBase)
+	{
+		BoneName = InBoneName;
+		bDidSucceed = UBasedMovementUtils::GetMovementBaseTransform(MovementBase, BoneName, /*out*/Location, /*out*/Rotation);
+	}
+
+	if (!bDidSucceed)
+	{
+		Clear();
+	}
+}
+
+bool UBasedMovementUtils::IsADynamicBase(const UPrimitiveComponent* MovementBase)
+{
+	return (MovementBase && MovementBase->Mobility == EComponentMobility::Movable);
+}
+
+bool UBasedMovementUtils::TryMoveToStayWithBase(USceneComponent* UpdatedComponent, UPrimitiveComponent* UpdatedPrimitive, const FRelativeBaseInfo& OldBaseInfo, FMovementRecord& MoveRecord, const bool& bIgnoreBaseRotation)
+{
+	FVector NewBaseLocation;
+	FQuat NewBaseQuat;
+
+
+	if (UBasedMovementUtils::GetMovementBaseTransform(OldBaseInfo.MovementBase, OldBaseInfo.BoneName, OUT NewBaseLocation, OUT NewBaseQuat))
+	{
+		const bool bDidBaseRotationChange = !OldBaseInfo.Rotation.Equals(NewBaseQuat, UE_SMALL_NUMBER);
+		const bool bDidBaseLocationChange = (OldBaseInfo.Location != NewBaseLocation);
+
+		// Find change in rotation
+		FQuat DeltaQuat = FQuat::Identity;
+		FVector DeltaLoc = FVector::ZeroVector;
+		FQuat TargetQuat = UpdatedComponent->GetComponentQuat();
+
+		if (bDidBaseRotationChange && !bIgnoreBaseRotation)
+		{
+			DeltaQuat = NewBaseQuat * OldBaseInfo.Rotation.Inverse();
+			TargetQuat = DeltaQuat * TargetQuat;
+
+			// TODO: make this respect the "up" direction, rather than assuming +Z = up
+			FVector TargetForwVector = TargetQuat.GetForwardVector();
+			TargetForwVector.Z = 0.f;
+			TargetForwVector.Normalize();
+
+			TargetQuat = UKismetMathLibrary::MakeRotFromX(TargetForwVector).Quaternion();
+		}
+
+		if (bDidBaseRotationChange || bDidBaseLocationChange)
+		{
+			// Calculate new transform matrix of base actor (ignoring scale).
+			const FQuatRotationTranslationMatrix OldLocalToWorld(OldBaseInfo.Rotation, OldBaseInfo.Location);
+			const FQuatRotationTranslationMatrix NewLocalToWorld(NewBaseQuat, NewBaseLocation);
+
+			// Find change in location
+			// NOTE that we need to use the floor hit location, not the actor's root position which may be floating above the base
+			const FVector NewWorldBaseContactPos = NewLocalToWorld.TransformPosition(OldBaseInfo.ContactLocalPosition);
+
+			const FVector OldWorldBaseContactPos = OldLocalToWorld.TransformPosition(OldBaseInfo.ContactLocalPosition);
+
+			DeltaLoc = NewWorldBaseContactPos - OldWorldBaseContactPos;
+
+			EMoveComponentFlags MoveComponentFlags = MOVECOMP_IgnoreBases;
+
+			const bool bSweep = true;
+			FHitResult MoveHitResult;
+			
+			bool bDidMove = UMovementUtils::TryMoveUpdatedComponent_Internal(UpdatedComponent, DeltaLoc, TargetQuat, bSweep, MoveComponentFlags, &MoveHitResult, ETeleportType::None);
+			
+			return bDidMove;
+		}
+
+		// TODO: rework this to ensure MoveRecord contains the appropriate thing
+	}
+
+	return false;
+}
+
+bool UBasedMovementUtils::GetMovementBaseTransform(const UPrimitiveComponent* MovementBase, const FName BoneName, FVector& OutLocation, FQuat& OutQuat)
+{
+	if (MovementBase)
+	{
+		bool bBoneNameIsInvalid = false;
+
+		if (BoneName != NAME_None)
+		{
+			// Check if this socket or bone exists (DoesSocketExist checks for either, as does requesting the transform).
+			if (MovementBase->DoesSocketExist(BoneName))
+			{
+				MovementBase->GetSocketWorldLocationAndRotation(BoneName, OutLocation, OutQuat);
+				return true;
+			}
+
+			bBoneNameIsInvalid = true;
+			UE_LOG(LogMover, Warning, TEXT("GetMovementBaseTransform(): Invalid bone or socket '%s' for PrimitiveComponent base %s. Using component's root transform instead."), *BoneName.ToString(), *GetPathNameSafe(MovementBase));
+		}
+
+		OutLocation = MovementBase->GetComponentLocation();
+		OutQuat = MovementBase->GetComponentQuat();
+		return !bBoneNameIsInvalid;
+	}
+
+	// nullptr MovementBase
+	OutLocation = FVector::ZeroVector;
+	OutQuat = FQuat::Identity;
+	return false;
+}
+
+
+bool UBasedMovementUtils::TransformBasedLocationToWorld(const UPrimitiveComponent* MovementBase, const FName BoneName, FVector LocalLocation, FVector& OutLocationWorldSpace)
+{
+	FVector BaseLocation;
+	FQuat BaseQuat;
+	
+	if (GetMovementBaseTransform(MovementBase, BoneName, /*out*/ BaseLocation, /*out*/ BaseQuat))
+	{ 
+		TransformLocationToWorld(BaseLocation, BaseQuat, LocalLocation, OutLocationWorldSpace);
+		return true;
+	}
+	
+	return false;
+}
+
+
+bool UBasedMovementUtils::TransformWorldLocationToBased(const UPrimitiveComponent* MovementBase, const FName BoneName, FVector WorldSpaceLocation, FVector& OutLocalLocation)
+{
+	FVector BaseLocation;
+	FQuat BaseQuat;
+	if (GetMovementBaseTransform(MovementBase, BoneName, /*out*/ BaseLocation, /*out*/ BaseQuat))
+	{
+		TransformLocationToLocal(BaseLocation, BaseQuat, WorldSpaceLocation, OutLocalLocation);
+		return true;
+	}
+
+	return false;
+}
+
+
+bool UBasedMovementUtils::TransformBasedDirectionToWorld(const UPrimitiveComponent* MovementBase, const FName BoneName, FVector LocalDirection, FVector& OutDirectionWorldSpace)
+{
+	FVector IgnoredLocation;
+	FQuat BaseQuat;
+	if (GetMovementBaseTransform(MovementBase, BoneName, /*out*/ IgnoredLocation, /*out*/ BaseQuat))
+	{
+		TransformDirectionToWorld(BaseQuat, LocalDirection, OutDirectionWorldSpace);
+		return true;
+	}
+
+	return false;
+}
+
+
+bool UBasedMovementUtils::TransformWorldDirectionToBased(const UPrimitiveComponent* MovementBase, const FName BoneName, FVector WorldSpaceDirection, FVector& OutLocalDirection)
+{
+	FVector IgnoredLocation;
+	FQuat BaseQuat;
+	if (GetMovementBaseTransform(MovementBase, BoneName, /*out*/ IgnoredLocation, /*out*/ BaseQuat))
+	{
+		TransformDirectionToLocal(BaseQuat, WorldSpaceDirection, OutLocalDirection);
+		return true;
+	}
+
+	return false;
+}
+
+
+bool UBasedMovementUtils::TransformBasedRotatorToWorld(const UPrimitiveComponent* MovementBase, const FName BoneName, FRotator LocalRotator, FRotator& OutWorldSpaceRotator)
+{
+	FVector IgnoredLocation;
+	FQuat BaseQuat;
+	
+	if (GetMovementBaseTransform(MovementBase, BoneName, /*out*/ IgnoredLocation, /*out*/ BaseQuat))
+	{
+		TransformRotatorToWorld(BaseQuat, LocalRotator, OutWorldSpaceRotator);
+		return true;
+	}
+
+	return false;
+}
+
+
+bool UBasedMovementUtils::TransformWorldRotatorToBased(const UPrimitiveComponent* MovementBase, const FName BoneName, FRotator WorldSpaceRotator, FRotator& OutLocalRotator)
+{
+	FVector IgnoredLocation;
+	FQuat BaseQuat;
+	if (GetMovementBaseTransform(MovementBase, BoneName, /*out*/ IgnoredLocation, /*out*/ BaseQuat))
+	{
+		TransformRotatorToLocal(BaseQuat, WorldSpaceRotator, OutLocalRotator);
+		return true;
+	}
+	return false;
+}
+
+
+void UBasedMovementUtils::TransformLocationToWorld(FVector BasePos, FQuat BaseQuat, FVector LocalLocation, FVector& OutLocationWorldSpace)
+{
+	OutLocationWorldSpace = FTransform(BaseQuat, BasePos).TransformPositionNoScale(LocalLocation);
+}
+
+void UBasedMovementUtils::TransformLocationToLocal(FVector BasePos, FQuat BaseQuat, FVector WorldSpaceLocation, FVector& OutLocalLocation)
+{
+	OutLocalLocation = FTransform(BaseQuat, BasePos).InverseTransformPositionNoScale(WorldSpaceLocation);
+}
+
+void UBasedMovementUtils::TransformDirectionToWorld(FQuat BaseQuat, FVector LocalDirection, FVector& OutDirectionWorldSpace)
+{
+	OutDirectionWorldSpace = BaseQuat.RotateVector(LocalDirection);
+}
+
+void UBasedMovementUtils::TransformDirectionToLocal(FQuat BaseQuat, FVector WorldSpaceDirection, FVector& OutLocalDirection)
+{
+	OutLocalDirection = BaseQuat.UnrotateVector(WorldSpaceDirection);
+}
+
+void UBasedMovementUtils::TransformRotatorToWorld(FQuat BaseQuat, FRotator LocalRotator, FRotator& OutWorldSpaceRotator)
+{
+	FQuat LocalQuat(LocalRotator);
+	OutWorldSpaceRotator = (BaseQuat * LocalQuat).Rotator();
+}
+
+void UBasedMovementUtils::TransformRotatorToLocal(FQuat BaseQuat, FRotator WorldSpaceRotator, FRotator& OutLocalRotator)
+{
+	FQuat WorldQuat(WorldSpaceRotator);
+	OutLocalRotator = (BaseQuat.Inverse() * WorldQuat).Rotator();
+}
diff --git a/Engine/Plugins/Experimental/Mover/Source/Mover/Private/MoveLibrary/FloorQueryUtils.cpp b/Engine/Plugins/Experimental/Mover/Source/Mover/Private/MoveLibrary/FloorQueryUtils.cpp
new file mode 100644
--- /dev/null
+++ b/Engine/Plugins/Experimental/Mover/Source/Mover/Private/MoveLibrary/FloorQueryUtils.cpp
@@ -0,0 +1,225 @@
+﻿// Copyright Epic Games, Inc. All Rights Reserved.
+
+#include "MoveLibrary/FloorQueryUtils.h"
+#include "MoveLibrary/MovementUtils.h"
+#include "Engine/World.h"
+
+namespace UE::FloorQueryUtility
+{
+	const float MIN_FLOOR_DIST = 1.9f;
+	const float MAX_FLOOR_DIST = 2.4f;
+	const float SWEEP_EDGE_REJECT_DISTANCE = 0.15f;
+}
+
+void UFloorQueryUtils::FindFloor(const USceneComponent* UpdatedComponent, const UPrimitiveComponent* UpdatedPrimitive, float FloorSweepDistance, float MaxWalkSlopeCosine, const FVector& Location, FFloorCheckResult& OutFloorResult)
+{
+	if (!UpdatedComponent || !UpdatedComponent->IsQueryCollisionEnabled())
+	{
+		OutFloorResult.Clear();
+		return;
+	}
+
+	// Sweep for the floor
+	// TODO: Might need to plug in a different value for LineTraceDistance - using the same value as FloorSweepDistance for now - function takes both so we can plug in different values if needed
+	ComputeFloorDist(UpdatedComponent, UpdatedPrimitive,FloorSweepDistance, FloorSweepDistance, MaxWalkSlopeCosine, Location, OutFloorResult);
+}
+
+void UFloorQueryUtils::ComputeFloorDist(const USceneComponent* UpdatedComponent, const UPrimitiveComponent* UpdatedPrimitive, float LineTraceDistance, float FloorSweepDistance, float MaxWalkSlopeCosine, const FVector& Location, FFloorCheckResult& OutFloorResult)
+{
+	OutFloorResult.Clear();
+
+	FCollisionQueryParams QueryParams(SCENE_QUERY_STAT(ComputeFloorDist), false, UpdatedPrimitive->GetOwner());
+	FCollisionResponseParams ResponseParam;
+	UMovementUtils::InitCollisionParams(UpdatedPrimitive, QueryParams, ResponseParam);
+	const ECollisionChannel CollisionChannel = UpdatedComponent->GetCollisionObjectType();
+
+	// TODO: pluggable shapes
+	float PawnRadius = 0.0f;
+	float PawnHalfHeight = 0.0f;
+	UpdatedPrimitive->CalcBoundingCylinder(PawnRadius, PawnHalfHeight);
+
+	bool bBlockingHit = false;
+	
+	// Sweep test
+	if (FloorSweepDistance > 0.f)
+	{
+		// Use a shorter height to avoid sweeps giving weird results if we start on a surface.
+		// This also allows us to adjust out of penetrations.
+		const float ShrinkScale = 0.9f;
+		const float ShrinkScaleOverlap = 0.1f;
+		float ShrinkHeight = (PawnHalfHeight - PawnRadius) * (1.f - ShrinkScale);
+		float TraceDist = FloorSweepDistance + ShrinkHeight;
+		FCollisionShape CapsuleShape = FCollisionShape::MakeCapsule(PawnRadius, PawnHalfHeight - ShrinkHeight);
+
+		FHitResult Hit(1.f);
+		// TODO: arbitrary direction
+		FVector SweepDirection = FVector(0.f, 0.f, -TraceDist);
+		bBlockingHit = FloorSweepTest(UpdatedPrimitive, Hit, Location, Location + SweepDirection, CollisionChannel, CapsuleShape, QueryParams, ResponseParam);
+
+		if (bBlockingHit)
+		{
+			// Reject hits adjacent to us, we only care about hits on the bottom portion of our capsule.
+			// Check 2D distance to impact point, reject if within a tolerance from radius.
+			if (Hit.bStartPenetrating || !IsWithinEdgeTolerance(Location, Hit.ImpactPoint, CapsuleShape.Capsule.Radius))
+			{
+				// Use a capsule with a slightly smaller radius and shorter height to avoid the adjacent object.
+				// Capsule must not be nearly zero or the trace will fall back to a line trace from the start point and have the wrong length.
+				CapsuleShape.Capsule.Radius = FMath::Max(0.f, CapsuleShape.Capsule.Radius - UE::FloorQueryUtility::SWEEP_EDGE_REJECT_DISTANCE - KINDA_SMALL_NUMBER);
+				if (!CapsuleShape.IsNearlyZero())
+				{
+					ShrinkHeight = (PawnHalfHeight - PawnRadius) * (1.f - ShrinkScaleOverlap);
+					TraceDist = FloorSweepDistance + ShrinkHeight;
+					SweepDirection = FVector(0.f, 0.f, -TraceDist);
+					CapsuleShape.Capsule.HalfHeight = FMath::Max(PawnHalfHeight - ShrinkHeight, CapsuleShape.Capsule.Radius);
+					Hit.Reset(1.f, false);
+
+					bBlockingHit = FloorSweepTest(UpdatedPrimitive, Hit, Location, Location + SweepDirection, CollisionChannel, CapsuleShape, QueryParams, ResponseParam);
+				}
+			}
+
+			// Reduce hit distance by ShrinkHeight because we shrank the capsule for the trace.
+			// We allow negative distances here, because this allows us to pull out of penetrations.
+			// JAH TODO: move magic numbers to a common location
+			const float MaxPenetrationAdjust = FMath::Max(UE::FloorQueryUtility::MAX_FLOOR_DIST, PawnRadius);
+			const float SweepResult = FMath::Max(-MaxPenetrationAdjust, Hit.Time * TraceDist - ShrinkHeight);
+
+			OutFloorResult.SetFromSweep(Hit, SweepResult, false);
+			if (Hit.IsValidBlockingHit() && IsHitSurfaceWalkable(Hit, MaxWalkSlopeCosine))
+			{
+				if (SweepResult <= FloorSweepDistance)
+				{
+					// Hit within test distance.
+					OutFloorResult.bWalkableFloor = true;
+					return;
+				}
+			}
+		}
+	}
+
+	// Since we require a longer sweep than line trace, we don't want to run the line trace if the sweep missed everything.
+	// We do however want to try a line trace if the sweep was stuck in penetration.
+	if (!OutFloorResult.bBlockingHit && !OutFloorResult.HitResult.bStartPenetrating)
+	{
+		OutFloorResult.FloorDist = FloorSweepDistance;
+		return;
+	}
+
+	// Line trace
+	if (LineTraceDistance > 0.f)
+	{
+		const float ShrinkHeight = PawnHalfHeight;
+		const FVector LineTraceStart = Location;	
+		const float TraceDist = LineTraceDistance + ShrinkHeight;
+		const FVector Down = FVector(0.f, 0.f, -TraceDist);
+		QueryParams.TraceTag = SCENE_QUERY_STAT_NAME_ONLY(FloorLineTrace);
+
+		FHitResult Hit(1.f);
+		bBlockingHit = UpdatedComponent->GetWorld()->LineTraceSingleByChannel(Hit, LineTraceStart, LineTraceStart + Down, CollisionChannel, QueryParams, ResponseParam);
+		
+		if (bBlockingHit && Hit.Time > 0.f)
+		{
+			// Reduce hit distance by ShrinkHeight because we started the trace higher than the base.
+			// We allow negative distances here, because this allows us to pull out of penetrations.
+			const float MaxPenetrationAdjust = FMath::Max(UE::FloorQueryUtility::MAX_FLOOR_DIST, PawnRadius);
+			const float LineResult = FMath::Max(-MaxPenetrationAdjust, Hit.Time * TraceDist - ShrinkHeight);
+			
+			OutFloorResult.bBlockingHit = true;
+			if (LineResult <= LineTraceDistance && IsHitSurfaceWalkable(Hit, MaxWalkSlopeCosine))
+			{
+				OutFloorResult.SetFromLineTrace(Hit, OutFloorResult.FloorDist, LineResult, true);
+				return;
+			}
+		}
+	}
+
+	// No hits were acceptable.
+	OutFloorResult.bWalkableFloor = false;
+}
+
+bool UFloorQueryUtils::FloorSweepTest(const UPrimitiveComponent* UpdatedPrimitive, FHitResult& OutHit, const FVector& Start, const FVector& End, ECollisionChannel TraceChannel, const struct FCollisionShape& CollisionShape, const struct FCollisionQueryParams& Params, const struct FCollisionResponseParams& ResponseParam)
+{
+	bool bBlockingHit = false;
+
+	if (UpdatedPrimitive)
+	{
+		bBlockingHit = UpdatedPrimitive->GetWorld()->SweepSingleByChannel(OutHit, Start, End, FQuat::Identity, TraceChannel, CollisionShape, Params, ResponseParam);
+	}
+
+	return bBlockingHit;
+}
+
+bool UFloorQueryUtils::IsHitSurfaceWalkable(const FHitResult& Hit, float MaxWalkSlopeCosine)
+{
+	// JAH TODO: refactor this to support arbitrary movement planes (currently assumes it's the Z=0 plane)
+
+	if (!Hit.IsValidBlockingHit())
+	{
+		// No hit, or starting in penetration
+		return false;
+	}
+
+	// Never walk up vertical surfaces.
+	if (Hit.ImpactNormal.Z < KINDA_SMALL_NUMBER)
+	{
+		return false;
+	}
+
+	float TestWalkableZ = MaxWalkSlopeCosine;
+
+	// See if this component overrides the walkable floor z.
+	const UPrimitiveComponent* HitComponent = Hit.Component.Get();
+	if (HitComponent)
+	{
+		const FWalkableSlopeOverride& SlopeOverride = HitComponent->GetWalkableSlopeOverride();
+		TestWalkableZ = SlopeOverride.ModifyWalkableFloorZ(TestWalkableZ);
+	}
+
+	// Can't walk on this surface if it is too steep.
+	if (Hit.ImpactNormal.Z < TestWalkableZ)
+	{
+		return false;
+	}
+
+	return true;
+}
+
+bool UFloorQueryUtils::IsWithinEdgeTolerance(const FVector& CapsuleLocation, const FVector& TestImpactPoint, float CapsuleRadius)
+{
+	const float DistFromCenterSq = (TestImpactPoint - CapsuleLocation).SizeSquared2D();
+	const float ReducedRadiusSq = FMath::Square(FMath::Max(UE::FloorQueryUtility::SWEEP_EDGE_REJECT_DISTANCE + UE_KINDA_SMALL_NUMBER, CapsuleRadius - UE::FloorQueryUtility::SWEEP_EDGE_REJECT_DISTANCE));
+	return DistFromCenterSq < ReducedRadiusSq;
+}
+
+void FFloorCheckResult::SetFromSweep(const FHitResult& InHit, const float InSweepFloorDist, const bool bIsWalkableFloor)
+{
+	bBlockingHit = InHit.IsValidBlockingHit();
+	bWalkableFloor = bIsWalkableFloor;
+	FloorDist = InSweepFloorDist;
+	HitResult = InHit;
+	bLineTrace = false;
+	LineDist = 0.f;
+}
+
+void FFloorCheckResult::SetFromLineTrace(const FHitResult& InHit, const float InSweepFloorDist, const float InLineDist, const bool bIsWalkableFloor)
+{
+	// We require a sweep that hit if we are going to use a line result.
+	ensure(HitResult.bBlockingHit);
+	if (HitResult.bBlockingHit && InHit.bBlockingHit)
+	{
+		// Override most of the sweep result with the line result, but save some values
+		FHitResult OldHit(HitResult);
+		HitResult = InHit;
+
+		// Restore some of the old values. We want the new normals and hit actor, however.
+		HitResult.Time = OldHit.Time;
+		HitResult.ImpactPoint = OldHit.ImpactPoint;
+		HitResult.Location = OldHit.Location;
+		HitResult.TraceStart = OldHit.TraceStart;
+		HitResult.TraceEnd = OldHit.TraceEnd;
+		bLineTrace = true;
+		LineDist = InLineDist;
+		
+		FloorDist = InSweepFloorDist;
+		bWalkableFloor = bIsWalkableFloor;
+	}
+}
