Moving Mover, MoverExamples, MoverTests out of Restricted engine folder

diff --git a/Engine/Plugins/Experimental/Mover/Source/Mover/Private/MoveLibrary/GroundMovementUtils.cpp b/Engine/Plugins/Experimental/Mover/Source/Mover/Private/MoveLibrary/GroundMovementUtils.cpp
new file mode 100644
--- /dev/null
+++ b/Engine/Plugins/Experimental/Mover/Source/Mover/Private/MoveLibrary/GroundMovementUtils.cpp
@@ -0,0 +1,541 @@
+﻿// Copyright Epic Games, Inc. All Rights Reserved.
+
+#include "MoveLibrary/GroundMovementUtils.h"
+
+#include "MoverComponent.h"
+#include "MoverLog.h"
+#include "Components/CapsuleComponent.h"
+#include "GameFramework/Pawn.h"
+#include "MoveLibrary/FloorQueryUtils.h"
+#include "MoveLibrary/MovementRecord.h"
+#include "MoveLibrary/MovementUtils.h"
+
+FProposedMove UGroundMovementUtils::ComputeControlledGroundMove(const FGroundMoveParams& InParams)
+{
+	FProposedMove OutMove;
+
+	const FVector MoveDirIntent = UMovementUtils::ComputeDirectionIntent(InParams.MoveInput, InParams.MoveInputType);
+
+	const FPlane MovementPlane(FVector::ZeroVector, FVector::UpVector);
+	FVector MoveDirIntentInMovementPlane = UMovementUtils::ConstrainToPlane(MoveDirIntent, MovementPlane, true);
+
+	const FPlane GroundSurfacePlane(FVector::ZeroVector, InParams.GroundNormal);
+	OutMove.DirectionIntent = UMovementUtils::ConstrainToPlane(MoveDirIntentInMovementPlane, GroundSurfacePlane, true);
+	
+	OutMove.bHasDirIntent = !OutMove.DirectionIntent.IsNearlyZero();
+
+	FComputeVelocityParams ComputeVelocityParams;
+	ComputeVelocityParams.DeltaSeconds = InParams.DeltaSeconds;
+	ComputeVelocityParams.InitialVelocity = InParams.PriorVelocity;
+	ComputeVelocityParams.MoveDirectionIntent = InParams.MoveInput;
+	ComputeVelocityParams.MaxSpeed = InParams.MaxSpeed;
+	ComputeVelocityParams.TurningBoost = InParams.TurningBoost;
+	ComputeVelocityParams.Deceleration = InParams.Deceleration;
+	ComputeVelocityParams.Acceleration = InParams.Acceleration;
+	ComputeVelocityParams.Friction = InParams.Friction;
+	
+	// Figure out linear velocity
+	OutMove.MovePlaneVelocity = UMovementUtils::ComputeVelocity(ComputeVelocityParams);
+	OutMove.LinearVelocity = UMovementUtils::ConstrainToPlane(OutMove.MovePlaneVelocity, GroundSurfacePlane, true);
+
+	// Linearly rotate in place
+	OutMove.AngularVelocity = UMovementUtils::ComputeAngularVelocity(InParams.PriorOrientation, InParams.OrientationIntent, InParams.DeltaSeconds, InParams.TurningRate);
+
+	return OutMove;
+}
+
+static const FName StepUpSubstepName = "StepUp";
+static const FName StepFwdSubstepName = "StepFwd";
+static const FName StepDownSubstepName = "StepDown";
+static const FName SlideSubstepName = "SlideFromStep";
+
+bool UGroundMovementUtils::TryMoveToStepUp(USceneComponent* UpdatedComponent, UPrimitiveComponent* UpdatedPrimitive, UMoverComponent* MoverComponent, const FVector& GravDir, float MaxStepHeight, float MaxWalkSlopeCosine, float FloorSweepDistance, const FVector& MoveDelta, const FHitResult& MoveHitResult, const FFloorCheckResult& CurrentFloor, bool bIsFalling, FOptionalFloorCheckResult* OutFloorTestResult, FMovementRecord& MoveRecord)
+{
+	UCapsuleComponent* CapsuleComponent = Cast<UCapsuleComponent>(UpdatedPrimitive);
+
+	if (CapsuleComponent == nullptr || !CanStepUpOnHitSurface(MoveHitResult) || MaxStepHeight <= 0.f)
+	{
+		return false;
+	}
+
+	TArray<FMovementSubstep> QueuedSubsteps;	// keeping track of substeps before committing, because some moves can be backed out
+
+
+	const FVector OldLocation = UpdatedPrimitive->GetComponentLocation();
+	FVector LastComponentLocation = OldLocation;
+
+	float PawnRadius, PawnHalfHeight;
+
+	CapsuleComponent->GetScaledCapsuleSize(PawnRadius, PawnHalfHeight);
+
+	// Don't bother stepping up if top of capsule is hitting something.
+	const float InitialImpactZ = MoveHitResult.ImpactPoint.Z;
+	if (InitialImpactZ > OldLocation.Z + (PawnHalfHeight - PawnRadius))
+	{
+		UE_LOG(LogMover, VeryVerbose, TEXT("Not stepping up due to top of capsule hitting something"));
+		return false;
+	}
+
+	// TODO: We should rely on movement plane normal, rather than gravity direction
+	if (GravDir.IsZero())
+	{
+		UE_LOG(LogMover, VeryVerbose, TEXT("Not stepping up because there's no gravity"));
+		return false;
+	}
+
+	// Gravity should be a normalized direction
+	ensure(GravDir.IsNormalized());
+
+	float StepTravelUpHeight = MaxStepHeight;
+	float StepTravelDownHeight = StepTravelUpHeight;
+	const float StepSideZ = -1.f * FVector::DotProduct(MoveHitResult.ImpactNormal, GravDir);
+	float PawnInitialFloorBaseZ = OldLocation.Z - PawnHalfHeight;
+	float PawnFloorPointZ = PawnInitialFloorBaseZ;
+
+	//if (IsMovingOnGround() && CurrentFloor.IsWalkableFloor())
+	if (CurrentFloor.IsWalkableFloor())
+	{
+		// Since we float a variable amount off the floor, we need to enforce max step height off the actual point of impact with the floor.
+		const float FloorDist = FMath::Max(0.f, CurrentFloor.GetDistanceToFloor());
+		PawnInitialFloorBaseZ -= FloorDist;
+		StepTravelUpHeight = FMath::Max(StepTravelUpHeight - FloorDist, 0.f);
+		StepTravelDownHeight = (MaxStepHeight + UE::FloorQueryUtility::MAX_FLOOR_DIST * 2.f);
+
+		const bool bHitVerticalFace = !UFloorQueryUtils::IsWithinEdgeTolerance(MoveHitResult.Location, MoveHitResult.ImpactPoint, PawnRadius);
+		if (!CurrentFloor.bLineTrace && !bHitVerticalFace)
+		{
+			PawnFloorPointZ = CurrentFloor.HitResult.ImpactPoint.Z;
+		}
+		else
+		{
+			// Base floor point is the base of the capsule moved down by how far we are hovering over the surface we are hitting.
+			PawnFloorPointZ -= CurrentFloor.FloorDist;
+		}
+	}
+
+	// Don't step up if the impact is below us, accounting for distance from floor.
+	if (InitialImpactZ <= PawnInitialFloorBaseZ)
+	{
+		UE_LOG(LogMover, VeryVerbose, TEXT("Not stepping up because the impact is below us"));
+		return false;
+	}
+
+	// Scope our movement updates, and do not apply them until all intermediate moves are completed.
+	FScopedMovementUpdate ScopedStepUpMovement(UpdatedComponent, EScopedUpdate::DeferredUpdates);
+
+	// step up - treat as vertical wall
+	FHitResult SweepUpHit(1.f);
+	const FQuat PawnRotation = UpdatedComponent->GetComponentQuat();
+
+	const FVector UpAdjustment = -GravDir * StepTravelUpHeight;
+	const bool bDidStepUp = UMovementUtils::TryMoveUpdatedComponent_Internal(UpdatedComponent, UpAdjustment, PawnRotation, true, MOVECOMP_NoFlags, &SweepUpHit, ETeleportType::None);
+
+	UE_LOG(LogMover, VeryVerbose, TEXT("TryMoveToStepUp Up: %s (role %i) UpAdjustment=%s DidMove=%i"),
+		*GetNameSafe(UpdatedComponent->GetOwner()), UpdatedComponent->GetOwnerRole(), *UpAdjustment.ToCompactString(), bDidStepUp);
+
+	if (SweepUpHit.bStartPenetrating)
+	{
+		// Undo movement
+		UE_LOG(LogMover, VeryVerbose, TEXT("Reverting step-up attempt because we started in a penetrating state"));
+		ScopedStepUpMovement.RevertMove();
+		return false;
+	}
+
+	// Cache upwards substep
+	QueuedSubsteps.Add(FMovementSubstep(StepUpSubstepName, UpdatedPrimitive->GetComponentLocation()-LastComponentLocation, false));
+	LastComponentLocation = UpdatedPrimitive->GetComponentLocation();
+
+	// step fwd
+	FHitResult StepFwdHit(1.f);
+	const bool bDidStepFwd = UMovementUtils::TryMoveUpdatedComponent_Internal(UpdatedComponent, MoveDelta, PawnRotation, true, MOVECOMP_NoFlags, &StepFwdHit, ETeleportType::None);
+
+	UE_LOG(LogMover, VeryVerbose, TEXT("TryMoveToStepUp Fwd: %s (role %i) MoveDelta=%s DidMove=%i"),
+		*GetNameSafe(UpdatedComponent->GetOwner()), UpdatedComponent->GetOwnerRole(), *MoveDelta.ToCompactString(), bDidStepFwd);
+
+	// Check result of forward movement
+	if (StepFwdHit.bBlockingHit)
+	{
+		if (StepFwdHit.bStartPenetrating)
+		{
+			// Undo movement
+			UE_LOG(LogMover, VeryVerbose, TEXT("Reverting step-fwd attempt during step-up, because we started in a penetrating state"));
+			ScopedStepUpMovement.RevertMove();
+			return false;
+		}
+
+		// If we hit something above us and also something ahead of us, we should notify about the upward hit as well.
+		// The forward hit will be handled later (in the bSteppedOver case below).
+		// In the case of hitting something above but not forward, we are not blocked from moving so we don't need the notification.
+		if (MoverComponent && SweepUpHit.bBlockingHit && StepFwdHit.bBlockingHit)
+		{
+			FMoverOnImpactParams ImpactParams(NAME_None, SweepUpHit, MoveDelta);
+			MoverComponent->HandleImpact(ImpactParams);
+		}
+
+		// pawn ran into a wall
+		if (MoverComponent)
+		{
+			FMoverOnImpactParams ImpactParams(NAME_None, StepFwdHit, MoveDelta);
+			MoverComponent->HandleImpact(ImpactParams);
+		}
+		
+		if (bIsFalling)
+		{
+			QueuedSubsteps.Add( FMovementSubstep(StepFwdSubstepName, UpdatedComponent->GetComponentLocation()-LastComponentLocation, true) );
+
+			// Commit queued substeps to movement record
+			for (FMovementSubstep Substep : QueuedSubsteps)
+			{
+				MoveRecord.Append(Substep);
+			}
+
+			return true;
+		}
+
+		// Cache forwards substep before the slide attempt
+		QueuedSubsteps.Add(FMovementSubstep(StepFwdSubstepName, UpdatedPrimitive->GetComponentLocation() - LastComponentLocation, true));
+		LastComponentLocation = UpdatedPrimitive->GetComponentLocation();
+
+		// adjust and try again
+		const float ForwardHitTime = StepFwdHit.Time;
+
+		// locking relevancy so velocity isn't added until it is needed to (adding it to the QueuedSubsteps so it can get added later)
+		MoveRecord.LockRelevancy(false);
+		const float ForwardSlideAmount = TryWalkToSlideAlongSurface(UpdatedComponent, UpdatedPrimitive, MoverComponent, MoveDelta, 1.f - StepFwdHit.Time, PawnRotation, StepFwdHit.Normal, StepFwdHit, true, MoveRecord, MaxWalkSlopeCosine, MaxStepHeight);
+		QueuedSubsteps.Add( FMovementSubstep(SlideSubstepName, UpdatedComponent->GetComponentLocation()-LastComponentLocation, true) );
+		LastComponentLocation = UpdatedPrimitive->GetComponentLocation();
+		MoveRecord.UnlockRelevancy();
+
+		if (bIsFalling)
+		{
+			UE_LOG(LogMover, VeryVerbose, TEXT("Reverting step-fwd attempt during step-up, because we could not adjust without falling"));
+			ScopedStepUpMovement.RevertMove();
+			return false;
+		}
+
+		// If both the forward hit and the deflection got us nowhere, there is no point in this step up.
+		if (ForwardHitTime == 0.f && ForwardSlideAmount == 0.f)
+		{
+			UE_LOG(LogMover, VeryVerbose, TEXT("Reverting step-fwd attempt during step-up, because no movement differences occurred"));
+			ScopedStepUpMovement.RevertMove();
+			return false;
+		}
+	}
+	else
+	{
+		// Our forward move attempt was unobstructed - cache it
+		QueuedSubsteps.Add(FMovementSubstep(StepFwdSubstepName, UpdatedPrimitive->GetComponentLocation() - LastComponentLocation, true));
+		LastComponentLocation = UpdatedPrimitive->GetComponentLocation();
+	}
+
+
+	// Step down
+	const FVector StepDownAdjustment = GravDir * StepTravelDownHeight;
+	const bool bDidStepDown = UMovementUtils::TryMoveUpdatedComponent_Internal(UpdatedComponent, StepDownAdjustment, UpdatedComponent->GetComponentQuat(), true, MOVECOMP_NoFlags, &StepFwdHit, ETeleportType::None);
+
+	UE_LOG(LogMover, VeryVerbose, TEXT("TryMoveToStepUp Down: %s (role %i) StepDownAdjustment=%s DidMove=%i"),
+		*GetNameSafe(UpdatedComponent->GetOwner()), UpdatedComponent->GetOwnerRole(), *StepDownAdjustment.ToCompactString(), bDidStepDown);
+
+
+	// If step down was initially penetrating abort the step up
+	if (StepFwdHit.bStartPenetrating)
+	{
+		UE_LOG(LogMover, VeryVerbose, TEXT("Reverting step-down attempt during step-up/step-fwd, because we started in a penetrating state"));
+		ScopedStepUpMovement.RevertMove();
+		return false;
+	}
+
+	FOptionalFloorCheckResult StepDownResult;
+	if (StepFwdHit.IsValidBlockingHit())
+	{
+		// See if this step sequence would have allowed us to travel higher than our max step height allows.
+		const float DeltaZ = StepFwdHit.ImpactPoint.Z - PawnFloorPointZ;
+		if (DeltaZ > MaxStepHeight)
+		{
+			UE_LOG(LogMover, VeryVerbose, TEXT("Reject step-down attempt during step-up/step-fwd, because it made us travel too high (too high Height %.3f) up from floor base %f to %f"), DeltaZ, PawnInitialFloorBaseZ, StepFwdHit.ImpactPoint.Z);
+			ScopedStepUpMovement.RevertMove();
+			return false;
+		}
+
+		// Reject unwalkable surface normals here.
+		if (!UFloorQueryUtils::IsHitSurfaceWalkable(StepFwdHit, MaxWalkSlopeCosine))
+		{
+			// Reject if normal opposes movement direction
+			const bool bNormalTowardsMe = (MoveDelta | StepFwdHit.ImpactNormal) < 0.f;
+			if (bNormalTowardsMe)
+			{
+				UE_LOG(LogMover, VeryVerbose, TEXT("Reject step-down attempt during step-up/step-fwd, due to unwalkable normal %s opposed to movement"), *StepFwdHit.ImpactNormal.ToString());
+				ScopedStepUpMovement.RevertMove();
+				return false;
+			}
+
+			// Also reject if we would end up being higher than our starting location by stepping down.
+			// It's fine to step down onto an unwalkable normal below us, we will just slide off. Rejecting those moves would prevent us from being able to walk off the edge.
+			if (StepFwdHit.Location.Z > OldLocation.Z)
+			{
+				UE_LOG(LogMover, VeryVerbose, TEXT("Reject step-down attempt during step-up/step-fwd, due to unwalkable normal %s above old position)"), *StepFwdHit.ImpactNormal.ToString());
+				ScopedStepUpMovement.RevertMove();
+				return false;
+			}
+		}
+
+		// Reject moves where the downward sweep hit something very close to the edge of the capsule. This maintains consistency with FindFloor as well.
+		if (!UFloorQueryUtils::IsWithinEdgeTolerance(StepFwdHit.Location, StepFwdHit.ImpactPoint, PawnRadius))
+		{
+			UE_LOG(LogMover, VeryVerbose, TEXT("Reject step-down attempt during step-up/step-fwd, due to being outside edge tolerance"));
+			ScopedStepUpMovement.RevertMove();
+			return false;
+		}
+
+		// Don't step up onto invalid surfaces if traveling higher.
+		if (DeltaZ > 0.f && !CanStepUpOnHitSurface(StepFwdHit))
+		{
+			UE_LOG(LogMover, VeryVerbose, TEXT("Reject step-down attempt during step-up/step-fwd, due to being up onto surface with !CanStepUpOnHitSurface")); 
+			ScopedStepUpMovement.RevertMove();
+			return false;
+		}
+
+		// See if we can validate the floor as a result of this step down. In almost all cases this should succeed, and we can avoid computing the floor outside this method.
+		if (OutFloorTestResult != NULL)
+		{
+
+			UFloorQueryUtils::FindFloor(UpdatedComponent, UpdatedPrimitive,
+				FloorSweepDistance, MaxWalkSlopeCosine,
+				UpdatedComponent->GetComponentLocation(), StepDownResult.FloorTestResult);
+
+			// Reject unwalkable normals if we end up higher than our initial height.
+			// It's fine to walk down onto an unwalkable surface, don't reject those moves.
+			if (StepFwdHit.Location.Z > OldLocation.Z)
+			{
+				// We should reject the floor result if we are trying to step up an actual step where we are not able to perch (this is rare).
+				// In those cases we should instead abort the step up and try to slide along the stair.
+				const float MAX_STEP_SIDE_Z = 0.08f; // TODO: Move magic numbers elsewhere
+				if (!StepDownResult.FloorTestResult.bBlockingHit && StepSideZ < MAX_STEP_SIDE_Z)
+				{
+					UE_LOG(LogMover, VeryVerbose, TEXT("Reject step-down attempt during step-up/step-fwd, due to it being an unperchable step")); 
+					ScopedStepUpMovement.RevertMove();
+					return false;
+				}
+			}
+
+			StepDownResult.bHasFloorResult = true;
+		}
+	}
+
+	// Cache downwards substep
+	QueuedSubsteps.Add(FMovementSubstep(StepDownSubstepName, UpdatedPrimitive->GetComponentLocation() - LastComponentLocation, false));
+	LastComponentLocation = UpdatedPrimitive->GetComponentLocation();
+
+	// Copy step down result.
+	if (OutFloorTestResult != NULL)
+	{
+		*OutFloorTestResult = StepDownResult;
+	}
+
+	// Don't recalculate velocity based on this height adjustment, if considering vertical adjustments.
+	//bJustTeleported |= !bMaintainHorizontalGroundVelocity;
+
+	// Commit queued substeps to movement record
+	for (FMovementSubstep Substep : QueuedSubsteps)
+	{
+		MoveRecord.Append(Substep);
+	}
+
+	return true;
+
+}
+
+bool UGroundMovementUtils::TryMoveToAdjustHeightAboveFloor(USceneComponent* UpdatedComponent, UPrimitiveComponent* UpdatedPrimitive, FFloorCheckResult& CurrentFloor, float MaxWalkSlopeCosine, FMovementRecord& MoveRecord)
+{
+	// If we have a floor check that hasn't hit anything, don't adjust height.
+	if (!CurrentFloor.IsWalkableFloor())
+	{
+		return false;
+	}
+
+	float OldFloorDist = CurrentFloor.FloorDist;
+	if (CurrentFloor.bLineTrace)
+	{
+		if (OldFloorDist < UE::FloorQueryUtility::MIN_FLOOR_DIST && CurrentFloor.LineDist >= UE::FloorQueryUtility::MIN_FLOOR_DIST)
+		{
+			// This would cause us to scale unwalkable walls
+			return false;
+		}
+		// Falling back to a line trace means the sweep was unwalkable (or in penetration). Use the line distance for the vertical adjustment.
+		OldFloorDist = CurrentFloor.LineDist;
+	}
+
+	// Move up or down to maintain floor height.
+	if (OldFloorDist < UE::FloorQueryUtility::MIN_FLOOR_DIST || OldFloorDist > UE::FloorQueryUtility::MAX_FLOOR_DIST)
+	{
+		FHitResult AdjustHit(1.f);
+		const float InitialZ = UpdatedComponent->GetComponentLocation().Z;
+		const float AvgFloorDist = (UE::FloorQueryUtility::MIN_FLOOR_DIST + UE::FloorQueryUtility::MAX_FLOOR_DIST) * 0.5f;
+		const float MoveDist = AvgFloorDist - OldFloorDist;
+
+		MoveRecord.LockRelevancy(false);
+
+		UMovementUtils::TrySafeMoveUpdatedComponent(UpdatedComponent, UpdatedPrimitive,
+			FVector(0.f, 0.f, MoveDist), UpdatedComponent->GetComponentQuat(), 
+			true, AdjustHit, ETeleportType::None, MoveRecord);
+
+		MoveRecord.UnlockRelevancy();
+
+		if (!AdjustHit.IsValidBlockingHit())
+		{
+			CurrentFloor.FloorDist += MoveDist;
+		}
+		else if (MoveDist > 0.f)
+		{
+			const float CurrentZ = UpdatedComponent->GetComponentLocation().Z;
+			CurrentFloor.FloorDist += CurrentZ - InitialZ;
+		}
+		else
+		{
+			checkSlow(MoveDist < 0.f);
+			const float CurrentZ = UpdatedComponent->GetComponentLocation().Z;
+			CurrentFloor.FloorDist = CurrentZ - AdjustHit.Location.Z;
+			if (UFloorQueryUtils::IsHitSurfaceWalkable(AdjustHit, MaxWalkSlopeCosine))
+			{
+				CurrentFloor.SetFromSweep(AdjustHit, CurrentFloor.FloorDist, true);
+			}
+		}
+
+		return true;
+	}
+
+	return false;
+}
+
+float UGroundMovementUtils::TryWalkToSlideAlongSurface(USceneComponent* UpdatedComponent, UPrimitiveComponent* UpdatedPrimitive, UMoverComponent* MovementComponent, const FVector& Delta, float PctOfDeltaToMove, const FQuat Rotation, const FVector& Normal, FHitResult& Hit, bool bHandleImpact, FMovementRecord& MoveRecord, float MaxWalkSlopeCosine, float MaxStepHeight)
+{
+	if (!Hit.bBlockingHit)
+	{
+		return 0.f;
+	}
+	
+	FVector SafeWalkNormal(Normal);
+
+	// We don't want to be pushed up an unwalkable surface.
+	if (SafeWalkNormal.Z > 0.f && !UFloorQueryUtils::IsHitSurfaceWalkable(Hit, MaxWalkSlopeCosine))
+	{
+		SafeWalkNormal = SafeWalkNormal.GetSafeNormal2D();
+	}
+
+	float PctOfTimeUsed = 0.f;
+	const FVector OldSafeHitNormal = SafeWalkNormal;
+
+	FVector SlideDelta = UMovementUtils::ComputeSlideDelta(Delta, PctOfDeltaToMove, SafeWalkNormal, Hit);
+
+	if (SlideDelta.Dot(Delta) > 0.f)
+	{
+		UMovementUtils::TrySafeMoveUpdatedComponent(UpdatedComponent, UpdatedPrimitive, SlideDelta, Rotation, true, Hit, ETeleportType::None, MoveRecord);
+
+		PctOfTimeUsed = Hit.Time;
+
+		if (Hit.IsValidBlockingHit())
+		{
+			// Notify first impact
+			if (MovementComponent && bHandleImpact)
+			{
+				FMoverOnImpactParams ImpactParams(NAME_None, Hit, SlideDelta);
+				MovementComponent->HandleImpact(ImpactParams);
+			}
+
+			// Compute new slide normal when hitting multiple surfaces.
+			SlideDelta = UMovementUtils::ComputeTwoWallAdjustedDelta(SlideDelta, Hit, OldSafeHitNormal);
+			if (SlideDelta.Z > 0.f && UFloorQueryUtils::IsHitSurfaceWalkable(Hit, MaxWalkSlopeCosine) && Hit.Normal.Z > UE_KINDA_SMALL_NUMBER)
+			{
+				// Maintain horizontal velocity
+				const float Time = (1.f - Hit.Time);
+				const FVector ScaledDelta = SlideDelta.GetSafeNormal() * SlideDelta.Size();
+				SlideDelta = FVector(SlideDelta.X, SlideDelta.Y, ScaledDelta.Z / Hit.Normal.Z) * Time;
+				// Should never exceed MaxStepHeight in vertical component, so rescale if necessary.
+				// This should be rare (Hit.Normal.Z above would have been very small) but we'd rather lose horizontal velocity than go too high.
+				if (SlideDelta.Z > MaxStepHeight)
+				{
+					const float Rescale = MaxStepHeight / SlideDelta.Z;
+					SlideDelta *= Rescale;
+				}
+			}
+			else
+			{
+				SlideDelta.Z = 0.f;
+			}
+			
+			// Only proceed if the new direction is of significant length and not in reverse of original attempted move.
+			if (!SlideDelta.IsNearlyZero(UE::MoverUtils::SMALL_MOVE_DISTANCE) && (SlideDelta | Delta) > 0.f)
+			{
+				// Perform second move
+				UMovementUtils::TrySafeMoveUpdatedComponent(UpdatedComponent, UpdatedPrimitive, SlideDelta, Rotation, true, Hit, ETeleportType::None, MoveRecord);
+				PctOfTimeUsed += (Hit.Time * (1.f - PctOfTimeUsed));
+
+				// Notify second impact
+				if (MovementComponent && bHandleImpact && Hit.bBlockingHit)
+				{
+					FMoverOnImpactParams ImpactParams(NAME_None, Hit, SlideDelta);
+					MovementComponent->HandleImpact(ImpactParams);
+				}
+			}
+		}
+
+		return FMath::Clamp(PctOfTimeUsed, 0.f, 1.f);
+	}
+
+	return 0.f;
+}
+
+FVector UGroundMovementUtils::ComputeDeflectedMoveOntoRamp(const FVector& OrigMoveDelta, const FHitResult& RampHitResult, float MaxWalkSlopeCosine, const bool bHitFromLineTrace)
+{
+	const FVector FloorNormal = RampHitResult.ImpactNormal;
+	const FVector ContactNormal = RampHitResult.Normal;
+
+	// JAH TODO: Change these Z tests to take movement plane into account, rather than assuming we're using the Z=0 plane
+	if (FloorNormal.Z < (1.f - UE_KINDA_SMALL_NUMBER) && FloorNormal.Z > UE_KINDA_SMALL_NUMBER && 
+		ContactNormal.Z > UE_KINDA_SMALL_NUMBER && 
+		UFloorQueryUtils::IsHitSurfaceWalkable(RampHitResult, MaxWalkSlopeCosine) && !bHitFromLineTrace)
+	{
+		// Compute a vector that moves parallel to the surface, by projecting the horizontal movement direction onto the ramp.
+		const FPlane RampSurfacePlane(FVector::ZeroVector, FloorNormal);
+		return UMovementUtils::ConstrainToPlane(OrigMoveDelta, RampSurfacePlane, true);
+	}
+
+	return OrigMoveDelta;
+}
+
+bool UGroundMovementUtils::CanStepUpOnHitSurface(const FHitResult& Hit)
+{
+	if (!Hit.IsValidBlockingHit())
+	{
+		return false;
+	}
+
+	// No component for "fake" hits when we are on a known good base.
+	const UPrimitiveComponent* HitComponent = Hit.Component.Get();
+	if (!HitComponent)
+	{
+		return true;
+	}
+
+	APawn* PawnOwner = Cast<APawn>(Hit.GetActor());
+
+	if (!HitComponent->CanCharacterStepUp(PawnOwner))
+	{
+		return false;
+	}
+
+	// No actor for "fake" hits when we are on a known good base.
+
+	if (!Hit.HitObjectHandle.IsValid())
+	{
+		return true;
+	}
+
+	const AActor* HitActor = Hit.HitObjectHandle.GetManagingActor();
+	if (!HitActor->CanBeBaseForCharacter(PawnOwner))
+	{
+		return false;
+	}
+
+	return true;
+}
diff --git a/Engine/Plugins/Experimental/Mover/Source/Mover/Private/MoveLibrary/ModularMovement.cpp b/Engine/Plugins/Experimental/Mover/Source/Mover/Private/MoveLibrary/ModularMovement.cpp
new file mode 100644
--- /dev/null
+++ b/Engine/Plugins/Experimental/Mover/Source/Mover/Private/MoveLibrary/ModularMovement.cpp
@@ -0,0 +1,72 @@
+// Copyright Epic Games, Inc. All Rights Reserved.
+
+#include "MoveLibrary/ModularMovement.h"
+#include "MoverTypes.h"
+#include "MoverSimulationTypes.h"
+
+
+FRotator ULinearTurnGenerator::GetTurn_Implementation(FRotator TargetOrientation, const FMoverTickStartData& FullStartState, const FMoverDefaultSyncState& MoverState, const FMoverTimeStep& TimeStep, const FProposedMove& ProposedMove, UMoverBlackboard* SimBlackboard)
+{
+	FRotator AngularVelocityDpS(FRotator::ZeroRotator);
+
+	const float DeltaSeconds = TimeStep.StepMs * 0.001f;
+
+	if (DeltaSeconds > 0.f)
+	{
+		FRotator AngularDelta = (TargetOrientation - MoverState.GetOrientation_WorldSpace());
+		FRotator Winding, Remainder;
+
+		AngularDelta.GetWindingAndRemainder(Winding, Remainder);	// to find the fastest turn, just get the (-180,180] remainder
+
+		AngularVelocityDpS = Remainder * (1.f / DeltaSeconds);
+
+		FRotator TurnRates(PitchRate, HeadingRate, RollRate);
+
+		if (HeadingRate >= 0.f)
+		{
+			AngularVelocityDpS.Yaw = FMath::Clamp(AngularVelocityDpS.Yaw, -HeadingRate, HeadingRate);
+		}
+
+		if (PitchRate >= 0.f)
+		{
+			AngularVelocityDpS.Pitch = FMath::Clamp(AngularVelocityDpS.Pitch, -PitchRate, PitchRate);
+		}
+
+		if (RollRate >= 0.f)
+		{
+			AngularVelocityDpS.Roll = FMath::Clamp(AngularVelocityDpS.Roll, -RollRate, RollRate);
+		}
+	}
+
+	return AngularVelocityDpS;
+}
+
+
+// Note the lack of argument range checking.  Value and Time arguments can be in any units, as long as they're consistent.
+static float CalcExactDampedInterpolation(float CurrentVal, float TargetVal, float HalflifeTime, float DeltaTime)
+{
+	return FMath::Lerp(CurrentVal, TargetVal, 1.f - FMath::Pow(2.f, (-DeltaTime / HalflifeTime)));
+}
+
+FRotator UExactDampedTurnGenerator::GetTurn_Implementation(FRotator TargetOrientation, const FMoverTickStartData& FullStartState, const FMoverDefaultSyncState& MoverState, const FMoverTimeStep& TimeStep, const FProposedMove& ProposedMove, UMoverBlackboard* SimBlackboard)
+{
+	FRotator AngularVelocityDpS(FRotator::ZeroRotator);
+
+	const float DeltaSeconds = TimeStep.StepMs * 0.001f;
+
+	if (DeltaSeconds > 0.f && HalfLifeSeconds > UE_KINDA_SMALL_NUMBER)
+	{
+		FRotator AngularDelta = (TargetOrientation - MoverState.GetOrientation_WorldSpace());
+		FRotator Winding, Remainder;
+
+		AngularDelta.GetWindingAndRemainder(Winding, Remainder);	// to find the fastest turn, just get the (-180,180] remainder
+
+		const float OneOverDeltaSecs = 1.f / DeltaSeconds;
+
+		AngularVelocityDpS.Yaw   = CalcExactDampedInterpolation(0.f, Remainder.Yaw,   HalfLifeSeconds, DeltaSeconds) * OneOverDeltaSecs;
+		AngularVelocityDpS.Pitch = CalcExactDampedInterpolation(0.f, Remainder.Pitch, HalfLifeSeconds, DeltaSeconds) * OneOverDeltaSecs;
+		AngularVelocityDpS.Roll  = CalcExactDampedInterpolation(0.f, Remainder.Roll,  HalfLifeSeconds, DeltaSeconds) * OneOverDeltaSecs;
+	}
+
+	return AngularVelocityDpS;
+}
\ No newline at end of file
diff --git a/Engine/Plugins/Experimental/Mover/Source/Mover/Private/MovementMode.cpp b/Engine/Plugins/Experimental/Mover/Source/Mover/Private/MovementMode.cpp
new file mode 100644
--- /dev/null
+++ b/Engine/Plugins/Experimental/Mover/Source/Mover/Private/MovementMode.cpp
@@ -0,0 +1,121 @@
+// Copyright Epic Games, Inc. All Rights Reserved.
+
+#include "MovementMode.h"
+#include "MoverComponent.h"
+#include "Engine/BlueprintGeneratedClass.h"
+
+#include UE_INLINE_GENERATED_CPP_BY_NAME(MovementMode)
+
+UBaseMovementMode::UBaseMovementMode(const FObjectInitializer& ObjectInitializer)
+	: Super(ObjectInitializer)
+{
+
+}
+
+void UBaseMovementMode::DoRegister(const FName ModeName)
+{
+	auto IsImplementedInBlueprint = [](const UFunction* Func) -> bool
+	{
+		return Func && ensure(Func->GetOuter())
+			&& Func->GetOuter()->IsA(UBlueprintGeneratedClass::StaticClass());
+	};
+
+	static FName SimTickFuncName = FName(TEXT("K2_OnSimulationTick"));
+	UFunction* SimulationTickFunction = GetClass()->FindFunctionByName(SimTickFuncName);
+	bHasBlueprintSimulationTick = IsImplementedInBlueprint(SimulationTickFunction);
+
+	static FName GenMoveFuncName = FName(TEXT("K2_OnGenerateMove"));
+	UFunction* GenMoveFunction = GetClass()->FindFunctionByName(GenMoveFuncName);
+	bHasBlueprintGenerateMove = IsImplementedInBlueprint(GenMoveFunction);
+
+	OnRegistered(ModeName);
+}
+
+void UBaseMovementMode::DoUnregister()
+{
+	OnUnregistered();
+}
+
+
+void UBaseMovementMode::DoGenerateMove(const FMoverTickStartData& StartState, const FMoverTimeStep& TimeStep, FProposedMove& OutProposedMove) const
+{
+	if (bHasBlueprintGenerateMove)
+	{
+		OutProposedMove = K2_OnGenerateMove(StartState, TimeStep);
+	}
+	else
+	{
+		OnGenerateMove(StartState, TimeStep, OutProposedMove);
+	}
+}
+
+
+void UBaseMovementMode::DoSimulationTick(const FSimulationTickParams& Params, FMoverTickEndData& OutputState)
+{
+	if (bHasBlueprintSimulationTick)
+	{
+		OutputState = K2_OnSimulationTick(Params);
+	}
+	else
+	{
+		OnSimulationTick(Params, OutputState);
+	}
+}
+
+
+UMoverComponent* UBaseMovementMode::GetMoverComponent() const
+{
+	return CastChecked<UMoverComponent>(GetOuter());
+}
+
+
+const UMoverBlackboard* UBaseMovementMode::GetBlackboard() const
+{
+	return GetBlackboard_Mutable();
+}
+
+
+UMoverBlackboard* UBaseMovementMode::GetBlackboard_Mutable() const
+{
+	if (UMoverComponent* OwnerComponent = GetMoverComponent())
+	{
+		return OwnerComponent->SimBlackboard.Get();
+	}
+
+	return nullptr;
+}
+
+
+#if WITH_EDITOR
+EDataValidationResult UBaseMovementMode::IsDataValid(FDataValidationContext& Context) const
+{
+	EDataValidationResult Result = EDataValidationResult::Valid;
+	for (UBaseMovementModeTransition* Transition : Transitions)
+	{
+		if (Transition->IsDataValid(Context) == EDataValidationResult::Invalid)
+		{
+			Result = EDataValidationResult::Invalid;
+		}
+	}
+
+	return Result;
+}
+#endif // WITH_EDITOR
+
+
+void UBaseMovementMode::OnRegistered(const FName ModeName)
+{
+}
+
+void UBaseMovementMode::OnUnregistered()
+{
+}
+
+void UBaseMovementMode::OnGenerateMove(const FMoverTickStartData& StartState, const FMoverTimeStep& TimeStep, FProposedMove& OutProposedMove) const
+{
+}
+
+void UBaseMovementMode::OnSimulationTick(const FSimulationTickParams& Params, FMoverTickEndData& OutputState)
+{
+}
+
diff --git a/Engine/Plugins/Experimental/Mover/Source/Mover/Private/MovementModeStateMachine.cpp b/Engine/Plugins/Experimental/Mover/Source/Mover/Private/MovementModeStateMachine.cpp
new file mode 100644
--- /dev/null
+++ b/Engine/Plugins/Experimental/Mover/Source/Mover/Private/MovementModeStateMachine.cpp
@@ -0,0 +1,429 @@
+// Copyright Epic Games, Inc. All Rights Reserved.
+
+
+#include "MovementModeStateMachine.h"
+#include "MoverDeveloperSettings.h"
+#include "MoverLog.h"
+#include "MoveLibrary/MovementUtils.h"
+#include "MoverComponent.h"
+#include "Templates/SubclassOf.h"
+
+
+#include UE_INLINE_GENERATED_CPP_BY_NAME(MovementModeStateMachine)
+
+const FName UNullMovementMode::NullModeName(TEXT("Null"));
+
+UNullMovementMode::UNullMovementMode(const FObjectInitializer& ObjectInitializer)
+	: Super(ObjectInitializer)
+{
+}
+
+void UNullMovementMode::OnSimulationTick(const FSimulationTickParams& Params, FMoverTickEndData& OutputState)
+{
+}
+
+
+
+UMovementModeStateMachine::UMovementModeStateMachine(const FObjectInitializer& ObjectInitializer)
+	: Super(ObjectInitializer)
+{
+}
+
+
+void UMovementModeStateMachine::RegisterMovementMode(FName ModeName, TObjectPtr<UBaseMovementMode> Mode, bool bIsDefaultMode)
+{
+	// JAH TODO: add validation and warnings for overwriting modes
+	// JAH TODO: add validation of Mode
+
+	Modes.Add(ModeName, Mode);
+
+	if (bIsDefaultMode)
+	{
+		//JAH TODO: add validation that we are only overriding the default null mode
+		DefaultModeName = ModeName;
+	}
+
+	Mode->DoRegister(ModeName);
+}
+
+
+void UMovementModeStateMachine::RegisterMovementMode(FName ModeName, TSubclassOf<UBaseMovementMode> ModeType, bool bIsDefaultMode)
+{
+	
+	RegisterMovementMode(ModeName, NewObject<UBaseMovementMode>(GetOuter(), ModeType), bIsDefaultMode);
+
+}
+
+
+
+void UMovementModeStateMachine::UnregisterMovementMode(FName ModeName)
+{
+	TObjectPtr<UBaseMovementMode> ModeToUnregister = Modes.FindAndRemoveChecked(ModeName);
+
+	if (ModeToUnregister)
+	{
+		ModeToUnregister->DoUnregister();
+	}
+}
+
+void UMovementModeStateMachine::ClearAllMovementModes()
+{
+	for (TPair<FName, TObjectPtr<UBaseMovementMode>>& Element : Modes)
+	{
+		Element.Value->DoUnregister();
+	}
+
+	Modes.Empty();
+
+	ConstructDefaultModes();	// Note that we're resetting to our defaults so we keep the null movement mode
+}
+
+void UMovementModeStateMachine::SetDefaultMode(FName NewDefaultModeName)
+{
+	Modes.FindChecked(NewDefaultModeName);
+
+	DefaultModeName = NewDefaultModeName;
+}
+
+
+void UMovementModeStateMachine::QueueNextMode(FName DesiredNextModeName, bool bShouldReenter)
+{
+	if (DesiredNextModeName != NAME_None)
+	{ 
+		const FName NextModeName          = QueuedModeTransition->GetNextModeName();
+		const bool bShouldNextModeReenter = QueuedModeTransition->ShoulReenter();
+
+		if ((NextModeName != NAME_None) &&
+		    (NextModeName != DesiredNextModeName || bShouldReenter != bShouldNextModeReenter))
+		{
+			const AActor* OwnerActor = GetOwnerActor();
+			UE_LOG(LogMover, Warning, TEXT("%s (%s) Overwriting of queued mode change (%s, reenter: %i) with (%s, reenter: %i)"), *GetNameSafe(OwnerActor), *UEnum::GetValueAsString(OwnerActor->GetLocalRole()), *NextModeName.ToString(), bShouldNextModeReenter, *DesiredNextModeName.ToString(), bShouldReenter);
+		}
+
+		if (Modes.Contains(DesiredNextModeName))
+		{
+			QueuedModeTransition->SetNextMode(DesiredNextModeName, bShouldReenter);
+		}
+		else
+		{
+			UE_LOG(LogMover, Warning, TEXT("Attempted to queue an unregistered movement mode: %s"), *DesiredNextModeName.ToString());
+		}
+	}
+}
+
+
+void UMovementModeStateMachine::SetModeImmediately(FName DesiredModeName, bool bShouldReenter)
+{
+	QueueNextMode(DesiredModeName, bShouldReenter);
+	AdvanceToNextMode();
+}
+
+
+void UMovementModeStateMachine::ClearQueuedMode()
+{
+	QueuedModeTransition->Clear();
+}
+
+
+void UMovementModeStateMachine::OnSimulationTick(USceneComponent* UpdatedComponent, UPrimitiveComponent* UpdatedPrimitive, UMoverBlackboard* SimBlackboard, const FMoverTickStartData& StartState, const FMoverTimeStep& TimeStep, FMoverTickEndData& OutputState)
+{
+	const float EndingSimTimeMs = TimeStep.BaseSimTimeMs + TimeStep.StepMs;
+
+	//GEngine->AddOnScreenDebugMessage(-1, -0.1f, FColor::White, FString::Printf(TEXT("Mode Tick: %s  Queued: %s"), *CurrentModeName.ToString(), *NextModeName.ToString()));
+
+	FMoverTimeStep SubTimeStep = TimeStep;
+	FMoverTickStartData SubstepStartData = StartState;
+	FMoverDefaultSyncState* SubstepStartSyncState = SubstepStartData.SyncState.SyncStateCollection.FindMutableDataByType<FMoverDefaultSyncState>();
+
+	UMoverComponent* MoverComp = CastChecked<UMoverComponent>(GetOuter());
+
+	if (!QueuedModeTransition->IsSet())
+	{
+		QueueNextMode(SubstepStartSyncState->MovementMode);
+	}
+
+	AdvanceToNextMode();
+
+	const int32 MaxConsecutiveFullRefundedSubsteps = GetDefault<UMoverDeveloperSettings>()->MaxTimesToRefundSubstep;
+	int32 NumConsecutiveFullRefundedSubsteps = 0;
+	while (SubTimeStep.BaseSimTimeMs < EndingSimTimeMs)
+	{
+		SubstepStartSyncState = SubstepStartData.SyncState.SyncStateCollection.FindMutableDataByType<FMoverDefaultSyncState>();
+		FMoverDefaultSyncState* OutputSyncState = &OutputState.SyncState.SyncStateCollection.FindOrAddMutableDataByType<FMoverDefaultSyncState>();
+		OutputSyncState->MovementMode = CurrentModeName;
+
+		OutputState.MovementEndState.ResetToDefaults();
+
+		SubTimeStep.StepMs = EndingSimTimeMs - SubTimeStep.BaseSimTimeMs;		// TODO: convert this to an overridable function that can support MaxStepTime, MaxIterations, etc.
+
+		// Transfer any queued moves into the starting state. They'll be started during the move generation.
+		FlushQueuedMovesToGroup(SubstepStartSyncState->LayeredMoves);
+		OutputSyncState->LayeredMoves = SubstepStartSyncState->LayeredMoves;
+		FLayeredMoveGroup& CurrentLayeredMoves = OutputSyncState->LayeredMoves;
+
+		// Gather any layered move contributions
+		FProposedMove CombinedLayeredMove;
+		const bool bHasLayeredMoveContributions = CurrentLayeredMoves.DoGenerateMove(SubstepStartData, SubTimeStep, MoverComp, SimBlackboard, OUT CombinedLayeredMove);
+
+		if (bHasLayeredMoveContributions && !CombinedLayeredMove.PreferredMode.IsNone())
+		{
+			SetModeImmediately(CombinedLayeredMove.PreferredMode);
+			OutputSyncState->MovementMode = CurrentModeName;
+		}
+
+		// Merge proposed movement from the current mode with movement from layered moves
+		if (!CurrentModeName.IsNone() && Modes.Contains(CurrentModeName))
+		{
+			UBaseMovementMode* CurrentMode = Modes[CurrentModeName];
+			FProposedMove CombinedMove;
+			CurrentMode->DoGenerateMove(SubstepStartData, SubTimeStep, OUT CombinedMove);
+
+			if (bHasLayeredMoveContributions)
+			{
+				if (CombinedLayeredMove.bHasDirIntent && CombinedMove.MixMode != EMoveMixMode::OverrideAll)
+				{
+					CombinedMove.bHasDirIntent = CombinedLayeredMove.bHasDirIntent;
+					CombinedMove.DirectionIntent = CombinedLayeredMove.DirectionIntent;
+				}
+				
+				// Combine movement parameters from layered moves into what the mode wants to do
+				if (CombinedLayeredMove.MixMode == EMoveMixMode::OverrideAll)
+				{
+					CombinedMove = CombinedLayeredMove;
+				}
+				else if (CombinedLayeredMove.MixMode == EMoveMixMode::AdditiveVelocity)
+				{
+					CombinedMove.LinearVelocity += CombinedLayeredMove.LinearVelocity;
+					CombinedMove.AngularVelocity += CombinedLayeredMove.AngularVelocity;
+				}
+				else if (CombinedLayeredMove.MixMode == EMoveMixMode::OverrideVelocity)
+				{
+					CombinedMove.LinearVelocity = CombinedLayeredMove.LinearVelocity;
+					CombinedMove.AngularVelocity = CombinedLayeredMove.AngularVelocity;
+				}
+			}
+
+			// Apply any layered move finish velocity settings
+			if (CurrentLayeredMoves.bApplyResidualVelocity)
+			{
+				CombinedMove.LinearVelocity = CurrentLayeredMoves.ResidualVelocity;
+			}
+			if (CurrentLayeredMoves.ResidualClamping >= 0.0f)
+			{
+				CombinedMove.LinearVelocity = CombinedMove.LinearVelocity.GetClampedToMaxSize2D(CurrentLayeredMoves.ResidualClamping);
+				CombinedMove.LinearVelocity.Z = FMath::Min<FVector::FReal>(CombinedMove.LinearVelocity.Z, CurrentLayeredMoves.ResidualClamping);
+			}
+			CurrentLayeredMoves.ResetResidualVelocity();
+
+			// Execute the combined proposed move
+			{
+				FSimulationTickParams SimTickParams;
+				SimTickParams.StartState = SubstepStartData;
+				SimTickParams.UpdatedComponent = UpdatedComponent;
+				SimTickParams.UpdatedPrimitive = UpdatedPrimitive;
+				SimTickParams.MoverComponent = MoverComp;
+				SimTickParams.TimeStep = SubTimeStep;
+				SimTickParams.ProposedMove = CombinedMove;
+
+				// Check any transitions that have been registered with the current movement mode
+				bool bTransitionTriggered = false;
+				for (UBaseMovementModeTransition* Transition : CurrentMode->Transitions)
+				{
+					FTransitionEvalResult EvalResult = Transition->DoEvaluate(SimTickParams);
+
+					if (!EvalResult.NextMode.IsNone())
+					{
+						OutputState.MovementEndState.NextModeName = EvalResult.NextMode;
+						OutputState.MovementEndState.RemainingMs = SimTickParams.TimeStep.StepMs; 	// Pass all remaining time to next mode
+						Transition->DoTrigger(SimTickParams);
+						bTransitionTriggered = true;
+						break;
+					}
+				}
+
+				if (!bTransitionTriggered)
+				{
+					CurrentMode->DoSimulationTick(SimTickParams, OutputState);
+				}
+			}
+
+			QueueNextMode(OutputState.MovementEndState.NextModeName);
+
+			// Check if all of the time for this Substep was refunded
+			if (FMath::IsNearlyEqual(SubTimeStep.StepMs, OutputState.MovementEndState.RemainingMs, UE_KINDA_SMALL_NUMBER))
+			{
+				NumConsecutiveFullRefundedSubsteps++;
+				// if we've done this sub step a lot before go ahead and just advance time to avoid freezing editor
+				if (NumConsecutiveFullRefundedSubsteps >= MaxConsecutiveFullRefundedSubsteps)
+				{
+					UE_LOG(LogMover, Warning, TEXT("Movement mode %s and %s on %s are stuck giving time back to eachother - Advancing substep."), *CurrentModeName.ToString(), *OutputState.MovementEndState.NextModeName.ToString(), *MoverComp->GetOwner()->GetName());
+					SubTimeStep.BaseSimTimeMs += SubTimeStep.StepMs;
+				}
+			}
+			else
+			{
+				NumConsecutiveFullRefundedSubsteps = 0;
+			}
+
+			//GEngine->AddOnScreenDebugMessage(-1, -0.1f, FColor::White, FString::Printf(TEXT("NextModeName: %s  Queued: %s"), *Output.MovementEndState.NextModeName.ToString(), *NextModeName.ToString()));
+		}
+
+		// Switch modes if necessary (note that this will allow exit/enter on the same state)
+		AdvanceToNextMode();
+		OutputSyncState->MovementMode = CurrentModeName;
+
+		const float RemainingMs = FMath::Clamp(OutputState.MovementEndState.RemainingMs, 0.0f, SubTimeStep.StepMs);
+		SubTimeStep.BaseSimTimeMs += (SubTimeStep.StepMs - RemainingMs);
+		SubTimeStep.StepMs = EndingSimTimeMs - SubTimeStep.BaseSimTimeMs;
+
+		SubstepStartData.SyncState = OutputState.SyncState;
+		SubstepStartData.AuxState  = OutputState.AuxState;
+	}
+
+}
+
+
+void UMovementModeStateMachine::OnSimulationRollback(const FMoverSyncState* SyncState, const FMoverAuxStateContext* AuxState)
+{
+	ClearQueuedMode();
+	QueuedLayeredMoves.Empty();
+}
+
+
+const UBaseMovementMode* UMovementModeStateMachine::GetCurrentMode() const
+{
+	if (CurrentModeName != NAME_None && Modes.Contains(CurrentModeName))
+	{
+		return Modes[CurrentModeName];
+	}
+
+	return nullptr;
+}
+
+const UBaseMovementMode* UMovementModeStateMachine::FindMovementMode(FName ModeName) const
+{
+	if (ModeName != NAME_None && Modes.Contains(ModeName))
+	{
+		return Modes[ModeName];
+	}
+
+	return nullptr;
+}
+
+void UMovementModeStateMachine::QueueLayeredMove(TSharedPtr<FLayeredMoveBase> Move)
+{
+	QueuedLayeredMoves.Add(Move);
+}
+
+
+void UMovementModeStateMachine::ConstructDefaultModes()
+{
+	RegisterMovementMode(UNullMovementMode::NullModeName, UNullMovementMode::StaticClass(), true);
+
+	DefaultModeName = NAME_None;
+	CurrentModeName = UNullMovementMode::NullModeName;
+
+	ClearQueuedMode();
+}
+
+void UMovementModeStateMachine::AdvanceToNextMode()
+{
+	const FName NextModeName = QueuedModeTransition->GetNextModeName();
+	const bool bShouldNextModeReenter = QueuedModeTransition->ShoulReenter();
+
+	if ((NextModeName != NAME_None && Modes.Contains(NextModeName)) &&
+		(CurrentModeName != NextModeName || bShouldNextModeReenter))
+	{
+		const AActor* OwnerActor = GetOwnerActor();
+		UE_LOG(LogMover, Verbose, TEXT("AdvanceToNextMode: %s (%s) from %s to %s"), 
+			*GetNameSafe(OwnerActor), *UEnum::GetValueAsString(OwnerActor->GetLocalRole()), *CurrentModeName.ToString(), *NextModeName.ToString());
+
+		const FName PreviousModeName = CurrentModeName;
+		CurrentModeName = NextModeName;
+
+		// signal movement mode change event
+		const UMoverComponent* MoverComp = CastChecked<UMoverComponent>(GetOuter());
+		MoverComp->OnMovementModeChanged.Broadcast(PreviousModeName, NextModeName);
+	}
+
+	ClearQueuedMode();
+}
+
+void UMovementModeStateMachine::FlushQueuedMovesToGroup(FLayeredMoveGroup& Group)
+{
+	if (!QueuedLayeredMoves.IsEmpty())
+	{
+		for (TSharedPtr<FLayeredMoveBase>& QueuedMove : QueuedLayeredMoves)
+		{
+			Group.QueueLayeredMove(QueuedMove);
+		}
+		
+		QueuedLayeredMoves.Empty();
+	}
+}
+
+AActor* UMovementModeStateMachine::GetOwnerActor() const
+{
+	if (const UActorComponent* OwnerMoverComp = Cast<UActorComponent>(GetOuter()))
+	{
+		return OwnerMoverComp->GetOwner();
+	}
+
+	return nullptr;
+}
+
+void UMovementModeStateMachine::PostInitProperties()
+{
+	Super::PostInitProperties();
+
+	if (!HasAnyFlags(RF_ClassDefaultObject))
+	{
+		QueuedModeTransition = NewObject<UImmediateMovementModeTransition>(this, TEXT("QueuedModeTransition"), RF_Transient);
+	}
+}
+
+
+
+
+UImmediateMovementModeTransition::UImmediateMovementModeTransition(const FObjectInitializer& ObjectInitializer)
+	: Super(ObjectInitializer)
+{
+	Clear();
+}
+
+FTransitionEvalResult UImmediateMovementModeTransition::OnEvaluate(const FSimulationTickParams& Params) const
+{
+	if (NextMode != NAME_None)
+	{
+		if (bShouldNextModeReenter)
+		{
+			return FTransitionEvalResult(NextMode);
+		}
+		else if (const FMoverDefaultSyncState* SyncState = Params.StartState.SyncState.SyncStateCollection.FindDataByType<FMoverDefaultSyncState>())
+		{
+			if (NextMode != SyncState->MovementMode)
+			{
+				return FTransitionEvalResult(NextMode);
+			}
+		}
+	}
+
+	return FTransitionEvalResult::NoTransition;
+}
+
+void UImmediateMovementModeTransition::OnTrigger(const FSimulationTickParams& Params)
+{
+	Clear();
+}
+
+void UImmediateMovementModeTransition::SetNextMode(FName DesiredModeName, bool bShouldReenter)
+{
+	NextMode = DesiredModeName;
+	bShouldNextModeReenter = bShouldReenter;
+}
+
+void UImmediateMovementModeTransition::Clear()
+{
+	NextMode = NAME_None;
+	bShouldNextModeReenter = false;
+}
diff --git a/Engine/Plugins/Experimental/Mover/Source/Mover/Private/MovementModeTransition.cpp b/Engine/Plugins/Experimental/Mover/Source/Mover/Private/MovementModeTransition.cpp
new file mode 100644
--- /dev/null
+++ b/Engine/Plugins/Experimental/Mover/Source/Mover/Private/MovementModeTransition.cpp
@@ -0,0 +1,80 @@
+// Copyright Epic Games, Inc. All Rights Reserved.
+
+#include "MovementModeTransition.h"
+#include "Engine/BlueprintGeneratedClass.h"
+#include "MoverSimulationTypes.h"
+#include "MoverTypes.h"
+
+
+#include UE_INLINE_GENERATED_CPP_BY_NAME(MovementModeTransition)
+
+const FTransitionEvalResult FTransitionEvalResult::NoTransition = FTransitionEvalResult();
+
+UBaseMovementModeTransition::UBaseMovementModeTransition(const FObjectInitializer& ObjectInitializer)
+	: Super(ObjectInitializer)
+{
+}
+
+void UBaseMovementModeTransition::DoRegister()
+{
+	auto IsImplementedInBlueprint = [](const UFunction* Func) -> bool
+	{
+		return Func && ensure(Func->GetOuter())
+			&& Func->GetOuter()->IsA(UBlueprintGeneratedClass::StaticClass());
+	};
+
+	static FName EvaluateFuncName = FName(TEXT("K2_OnEvaluate"));
+	UFunction* EvaluateFunction = GetClass()->FindFunctionByName(EvaluateFuncName);
+	bHasBlueprintEvaluate = IsImplementedInBlueprint(EvaluateFunction);
+
+	static FName TriggerFuncName = FName(TEXT("K2_OnTrigger"));
+	UFunction* TriggerFunction = GetClass()->FindFunctionByName(TriggerFuncName);
+	bHasBlueprintTrigger = IsImplementedInBlueprint(TriggerFunction);
+
+}
+
+void UBaseMovementModeTransition::DoUnregister()
+{
+}
+
+FTransitionEvalResult UBaseMovementModeTransition::DoEvaluate(const FSimulationTickParams& Params) const
+{
+	if (bHasBlueprintEvaluate)
+	{
+		return K2_OnEvaluate(Params);
+	}
+	else
+	{
+		return OnEvaluate(Params);
+	}
+}
+
+FTransitionEvalResult UBaseMovementModeTransition::OnEvaluate(const FSimulationTickParams& Params) const
+{
+	return FTransitionEvalResult::NoTransition;
+}
+
+
+void UBaseMovementModeTransition::DoTrigger(const FSimulationTickParams& Params)
+{
+	if (bHasBlueprintTrigger)
+	{
+		return K2_OnTrigger(Params);
+	}
+	else
+	{
+		return OnTrigger(Params);
+	}
+}
+
+void UBaseMovementModeTransition::OnTrigger(const FSimulationTickParams& Params)
+{
+}
+
+
+#if WITH_EDITOR
+EDataValidationResult UBaseMovementModeTransition::IsDataValid(FDataValidationContext& Context) const
+{
+	return EDataValidationResult::Valid;
+}
+#endif // WITH_EDITOR
diff --git a/Engine/Plugins/Experimental/Mover/Source/Mover/Private/MoverBasePawn.cpp b/Engine/Plugins/Experimental/Mover/Source/Mover/Private/MoverBasePawn.cpp
new file mode 100644
--- /dev/null
+++ b/Engine/Plugins/Experimental/Mover/Source/Mover/Private/MoverBasePawn.cpp
@@ -0,0 +1,31 @@
+// Copyright Epic Games, Inc. All Rights Reserved.
+
+
+#include "MoverBasePawn.h"
+#include "MoverComponent.h"
+
+
+static const FName Name_CharacterMotionComponent(TEXT("MoverComponent"));
+
+
+AMoverBasePawn::AMoverBasePawn(const FObjectInitializer& ObjectInitializer)
+	: Super(ObjectInitializer)
+{
+
+	CharacterMotionComponent = CreateDefaultSubobject<UMoverComponent>(Name_CharacterMotionComponent);
+	ensure(CharacterMotionComponent);
+
+ 	// Set this pawn to call Tick() every frame.  You can turn this off to improve performance if you don't need it.
+	PrimaryActorTick.bCanEverTick = true;
+
+	SetReplicatingMovement(false);	// disable Actor-level movement replication, since our Mover component will handle it
+}
+
+
+
+UPrimitiveComponent* AMoverBasePawn::GetMovementBase() const
+{
+	return CharacterMotionComponent ? CharacterMotionComponent->GetMovementBase() : nullptr;
+}
+
+
diff --git a/Engine/Plugins/Experimental/Mover/Source/Mover/Private/MoverComponent.cpp b/Engine/Plugins/Experimental/Mover/Source/Mover/Private/MoverComponent.cpp
new file mode 100644
--- /dev/null
+++ b/Engine/Plugins/Experimental/Mover/Source/Mover/Private/MoverComponent.cpp
@@ -0,0 +1,1196 @@
+// Copyright Epic Games, Inc. All Rights Reserved.
+
+
+#include "MoverComponent.h"
+#include "MoverSimulationTypes.h"
+#include "MovementModeStateMachine.h"
+#include "Kinematic/Modes/WalkingMode.h"
+#include "Kinematic/Modes/FallingMode.h"
+#include "Kinematic/Modes/FlyingMode.h"
+#include "MoveLibrary/MovementUtils.h"
+#include "MoveLibrary/FloorQueryUtils.h"
+#include "MoverLog.h"
+#include "Backends/MoverNetworkPredictionLiaison.h"
+#include "Components/MeshComponent.h"
+#include "Components/PrimitiveComponent.h"
+#include "Engine/ScopedMovementUpdate.h"
+#include "Engine/World.h"
+#include "GameFramework/PhysicsVolume.h"
+#include "Misc/AssertionMacros.h"
+#include "Misc/TransactionObjectEvent.h"
+
+#if WITH_EDITOR
+#include "Misc/DataValidation.h"
+#endif
+
+#include UE_INLINE_GENERATED_CPP_BY_NAME(MoverComponent)
+
+#define LOCTEXT_NAMESPACE "Mover"
+
+
+namespace MoverComponentConstants
+{
+	const FVector DefaultGravityAccel	= FVector(0.0, 0.0, -980.0);
+	const FVector DefaultUpDir			= FVector(0.0, 0.0, 1.0);
+}
+
+
+static constexpr float ROTATOR_TOLERANCE = (1e-3);
+
+
+UMoverComponent::UMoverComponent()
+{
+	PrimaryComponentTick.TickGroup = TG_PrePhysics;
+	PrimaryComponentTick.bCanEverTick = true;
+
+	bWantsInitializeComponent = true;
+	bAutoActivate = true;
+
+	// Default movement modes
+	MovementModes.Add(KinematicModeNames::Walking, CreateDefaultSubobject<UWalkingMode>(TEXT("DefaultWalkingMode")));
+	MovementModes.Add(KinematicModeNames::Falling, CreateDefaultSubobject<UFallingMode>(TEXT("DefaultFallingMode")));
+	MovementModes.Add(KinematicModeNames::Flying,  CreateDefaultSubobject<UFlyingMode>(TEXT("DefaultFlyingMode")));
+
+	StartingMovementMode = KinematicModeNames::Falling;
+
+	PersistentSyncStateDataTypes.Add(FMoverDataPersistence(FMoverDefaultSyncState::StaticStruct(), true));
+
+	BackendClass = UMoverNetworkPredictionLiaisonComponent::StaticClass();
+}
+
+
+void UMoverComponent::InitializeComponent()
+{
+	TGuardValue<bool> InInitializeComponentGuard(bInInitializeComponent, true);
+
+	// RootComponent is null in OnRegister for blueprint (non-native) root components.
+	if (!UpdatedComponent)
+	{
+		// Auto-register owner's root component if found.
+		if (AActor* MyActor = GetOwner())
+		{
+			if (USceneComponent* NewUpdatedComponent = MyActor->GetRootComponent())
+			{
+				SetUpdatedComponent(NewUpdatedComponent);
+			}
+			else
+			{
+				ensureMsgf(false, TEXT("No root component found on %s. Simulation initialization will most likely fail."), *GetPathNameSafe(MyActor));
+			}
+		}
+	}
+
+	// Instantiate out sister backend component that will actually talk to the system driving the simulation
+	if (BackendClass)
+	{
+		UActorComponent* NewLiaisonComp = NewObject<UActorComponent>(GetOwner(), BackendClass, TEXT("BackendLiaisonComponent"));
+		BackendLiaisonComp = CastChecked<IMoverBackendLiaisonInterface>(NewLiaisonComp);
+		if (BackendLiaisonComp.Get())
+		{
+			NewLiaisonComp->RegisterComponent();
+			NewLiaisonComp->InitializeComponent();
+			NewLiaisonComp->SetNetAddressable();
+		}
+	}
+	else
+	{
+		UE_LOG(LogMover, Error, TEXT("No backend class set on %s. Mover actor will not function."), *GetNameSafe(GetOwner()));
+	}
+
+	Super::InitializeComponent();
+}
+
+
+void UMoverComponent::UninitializeComponent()
+{
+	if (UActorComponent* LiaisonAsComp = Cast<UActorComponent>(BackendLiaisonComp.Get()))
+	{
+		LiaisonAsComp->DestroyComponent();
+	}
+
+	Super::UninitializeComponent();
+}
+
+
+void UMoverComponent::OnRegister()
+{
+	TGuardValue<bool> InOnRegisterGuard(bInOnRegister, true);
+
+	UpdatedCompAsPrimitive = Cast<UPrimitiveComponent>(UpdatedComponent);
+	Super::OnRegister();
+
+	const UWorld* MyWorld = GetWorld();
+
+	if (MyWorld && MyWorld->IsGameWorld())
+	{
+		const AActor* MyActor = GetOwner();
+
+		USceneComponent* NewUpdatedComponent = UpdatedComponent;
+		if (!UpdatedComponent)
+		{
+			// Auto-register owner's root component if found.
+			if (MyActor)
+			{
+				NewUpdatedComponent = MyActor->GetRootComponent();
+			}
+		}
+
+		SetUpdatedComponent(NewUpdatedComponent);
+
+		// If no primary visual component is already set, fall back to searching for any kind of mesh
+		if (!PrimaryVisualComponent && MyActor)
+		{
+			PrimaryVisualComponent = MyActor->FindComponentByClass<UMeshComponent>();
+		}
+	}
+}
+
+
+void UMoverComponent::RegisterComponentTickFunctions(bool bRegister)
+{
+	Super::RegisterComponentTickFunctions(bRegister);
+
+	// Super may start up the tick function when we don't want to.
+	UpdateTickRegistration();
+
+	// If the owner ticks, make sure we tick first. This is to ensure the owner's location will be up to date when it ticks.
+	AActor* Owner = GetOwner();
+
+	if (bRegister && PrimaryComponentTick.bCanEverTick && Owner && Owner->CanEverTick())
+	{
+		Owner->PrimaryActorTick.AddPrerequisite(this, PrimaryComponentTick);
+	}
+}
+
+
+void UMoverComponent::PostLoad()
+{
+	Super::PostLoad();
+
+	RefreshSharedSettings();
+}
+
+void UMoverComponent::BeginPlay()
+{
+	Super::BeginPlay();
+
+	// If an InputProducer isn't already set, check the actor and its components for one
+	if (!InputProducer)
+	{
+		if (AActor* ActorOwner = GetOwner())
+		{
+			if (ActorOwner->GetClass()->ImplementsInterface(UMoverInputProducerInterface::StaticClass()))
+			{
+				InputProducer = ActorOwner;
+			}
+			else if (UActorComponent* FoundInputProducerComp = ActorOwner->FindComponentByInterface(UMoverInputProducerInterface::StaticClass()))
+			{
+				InputProducer = FoundInputProducerComp;
+			}
+		}
+	}
+}
+
+void UMoverComponent::TickComponent(float DeltaTime, enum ELevelTick TickType, FActorComponentTickFunction* ThisTickFunction)
+{
+	Super::TickComponent(DeltaTime, TickType, ThisTickFunction);
+}
+
+void UMoverComponent::ProduceInput(const int32 DeltaTimeMS, FMoverInputCmdContext* Cmd)
+{
+	Cmd->InputCollection.Empty();
+	
+	// Pass input production on to the right party
+	if (InputProducer)
+	{
+		IMoverInputProducerInterface::Execute_ProduceInput(InputProducer, DeltaTimeMS, OUT *Cmd);
+	}
+
+	CachedLastProducedInputCmd = *Cmd;
+	bHasValidLastProducedInput = true;
+
+}
+
+void UMoverComponent::RestoreFrame(const FMoverSyncState* SyncState, const FMoverAuxStateContext* AuxState)
+{
+	SetFrameStateFromContext(SyncState, AuxState, /* rebase? */ true);
+	OnSimulationRollback(SyncState, AuxState);
+}
+
+void UMoverComponent::FinalizeFrame(const FMoverSyncState* SyncState, const FMoverAuxStateContext* AuxState)
+{
+	const FMoverDefaultSyncState* MoverState = SyncState->SyncStateCollection.FindDataByType<FMoverDefaultSyncState>();
+
+	// TODO: Revisit this location check -- it seems simplistic now that we have composable state. Consider supporting a version that allows each sync state data struct a chance to react.
+	// The component will often be in the "right place" already on FinalizeFrame, so a comparison check makes sense before setting it.
+	if (MoverState &&
+			(UpdatedComponent->GetComponentLocation().Equals(MoverState->GetLocation_WorldSpace()) == false ||
+			 UpdatedComponent->GetComponentQuat().Rotator().Equals(MoverState->GetOrientation_WorldSpace(), ROTATOR_TOLERANCE) == false))
+	{
+		SetFrameStateFromContext(SyncState, AuxState, /* rebase? */ false);
+	}
+	else
+	{
+		// TODO: This is distasteful duplication -- consider moving to a util function
+		CachedLastSyncState = *SyncState;
+		CachedLastAuxState = *AuxState;
+		CachedLastSimTickTimeStep.BaseSimTimeMs = BackendLiaisonComp->GetCurrentSimTimeMs();
+		CachedLastSimTickTimeStep.ServerFrame = BackendLiaisonComp->GetCurrentSimFrame();
+		bHasValidCachedState = true;
+	}
+}
+
+
+void UMoverComponent::InitializeSimulationState(FMoverSyncState* OutSync, FMoverAuxStateContext* OutAux)
+{
+	npCheckSlow(UpdatedComponent);
+	npCheckSlow(OutSync);
+	npCheckSlow(OutAux);
+
+	// Add all initial persistent sync state types
+	for (const FMoverDataPersistence& PersistentSyncEntry : PersistentSyncStateDataTypes)
+	{
+		check(PersistentSyncEntry.RequiredType && PersistentSyncEntry.RequiredType->IsChildOf(FMoverDataStructBase::StaticStruct()));
+		OutSync->SyncStateCollection.FindOrAddDataByType(PersistentSyncEntry.RequiredType);
+	}
+
+	if (FMoverDefaultSyncState* MoverState = OutSync->SyncStateCollection.FindMutableDataByType<FMoverDefaultSyncState>())
+	{
+		MoverState->SetTransforms_WorldSpace(UpdatedComponent->GetComponentLocation(),
+		                                     UpdatedComponent->GetComponentRotation(),
+		                                     FVector::ZeroVector);	// no initial velocity
+
+		MoverState->MovementMode = StartingMovementMode;
+	}
+
+	*OutAux = FMoverAuxStateContext();
+
+	CachedLastSyncState = *OutSync;
+	CachedLastAuxState = *OutAux;
+	bHasValidCachedState = true;
+}
+
+void UMoverComponent::SimulationTick(const FMoverTimeStep& InTimeStep, const FMoverTickStartData& SimInput, OUT FMoverTickEndData& SimOutput)
+{
+	const bool bIsResimulating = InTimeStep.BaseSimTimeMs <= CachedNewestSimTickTimeStep.BaseSimTimeMs;
+
+	FMoverTimeStep MoverTimeStep(InTimeStep);
+	MoverTimeStep.bIsResimulating = bIsResimulating;
+
+	if (bHasRolledBack)
+	{
+		ProcessFirstSimTickAfterRollback(InTimeStep);
+	}
+
+	OnPreSimulationTick.Broadcast(MoverTimeStep, SimInput.InputCmd);
+
+	if (!ModeFSM->IsValidLowLevel())
+	{
+		SimOutput.SyncState = SimInput.SyncState;
+		SimOutput.AuxState = SimInput.AuxState;
+		return;
+	}
+
+	if (const FMoverDefaultSyncState* StartingSyncState = SimInput.SyncState.SyncStateCollection.FindDataByType<FMoverDefaultSyncState>())
+	{
+		ensureMsgf(UpdatedComponent->GetComponentLocation().Equals(StartingSyncState->GetLocation_WorldSpace()) || StartingSyncState->GetMovementBase(),
+			TEXT("Simulation start location (%s) disagrees with actual mover component location (%s). This indicates movement of the component out-of-band with the simulation, and will lead to poor quality motion."),
+			*StartingSyncState->GetLocation_WorldSpace().ToCompactString(),
+			*UpdatedComponent->GetComponentLocation().ToCompactString());
+	}
+
+	// Sync state data should carry over between frames
+	for (const FMoverDataPersistence& PersistentSyncEntry : PersistentSyncStateDataTypes)
+	{
+		bool bShouldAddDefaultData = true;
+
+		if (PersistentSyncEntry.bCopyFromPriorFrame)
+		{
+			if (const FMoverDataStructBase* PriorFrameData = SimInput.SyncState.SyncStateCollection.FindDataByType(PersistentSyncEntry.RequiredType))
+			{
+				SimOutput.SyncState.SyncStateCollection.AddOrOverwriteData(TSharedPtr<FMoverDataStructBase>(PriorFrameData->Clone()));
+				bShouldAddDefaultData = false;
+			}
+		}
+
+		if (bShouldAddDefaultData)
+		{
+			SimOutput.SyncState.SyncStateCollection.FindOrAddDataByType(PersistentSyncEntry.RequiredType);
+		}
+	}
+
+	SimOutput.AuxState = SimInput.AuxState;
+
+	FKinematicDefaultInputs* Input = SimInput.InputCmd.InputCollection.FindMutableDataByType<FKinematicDefaultInputs>();
+
+
+	if (Input && !Input->SuggestedMovementMode.IsNone())
+	{
+		ModeFSM->QueueNextMode(Input->SuggestedMovementMode);
+	}
+
+	// Tick the actual simulation. This is where the proposed moves are queried and executed, affecting change to the moving actor's gameplay state and captured in the output sim state
+	{
+		FScopedMovementUpdate ScopedMovementUpdate(UpdatedComponent, EScopedUpdate::DeferredUpdates);
+		ModeFSM->OnSimulationTick(UpdatedComponent, UpdatedCompAsPrimitive, SimBlackboard.Get(), SimInput, MoverTimeStep, SimOutput);
+	}
+
+	if (FMoverDefaultSyncState* OutputSyncState = SimOutput.SyncState.SyncStateCollection.FindMutableDataByType<FMoverDefaultSyncState>())
+	{
+		const FName MovementModeAfterTick = ModeFSM->GetCurrentModeName();
+		OutputSyncState->MovementMode = MovementModeAfterTick;
+
+		if (UpdatedComponent->GetComponentLocation().Equals(OutputSyncState->GetLocation_WorldSpace()) == false ||
+			UpdatedComponent->GetComponentQuat().Rotator().Equals(OutputSyncState->GetOrientation_WorldSpace(), ROTATOR_TOLERANCE) == false)
+		{
+			UE_LOG(LogMover, Warning, TEXT("Detected pos/rot difference between Mover actor (%s) sync state and scene component after sim ticking. This indicates a movement mode may not be authoring the final state correctly."), *GetNameSafe(UpdatedComponent->GetOwner()));
+		}
+	}
+
+	// Note that we don't pull the rotation out of the final update transform. Converting back from a quat will lead to a different FRotator than what we are storing
+	// here in the simulation layer. This may not be the best choice for all movement simulations, but is ok for this one.
+	// JAH TODO: re-evaluate the above comment about handling the rotation data
+
+	if (!SimOutput.MoveRecord.GetTotalMoveDelta().IsZero())
+	{
+		UE_LOG(LogMover, VeryVerbose, TEXT("KinematicSimTick: %s (role %i) frame %d: %s"),
+			*GetNameSafe(UpdatedComponent->GetOwner()), UpdatedComponent->GetOwnerRole(), MoverTimeStep.ServerFrame, *SimOutput.MoveRecord.ToString());
+	}
+
+	OnPostMovement.Broadcast(MoverTimeStep, SimOutput.SyncState, SimOutput.AuxState);
+
+	CachedLastUsedInputCmd = SimInput.InputCmd;
+	bHasValidCachedUsedInput = true;
+
+	OnPostSimulationTick.Broadcast(MoverTimeStep);
+
+	CachedLastSimTickTimeStep = MoverTimeStep;
+
+	if (MoverTimeStep.ServerFrame > CachedNewestSimTickTimeStep.ServerFrame || MoverTimeStep.BaseSimTimeMs > CachedNewestSimTickTimeStep.BaseSimTimeMs)
+	{
+		CachedNewestSimTickTimeStep = MoverTimeStep;
+	}	
+}
+
+UBaseMovementMode* UMoverComponent::FindMovementMode(TSubclassOf<UBaseMovementMode> MovementMode) const
+{
+	return FindMode_Mutable(MovementMode);
+}
+
+void UMoverComponent::SetFrameStateFromContext(const FMoverSyncState* SyncState, const FMoverAuxStateContext* AuxState, bool bRebaseBasedState)
+{
+
+	// TODO: This is distasteful duplication -- consider moving to a util function
+	// Copy these as the last official state, so they can be queried by other systems outside of the NP simulation
+	CachedLastSyncState = *SyncState;
+	CachedLastAuxState = *AuxState;
+	CachedLastSimTickTimeStep.BaseSimTimeMs = BackendLiaisonComp->GetCurrentSimTimeMs();
+	CachedLastSimTickTimeStep.ServerFrame = BackendLiaisonComp->GetCurrentSimFrame();
+	bHasValidCachedState = true;
+
+	if (FMoverDefaultSyncState* MoverState = CachedLastSyncState.SyncStateCollection.FindMutableDataByType<FMoverDefaultSyncState>())
+	{
+		if (bRebaseBasedState && MoverState->GetMovementBase())
+		{
+			// Note that this is modifying our cached mover state from what we received from Network Prediction. We are resampling
+			// the transform of the movement base, in case it has changed as well during the rollback.
+			MoverState->UpdateCurrentMovementBase();
+		}
+
+		// The state's properties are usually worldspace already, but may need to be adjusted to match the current movement base
+		const FVector WorldLocation = MoverState->GetLocation_WorldSpace();
+		const FRotator WorldOrientation = MoverState->GetOrientation_WorldSpace();
+		const FVector WorldVelocity = MoverState->GetVelocity_WorldSpace();
+
+		// Apply the desired transform to the scene component
+		FTransform Transform(WorldOrientation, WorldLocation, UpdatedComponent->GetComponentTransform().GetScale3D());
+		UpdatedComponent->SetWorldTransform(Transform, false, nullptr, ETeleportType::TeleportPhysics);
+		UpdatedComponent->ComponentVelocity = WorldVelocity;
+	}
+}
+
+
+bool UMoverComponent::InitMoverSimulation()
+{
+	check(UpdatedComponent);
+
+	SimBlackboard = NewObject<UMoverBlackboard>(this, TEXT("MoverBlackboard"), RF_Transient);
+
+	// Set up FSM and initial movement states
+	ModeFSM = NewObject<UMovementModeStateMachine>(this, TEXT("MoverStateMachine"), RF_Transient);
+	ModeFSM->ClearAllMovementModes();
+
+	bool bHasMatchingStartingState = false;
+
+	for (const TPair<FName, TObjectPtr<UBaseMovementMode>>& Element : MovementModes)
+	{
+		ModeFSM->RegisterMovementMode(Element.Key, Element.Value);
+
+		bHasMatchingStartingState |= (StartingMovementMode == Element.Key);
+	}
+
+	if (!bHasMatchingStartingState)
+	{
+		UE_LOG(LogMover, Warning, TEXT("Invalid StartingMovementMode '%s' specified on %s. Mover actor will not function."),
+			*StartingMovementMode.ToString(), *GetNameSafe(GetOwner()));
+	}
+
+	if (bHasMatchingStartingState && StartingMovementMode != NAME_None)
+	{
+		ModeFSM->SetDefaultMode(StartingMovementMode);
+		ModeFSM->QueueNextMode(StartingMovementMode);
+	}
+
+	return bHasMatchingStartingState;
+}
+
+void UMoverComponent::HandleImpact(FMoverOnImpactParams& ImpactParams)
+{
+	if (ImpactParams.MovementModeName.IsNone())
+	{
+		ImpactParams.MovementModeName = ModeFSM->GetCurrentModeName();
+	}
+	
+	OnHandleImpact(ImpactParams);
+	OnImpact.Broadcast(ImpactParams);
+}
+
+void UMoverComponent::OnHandleImpact(const FMoverOnImpactParams& ImpactParams)
+{
+	// TODO: Handle physics impacts here - ie when player runs into box, impart force onto box
+}
+
+void UMoverComponent::SetUpdatedComponent(USceneComponent* NewUpdatedComponent)
+{
+	// Remove delegates from old component
+	if (UpdatedComponent)
+	{
+		UpdatedComponent->SetShouldUpdatePhysicsVolume(false);
+		UpdatedComponent->SetPhysicsVolume(NULL, true);
+		UpdatedComponent->PhysicsVolumeChangedDelegate.RemoveDynamic(this, &UMoverComponent::PhysicsVolumeChanged);
+
+		// remove from tick prerequisite
+		UpdatedComponent->PrimaryComponentTick.RemovePrerequisite(this, PrimaryComponentTick);
+	}
+
+	if (UpdatedCompAsPrimitive)
+	{
+		UpdatedCompAsPrimitive->OnComponentBeginOverlap.RemoveDynamic(this, &UMoverComponent::OnBeginOverlap);
+	}
+
+	// Don't assign pending kill components, but allow those to null out previous UpdatedComponent.
+	UpdatedComponent = GetValid(NewUpdatedComponent);
+	UpdatedCompAsPrimitive = Cast<UPrimitiveComponent>(UpdatedComponent);
+
+	// Assign delegates
+	if (IsValid(UpdatedComponent))
+	{
+		UpdatedComponent->SetShouldUpdatePhysicsVolume(true);
+		UpdatedComponent->PhysicsVolumeChangedDelegate.AddUniqueDynamic(this, &UMoverComponent::PhysicsVolumeChanged);
+
+		if (!bInOnRegister && !bInInitializeComponent)
+		{
+			// UpdateOverlaps() in component registration will take care of this.
+			UpdatedComponent->UpdatePhysicsVolume(true);
+		}
+
+		// force ticks after movement component updates
+		UpdatedComponent->PrimaryComponentTick.AddPrerequisite(this, PrimaryComponentTick);
+	}
+
+	if (IsValid(UpdatedCompAsPrimitive))
+	{
+		UpdatedCompAsPrimitive->OnComponentBeginOverlap.AddDynamic(this, &UMoverComponent::OnBeginOverlap);
+	}
+
+
+	UpdateTickRegistration();
+}
+
+
+void UMoverComponent::UpdateTickRegistration()
+{
+	const bool bHasUpdatedComponent = (UpdatedComponent != NULL);
+	SetComponentTickEnabled(bHasUpdatedComponent && bAutoActivate);
+}
+
+
+void UMoverComponent::OnSimulationRollback(const FMoverSyncState* SyncState, const FMoverAuxStateContext* AuxState)
+{
+	SimBlackboard->Invalidate(EInvalidationReason::Rollback);
+	ModeFSM->OnSimulationRollback(SyncState, AuxState);
+
+	bHasRolledBack = true;
+}
+
+
+void UMoverComponent::ProcessFirstSimTickAfterRollback(const FMoverTimeStep& TimeStep)
+{
+	OnPostSimulationRollback.Broadcast(TimeStep, CachedLastSimTickTimeStep);
+	bHasRolledBack = false;
+}
+
+#if WITH_EDITOR
+
+void UMoverComponent::PostCDOCompiled(const FPostCDOCompiledContext& Context)
+{
+	Super::PostCDOCompiled(Context);
+
+	RefreshSharedSettings();
+}
+
+
+void UMoverComponent::PostEditChangeProperty(FPropertyChangedEvent& PropertyChangedEvent)
+{
+	if ((PropertyChangedEvent.Property) && (PropertyChangedEvent.Property->GetFName() == GET_MEMBER_NAME_CHECKED(UMoverComponent, MovementModes)))
+	{
+		RefreshSharedSettings();
+	}
+
+	Super::PostEditChangeProperty(PropertyChangedEvent);
+}
+
+
+void UMoverComponent::PostTransacted(const FTransactionObjectEvent& TransactionEvent)
+{
+	Super::PostTransacted(TransactionEvent);
+
+	if ((TransactionEvent.GetEventType() == ETransactionObjectEventType::Finalized || TransactionEvent.GetEventType() == ETransactionObjectEventType::UndoRedo) &&
+		TransactionEvent.HasPropertyChanges() &&
+		TransactionEvent.GetChangedProperties().Contains(GET_MEMBER_NAME_CHECKED(UMoverComponent, MovementModes)))
+	{
+		RefreshSharedSettings();		
+	}
+}
+
+
+EDataValidationResult UMoverComponent::IsDataValid(FDataValidationContext& Context) const
+{
+	EDataValidationResult Result = Super::IsDataValid(Context);
+
+	if (!ValidateSetup(Context))
+	{
+		Result = EDataValidationResult::Invalid;
+	}
+
+	return Result;
+}
+
+
+bool UMoverComponent::ValidateSetup(FDataValidationContext& Context) const
+{
+	bool bHasMatchingStartingMode = false;
+	bool bDidFindAnyProblems = false;
+
+	// Verify all movement modes
+	for (const TPair<FName, TObjectPtr<UBaseMovementMode>>& Element : MovementModes)
+	{
+		if (StartingMovementMode == Element.Key)
+		{
+			bHasMatchingStartingMode = true;
+		}
+
+		// Verify movement mode is valid
+		if (!Element.Value)
+		{
+			Context.AddError(FText::Format(LOCTEXT("InvalidMovementModeError", "Invalid movement mode on {0}, mapped as {1}. Mover actor will not function."),
+				FText::FromString(GetNameSafe(GetOwner())),
+				FText::FromName(Element.Key)));
+
+			bDidFindAnyProblems = true;
+		}
+		else if (Element.Value->IsDataValid(Context) == EDataValidationResult::Invalid)
+		{
+			bDidFindAnyProblems = true;
+		}
+
+		// Verify that the movement mode's shared settings object exists (if any)
+		if (Element.Value && Element.Value->SharedSettingsClass &&
+			FindSharedSettings(Element.Value->SharedSettingsClass) == nullptr)
+		{
+			Context.AddError(FText::Format(LOCTEXT("MissingModeSettingsError", "Movement mode on {0}, mapped as {1}, is missing its desired SharedSettingsClass {2} and may not function properly."),
+				FText::FromString(GetNameSafe(GetOwner())),
+				FText::FromName(Element.Key),
+				FText::FromString(Element.Value->SharedSettingsClass->GetName())));
+
+			bDidFindAnyProblems = true;
+		}
+	}
+
+	// Verify we have a matching starting mode
+	if (!bHasMatchingStartingMode && StartingMovementMode != NAME_None)
+	{
+		Context.AddError(FText::Format(LOCTEXT("InvalidStartingModeError", "Invalid StartingMovementMode {0} specified on {1}. Mover actor will not function."),
+			FText::FromName(StartingMovementMode),
+			FText::FromString(GetNameSafe(GetOwner()))));
+
+		bDidFindAnyProblems = true;
+	}
+
+	// Verify backend liaison
+	if (!BackendClass)
+	{
+		Context.AddError(FText::Format(LOCTEXT("MissingBackendClassError", "No BackendClass property specified on {0}. Mover actor will not function."),
+			FText::FromString(GetNameSafe(GetOwner()))));
+
+		bDidFindAnyProblems = true;
+	}
+	else if (!BackendClass->ImplementsInterface(UMoverBackendLiaisonInterface::StaticClass()))
+	{
+		Context.AddError(FText::Format(LOCTEXT("InvalidBackendClassError", "BackendClass {0} on {1} does not implement IMoverBackendLiaisonInterface. Mover actor will not function."),
+			FText::FromString(BackendClass->GetName()),
+			FText::FromString(GetNameSafe(GetOwner()))));
+
+		bDidFindAnyProblems = true;
+	}
+
+	// Verify persistent types
+	for (const FMoverDataPersistence& PersistentSyncEntry : PersistentSyncStateDataTypes)
+	{
+		if (!PersistentSyncEntry.RequiredType || !PersistentSyncEntry.RequiredType->IsChildOf(FMoverDataStructBase::StaticStruct()))
+		{
+			Context.AddError(FText::Format(LOCTEXT("InvalidSyncStateTypeError", "RequiredType '{0}' is not a valid type or is missing. Must be a child of FMoverDataStructBase."),
+				FText::FromString(GetNameSafe(PersistentSyncEntry.RequiredType))));
+
+			bDidFindAnyProblems = true;
+		}
+	}
+
+	
+	return !bDidFindAnyProblems;
+}
+
+#endif // WITH_EDITOR
+
+
+void UMoverComponent::PhysicsVolumeChanged(APhysicsVolume* NewVolume)
+{
+	// This itself feels bad. When will this be called? Its impossible to know what is allowed and not allowed to be done in this callback.
+	// Callbacks instead should be trapped within the simulation update function. This isn't really possible though since the UpdateComponent
+	// is the one that will call this.
+}
+
+
+void UMoverComponent::RefreshSharedSettings()
+{
+	TArray<TObjectPtr<UObject>> UnreferencedSettingsObjs = SharedSettings;
+
+	// Add any missing settings
+	for (const TPair<FName, TObjectPtr<UBaseMovementMode>>& Element : MovementModes)
+	{
+		if (UBaseMovementMode* Mode = Element.Value.Get())
+		{
+			if (Mode->SharedSettingsClass != nullptr)
+			{
+				bool bFoundMatchingClass = false;
+				for (const TObjectPtr<UObject>& SettingsObj : SharedSettings)
+				{
+					if (SettingsObj && SettingsObj->IsA(Mode->SharedSettingsClass))
+					{
+						bFoundMatchingClass = true;
+						UnreferencedSettingsObjs.Remove(SettingsObj);
+						break;
+					}
+				}
+
+				if (!bFoundMatchingClass)
+				{
+					UObject* NewSettings = NewObject<UObject>(this, Mode->SharedSettingsClass, NAME_None, GetMaskedFlags(RF_PropagateToSubObjects) | RF_Transactional);
+					SharedSettings.Add(NewSettings);
+				}
+
+			}
+		}
+	}
+
+	// Remove any settings that are no longer used
+	for (const TObjectPtr<UObject>& SettingsObjToRemove : UnreferencedSettingsObjs)
+	{
+		SharedSettings.Remove(SettingsObjToRemove);
+	}
+
+	// Sort by name for array order consistency
+	Algo::Sort(SharedSettings, [](const TObjectPtr<UObject>& LHS, const TObjectPtr<UObject>& RHS) 
+		{ return (LHS->GetClass()->GetPathName() < RHS.GetClass()->GetPathName()); });
+}
+
+
+void UMoverComponent::K2_QueueLayeredMove(const int32& MoveAsRawData)
+{
+	// This will never be called, the exec version below will be hit instead
+	checkNoEntry();
+}
+
+DEFINE_FUNCTION(UMoverComponent::execK2_QueueLayeredMove)
+{
+	Stack.StepCompiledIn<FStructProperty>(nullptr);
+	void* MovePtr = Stack.MostRecentPropertyAddress;
+	FStructProperty* StructProp = CastField<FStructProperty>(Stack.MostRecentProperty);
+
+	P_FINISH;
+
+	P_NATIVE_BEGIN;
+
+	const bool bHasValidStructProp = StructProp && StructProp->Struct && StructProp->Struct->IsChildOf(FLayeredMoveBase::StaticStruct());
+
+	if (ensureMsgf((bHasValidStructProp && MovePtr), TEXT("An invalid struct type (%s) was sent to a QueueLayeredMove node. A struct derived from FLayeredMoveBase is required. No layered move will be queued."), *GetNameSafe(StructProp->Struct)))
+	{
+		// Could we steal this instead of cloning? (move semantics)
+		FLayeredMoveBase* MoveAsBasePtr = reinterpret_cast<FLayeredMoveBase*>(MovePtr);
+		FLayeredMoveBase* ClonedMove = MoveAsBasePtr->Clone();
+
+		P_THIS->QueueLayeredMove(TSharedPtr<FLayeredMoveBase>(ClonedMove));
+	}
+
+	P_NATIVE_END;
+}
+
+
+void UMoverComponent::QueueLayeredMove(TSharedPtr<FLayeredMoveBase> LayeredMove)
+{	
+	ModeFSM->QueueLayeredMove(LayeredMove);
+}
+
+
+void UMoverComponent::QueueNextMode(FName DesiredModeName, bool bShouldReenter)
+{
+	ModeFSM->QueueNextMode(DesiredModeName, bShouldReenter);
+}
+
+UBaseMovementMode* UMoverComponent::AddMovementModeFromClass(FName ModeName, TSubclassOf<UBaseMovementMode> MovementMode)
+{
+	if (!MovementMode)
+	{
+		UE_LOG(LogMover, Warning, TEXT("Attempted to add a movement mode that wasn't valid. AddMovementModeFromClass did not add anything."));
+		return nullptr;
+	}
+
+	TObjectPtr<UBaseMovementMode> AddedMovementMode =  NewObject<UBaseMovementMode>(this, MovementMode);
+	AddMovementModeFromObject(ModeName, AddedMovementMode);
+	
+	return AddedMovementMode; 
+}
+
+bool UMoverComponent::AddMovementModeFromObject(FName ModeName, UBaseMovementMode* MovementMode)
+{
+	if (MovementMode)
+	{
+		if (TObjectPtr<UBaseMovementMode>* FoundMovementMode = MovementModes.Find(ModeName))
+		{
+			if (FoundMovementMode->Get()->GetClass() == MovementMode->GetClass())
+			{
+				UE_LOG(LogMover, Warning, TEXT("Added the same movement mode (%s) for a movement mode name (%s). AddMovementModeFromObject will add the mode but is likely unwanted/unnecessary behavior."), *GetNameSafe(MovementMode), *ModeName.ToString());
+			}
+
+			RemoveMovementMode(ModeName);
+		}
+		
+		if (MovementMode->GetOuter() != this)
+		{
+			UE_LOG(LogMover, Verbose, TEXT("Movement modes are expected to be parented to the MoverComponent. The %s movement mode was reparented to %s!"), *GetNameSafe(MovementMode), *GetNameSafe(this));
+			MovementMode->Rename(nullptr, this, REN_ForceNoResetLoaders | REN_DoNotDirty | REN_NonTransactional);
+		}
+		
+		MovementModes.Add(ModeName, MovementMode);
+		ModeFSM->RegisterMovementMode(ModeName, MovementMode);
+	}
+	else
+	{
+		UE_LOG(LogMover, Warning, TEXT("Attempted to add %s movement mode that wasn't valid to %s. AddMovementModeFromObject did not add anything."), *GetNameSafe(MovementMode), *GetNameSafe(this));
+		return false;
+	}
+
+	return true;
+}
+
+bool UMoverComponent::RemoveMovementMode(FName ModeName)
+{
+	if (ModeFSM->GetCurrentModeName() == ModeName)
+	{
+		UE_LOG(LogMover, Warning, TEXT("The mode being removed (%s Movement Mode) is the mode this actor is currently in. It was removed but may cause issues. Consider waiting to remove the mode or queueing a different valid mode to avoid issues."), *ModeName.ToString());
+	}
+	
+	TObjectPtr<UBaseMovementMode>* ModeToRemove = MovementModes.Find(ModeName);
+	const bool ModeRemoved = MovementModes.Remove(ModeName) > 0;
+	if (ModeRemoved && ModeToRemove)
+	{
+		ModeFSM->UnregisterMovementMode(ModeName);
+		ModeToRemove->Get()->ConditionalBeginDestroy();
+	}
+	
+	return ModeRemoved; 
+}
+
+
+FTransform UMoverComponent::GetUpdatedComponentTransform() const
+{
+	if (ensure(UpdatedComponent))
+	{
+		return UpdatedComponent->GetComponentTransform();
+	}
+	return FTransform::Identity;
+}
+
+
+USceneComponent* UMoverComponent::GetPrimaryVisualComponent() const
+{
+	return PrimaryVisualComponent.Get();
+}
+
+
+FVector UMoverComponent::GetVelocity() const
+{ 
+	if (bHasValidCachedState)
+	{
+		if (const FMoverDefaultSyncState* SyncState = CachedLastSyncState.SyncStateCollection.FindDataByType<FMoverDefaultSyncState>())
+		{
+			return SyncState->GetVelocity_WorldSpace();
+		}
+	}
+
+	return FVector::ZeroVector;
+}
+
+
+FVector UMoverComponent::GetMovementIntent() const
+{ 
+	if (bHasValidCachedState)
+	{
+		if (const FMoverDefaultSyncState* SyncState = CachedLastSyncState.SyncStateCollection.FindDataByType<FMoverDefaultSyncState>())
+		{
+			return SyncState->GetIntent_WorldSpace();
+		}
+	}
+
+	return FVector::ZeroVector; 
+}
+
+
+FRotator UMoverComponent::GetTargetOrientation() const
+{
+	// Prefer the input's intended orientation, but if it can't be determined, assume it matches the actual orientation
+	if (bHasValidCachedUsedInput)
+	{
+		const FMoverInputCmdContext& LastInputCmd = GetLastInputCmd();
+		const FKinematicDefaultInputs* MoverInputs = LastInputCmd.InputCollection.FindDataByType<FKinematicDefaultInputs>();
+
+		const FVector TargetOrientationDir = MoverInputs->GetOrientationIntentDir_WorldSpace();
+
+		if (!TargetOrientationDir.IsNearlyZero())
+		{
+			return TargetOrientationDir.ToOrientationRotator();
+		}
+	}
+	
+	if (bHasValidCachedState)
+	{
+		if (const FMoverDefaultSyncState* SyncState = CachedLastSyncState.SyncStateCollection.FindDataByType<FMoverDefaultSyncState>())
+		{
+			return SyncState->GetOrientation_WorldSpace();
+		}
+	}
+
+	return GetOwner() ? GetOwner()->GetActorRotation() : FRotator::ZeroRotator;
+}
+
+
+void UMoverComponent::SetGravityOverride(bool bOverrideGravity, FVector NewGravityAcceleration)
+{
+	bHasGravityOverride = bOverrideGravity;
+	GravityAccelOverride = NewGravityAcceleration;
+}
+
+
+FVector UMoverComponent::GetGravityAcceleration() const
+{
+	if (bHasGravityOverride)
+	{
+		return GravityAccelOverride;
+	}
+
+	if (UpdatedComponent)
+	{
+		APhysicsVolume* CurPhysVolume = UpdatedComponent->GetPhysicsVolume();
+		if (CurPhysVolume)
+		{
+			return CurPhysVolume->GetGravityZ() * FVector::UpVector;
+		}
+	}
+
+	return MoverComponentConstants::DefaultGravityAccel;
+}
+
+
+FVector UMoverComponent::GetUpDirection() const
+{
+	const FVector DeducedUpDir = -GetGravityAcceleration().GetSafeNormal();
+
+	if (DeducedUpDir.IsZero())
+	{
+		return MoverComponentConstants::DefaultUpDir;
+	}
+
+	return DeducedUpDir;
+}
+
+
+bool UMoverComponent::IsFalling() const
+{ 
+	if (bHasValidCachedState)
+	{
+		if (const FMoverDefaultSyncState* SyncState = CachedLastSyncState.SyncStateCollection.FindDataByType<FMoverDefaultSyncState>())
+		{
+			return SyncState->MovementMode == KinematicModeNames::Falling;
+		}
+	}
+
+	return false;
+}
+
+bool UMoverComponent::IsAirborne() const
+{
+	if (bHasValidCachedState)
+	{
+		if (const FMoverDefaultSyncState* SyncState = CachedLastSyncState.SyncStateCollection.FindDataByType<FMoverDefaultSyncState>())
+		{
+			return SyncState->MovementMode == KinematicModeNames::Flying || SyncState->MovementMode == KinematicModeNames::Falling;
+		}
+	}
+
+	return false;
+}
+
+bool UMoverComponent::IsOnGround() const
+{
+	if (bHasValidCachedState)
+	{
+		if (const FMoverDefaultSyncState* SyncState = CachedLastSyncState.SyncStateCollection.FindDataByType<FMoverDefaultSyncState>())
+		{
+			return SyncState->MovementMode == KinematicModeNames::Walking;
+		}
+	}
+
+	return false;
+}
+
+bool UMoverComponent::IsSlopeSliding() const
+{
+	if (IsAirborne())
+	{
+		FFloorCheckResult HitResult;
+		if (SimBlackboard->TryGet(KinematicBlackboard::LastFloorResult, HitResult))
+		{
+			return HitResult.bBlockingHit && !HitResult.bWalkableFloor;
+		}
+	}
+
+	return false;
+}
+
+TArray<FTrajectorySampleInfo> UMoverComponent::GetFutureTrajectory(float FutureSeconds, float SamplesPerSecond) const
+{
+	FMoverTickStartData StartingState;
+
+	StartingState.InputCmd = GetLastInputCmd();
+	StartingState.SyncState = CachedLastSyncState;
+	StartingState.AuxState = CachedLastAuxState;
+
+	FMoverTickStartData StepState = StartingState;
+
+	const int32 NumSamples = FMath::Max(1, FutureSeconds * SamplesPerSecond);
+	const float StepDeltaSeconds = FutureSeconds / (float)NumSamples;
+
+	FMoverTimeStep FutureTimeStep;
+	FutureTimeStep.StepMs = (FutureSeconds * 1000.f) / NumSamples;
+	FutureTimeStep.BaseSimTimeMs = CachedLastSimTickTimeStep.BaseSimTimeMs;
+	FutureTimeStep.ServerFrame = 0;
+
+	if (const UBaseMovementMode* CurrentMovementMode = ModeFSM->GetCurrentMode())
+	{
+		TArray<FTrajectorySampleInfo> OutSamples;
+		OutSamples.AddDefaulted(NumSamples);
+
+		if (FMoverDefaultSyncState* StepSyncState = StepState.SyncState.SyncStateCollection.FindMutableDataByType<FMoverDefaultSyncState>())
+		{
+			FVector PriorLocation = StepSyncState->GetLocation_WorldSpace();
+			FRotator PriorOrientation = StepSyncState->GetOrientation_WorldSpace();
+			FVector PriorVelocity = StepSyncState->GetVelocity_WorldSpace();
+
+			for (int32 i = 0; i < NumSamples; ++i)
+			{
+				// Capture sample from current step state
+				FTrajectorySampleInfo& Sample = OutSamples[i];
+
+				Sample.Transform.SetLocation(StepSyncState->GetLocation_WorldSpace());
+				Sample.Transform.SetRotation(StepSyncState->GetOrientation_WorldSpace().Quaternion());
+				Sample.LinearVelocity = StepSyncState->GetVelocity_WorldSpace();
+				Sample.InstantaneousAcceleration = (StepSyncState->GetVelocity_WorldSpace() - PriorVelocity) / StepDeltaSeconds;
+				Sample.AngularVelocity = (StepSyncState->GetOrientation_WorldSpace() - PriorOrientation) * (1.f / StepDeltaSeconds);
+
+				Sample.SimTimeMs = FutureTimeStep.BaseSimTimeMs;
+
+				// Cache prior values
+				PriorLocation = StepSyncState->GetLocation_WorldSpace();
+				PriorOrientation = StepSyncState->GetOrientation_WorldSpace();
+				PriorVelocity = StepSyncState->GetVelocity_WorldSpace();
+
+				// Generate next move from current step state
+				FProposedMove StepMove;
+				CurrentMovementMode->DoGenerateMove(StepState, FutureTimeStep, StepMove);
+
+				// Advance state based on move
+				StepSyncState->SetTransforms_WorldSpace(StepSyncState->GetLocation_WorldSpace() + (StepMove.LinearVelocity * StepDeltaSeconds),
+					StepSyncState->GetOrientation_WorldSpace() + (StepMove.AngularVelocity * StepDeltaSeconds),
+					StepMove.LinearVelocity,
+					StepSyncState->GetMovementBase(),
+					StepSyncState->GetMovementBaseBoneName());
+
+				FutureTimeStep.BaseSimTimeMs += FutureTimeStep.StepMs;
+				++FutureTimeStep.ServerFrame;
+			}
+		}
+
+		return OutSamples;
+	}
+
+	return TArray<FTrajectorySampleInfo>();
+}
+
+
+FName UMoverComponent::GetMovementModeName() const
+{ 
+	if (bHasValidCachedState)
+	{
+		if (const FMoverDefaultSyncState* SyncState = CachedLastSyncState.SyncStateCollection.FindDataByType<FMoverDefaultSyncState>())
+		{
+			return SyncState->MovementMode;
+		}
+	}
+
+	return NAME_None;
+}
+
+
+UPrimitiveComponent* UMoverComponent::GetMovementBase() const
+{
+	if (bHasValidCachedState)
+	{
+		if (const FMoverDefaultSyncState* SyncState = CachedLastSyncState.SyncStateCollection.FindDataByType<FMoverDefaultSyncState>())
+		{
+			return SyncState->GetMovementBase();
+		}
+	}
+
+	return nullptr;
+}
+
+FName UMoverComponent::GetMovementBaseBoneName() const
+{
+	if (bHasValidCachedState)
+	{
+		if (const FMoverDefaultSyncState* SyncState = CachedLastSyncState.SyncStateCollection.FindDataByType<FMoverDefaultSyncState>())
+		{
+			return SyncState->GetMovementBaseBoneName();
+		}
+	}
+
+	return NAME_None;
+}
+
+
+bool UMoverComponent::HasValidCachedState() const
+{
+	return bHasValidCachedState;
+}
+
+const FMoverSyncState& UMoverComponent::GetSyncState() const
+{
+	if (!bHasValidCachedState)
+	{
+		UE_LOG(LogMover, Warning, TEXT("Attempting direct access to the last-cached sync state before one has been set. Results will be unreliable."));
+	}
+
+	return CachedLastSyncState;
+}
+
+bool UMoverComponent::TryGetFloorCheckHitResult(FHitResult& OutHitResult) const
+{
+	FFloorCheckResult FloorCheck;
+	if (SimBlackboard->TryGet(KinematicBlackboard::LastFloorResult, FloorCheck))
+	{
+		OutHitResult = FloorCheck.HitResult;
+		return true;
+	}
+	return false;
+}
+
+const UMoverBlackboard* UMoverComponent::GetSimBlackboard() const
+{
+	return SimBlackboard;
+}
+
+bool UMoverComponent::HasValidCachedInputCmd() const
+{
+	return bHasValidCachedUsedInput;
+}
+
+const FMoverInputCmdContext& UMoverComponent::GetLastInputCmd() const
+{
+	if (!bHasValidCachedUsedInput)
+	{
+		UE_LOG(LogMover, Warning, TEXT("Attempting direct access to the last-cached used input cmd before one has been set. Results will be unreliable."));
+	}
+
+	return CachedLastUsedInputCmd;
+}
+
+IMovementSettingsInterface* UMoverComponent::FindSharedSettings_Mutable(const UClass* ByType) const
+{
+	check(ByType);
+
+	for (const TObjectPtr<UObject>& SettingsObj : SharedSettings)
+	{
+		if (SettingsObj->IsA(ByType))
+		{
+			return Cast<IMovementSettingsInterface>(SettingsObj);
+		}
+	}
+
+	return nullptr;
+}
+
+UObject* UMoverComponent::FindSharedSettings_Mutable_BP(TSubclassOf<UObject> SharedSetting) const
+{
+	if (SharedSetting->ImplementsInterface(UMovementSettingsInterface::StaticClass()))
+    {
+    	return Cast<UObject>(FindSharedSettings_Mutable(SharedSetting));
+    }
+    
+    return nullptr;
+}
+
+const UObject* UMoverComponent::FindSharedSettings_BP(TSubclassOf<UObject> SharedSetting) const
+{
+	if (SharedSetting->ImplementsInterface(UMovementSettingsInterface::StaticClass()))
+	{
+		return Cast<UObject>(FindSharedSettings(SharedSetting));
+	}
+
+	return nullptr;
+}
+
+UBaseMovementMode* UMoverComponent::FindMode_Mutable(const UClass* ByType, bool bRequireExactClass) const
+{
+	for (const TPair<FName, TObjectPtr<UBaseMovementMode>>& Element : MovementModes)
+	{
+		if ( (!bRequireExactClass && Element.Value->IsA(ByType)) || 
+		     (Element.Value->GetClass() == ByType) )
+		{
+			return Element.Value.Get();
+		}
+	}
+
+	return nullptr;
+}
+
+#undef LOCTEXT_NAMESPACE
diff --git a/Engine/Plugins/Experimental/Mover/Source/Mover/Private/MoverDeveloperSettings.cpp b/Engine/Plugins/Experimental/Mover/Source/Mover/Private/MoverDeveloperSettings.cpp
new file mode 100644
--- /dev/null
+++ b/Engine/Plugins/Experimental/Mover/Source/Mover/Private/MoverDeveloperSettings.cpp
@@ -0,0 +1,8 @@
+﻿// Copyright Epic Games, Inc. All Rights Reserved.
+
+#include "MoverDeveloperSettings.h"
+
+UMoverDeveloperSettings::UMoverDeveloperSettings()
+	: MaxTimesToRefundSubstep(10)
+{
+}
diff --git a/Engine/Plugins/Experimental/Mover/Source/Mover/Private/PhysicsMover/Modes/PhysicsDrivenFallingMode.cpp b/Engine/Plugins/Experimental/Mover/Source/Mover/Private/PhysicsMover/Modes/PhysicsDrivenFallingMode.cpp
new file mode 100644
--- /dev/null
+++ b/Engine/Plugins/Experimental/Mover/Source/Mover/Private/PhysicsMover/Modes/PhysicsDrivenFallingMode.cpp
@@ -0,0 +1,126 @@
+// Copyright Epic Games, Inc. All Rights Reserved.
+
+#include "PhysicsMover/Modes/PhysicsDrivenFallingMode.h"
+
+#include "Chaos/Character/CharacterGroundConstraint.h"
+#include "GameFramework/PhysicsVolume.h"
+#include "MoverComponent.h"
+#include "Kinematic/Settings/CommonLegacyMovementSettings.h"
+#include "Math/UnitConversion.h"
+#include "MoveLibrary/FloorQueryUtils.h"
+#include "MoveLibrary/WaterMovementUtils.h"
+#include "PhysicsMover/PhysicsMovementUtils.h"
+
+#include UE_INLINE_GENERATED_CPP_BY_NAME(PhysicsDrivenFallingMode)
+
+UPhysicsDrivenFallingMode::UPhysicsDrivenFallingMode(const FObjectInitializer& ObjectInitializer)
+	: Super(ObjectInitializer)
+{
+}
+
+void UPhysicsDrivenFallingMode::UpdateConstraintSettings(Chaos::FCharacterGroundConstraint& Constraint) const
+{
+	Constraint.SetRadialForceLimit(300000.0); // TEMP - Move radial force limit to shared mode data
+	Constraint.SetTwistTorqueLimit(FUnitConversion::Convert(TwistTorqueLimit, EUnit::NewtonMeters, EUnit::KilogramCentimetersSquaredPerSecondSquared));
+	Constraint.SetSwingTorqueLimit(FUnitConversion::Convert(SwingTorqueLimit, EUnit::NewtonMeters, EUnit::KilogramCentimetersSquaredPerSecondSquared));
+	Constraint.SetTargetHeight(TargetHeight);
+}
+
+void UPhysicsDrivenFallingMode::OnSimulationTick(const FSimulationTickParams& Params, FMoverTickEndData& OutputState)
+{
+	const FMoverTickStartData& StartState = Params.StartState;
+	USceneComponent* UpdatedComponent = Params.UpdatedComponent;
+	UPrimitiveComponent* UpdatedPrimitive = Params.UpdatedPrimitive;
+	FProposedMove ProposedMove = Params.ProposedMove;
+
+	const FKinematicDefaultInputs* KinematicInputs = StartState.InputCmd.InputCollection.FindDataByType<FKinematicDefaultInputs>();
+	const FMoverDefaultSyncState* StartingSyncState = StartState.SyncState.SyncStateCollection.FindDataByType<FMoverDefaultSyncState>();
+	check(StartingSyncState);
+
+	FMoverDefaultSyncState& OutputSyncState = OutputState.SyncState.SyncStateCollection.FindOrAddMutableDataByType<FMoverDefaultSyncState>();
+
+	const float DeltaSeconds = Params.TimeStep.StepMs * 0.001f;
+	const FVector UpDir = GetMoverComponent()->GetUpDirection();
+
+	// Instantaneous movement changes that are executed and we exit before consuming any time
+	if (ProposedMove.bHasTargetLocation && AttemptTeleport(UpdatedComponent, ProposedMove.TargetLocation, UpdatedComponent->GetComponentRotation(), StartingSyncState->GetVelocity_WorldSpace(), OutputSyncState))
+	{
+		OutputState.MovementEndState.RemainingMs = Params.TimeStep.StepMs; 	// Give back all the time
+		return;
+	}
+
+	// Floor query
+
+	UMoverBlackboard* SimBlackboard = GetBlackboard_Mutable();
+	if (!SimBlackboard)
+	{
+		OutputSyncState = *StartingSyncState;
+		return;
+	}
+
+	SimBlackboard->Invalidate(KinematicBlackboard::LastFloorResult);
+	SimBlackboard->Invalidate(KinematicBlackboard::LastWaterResult);
+
+	// Find floor
+
+	float PawnHalfHeight;
+	float PawnRadius;
+	UpdatedPrimitive->CalcBoundingCylinder(PawnRadius, PawnHalfHeight);
+
+	FFloorCheckResult FloorResult;
+	FWaterCheckResult WaterResult;
+	UPhysicsMovementUtils::FindFloor(StartingSyncState->GetLocation_WorldSpace(), StartingSyncState->GetVelocity_WorldSpace() * DeltaSeconds,
+	UpdatedPrimitive, UpDir, PawnRadius, TargetHeight, CommonLegacySettings->MaxStepHeight,
+		CommonLegacySettings->MaxWalkSlopeCosine, FloorResult, WaterResult);
+
+	SimBlackboard->Set(KinematicBlackboard::LastFloorResult, FloorResult);
+	SimBlackboard->Set(KinematicBlackboard::LastWaterResult, WaterResult);
+
+	const bool bIsMovingUp = KinematicInputs->bIsJumpJustPressed || (UpDir.Dot(ProposedMove.LinearVelocity) > 0.0f);
+	const bool bStartSwimming = WaterResult.WaterSplineData.ImmersionDepth > CommonLegacySettings->SwimmingStartImmersionDepth;
+
+	if (WaterResult.IsSwimmableVolume() && bStartSwimming && !bIsMovingUp)
+	{
+		OutputState.MovementEndState.NextModeName = KinematicModeNames::Swimming;
+		OutputState.MovementEndState.RemainingMs = Params.TimeStep.StepMs;
+		return;
+	}
+
+	// In air steering
+
+	FRotator TargetOrient = StartingSyncState->GetOrientation_WorldSpace();
+	if (!ProposedMove.AngularVelocity.IsZero())
+	{
+		TargetOrient += (ProposedMove.AngularVelocity * DeltaSeconds);
+	}
+
+	FVector TargetVel = ProposedMove.LinearVelocity;
+	if (const APhysicsVolume* CurPhysVolume = UpdatedComponent->GetPhysicsVolume())
+	{
+		// The physics simulation applies Z-only gravity acceleration via physics volumes, so we need to account for it here 
+		TargetVel -= (CurPhysVolume->GetGravityZ() * FVector::UpVector * DeltaSeconds);
+	}
+
+	FVector TargetPos = StartingSyncState->GetLocation_WorldSpace() + TargetVel * DeltaSeconds;
+
+	constexpr float FloorDistanceTolerance = 2.0f;
+	const float ProjectedFloorDistance = FloorResult.FloorDist + UpDir.Dot(ProposedMove.LinearVelocity) * DeltaSeconds;
+	const bool bIsFloorWithinReach = ProjectedFloorDistance < TargetHeight + FloorDistanceTolerance;
+
+	if (FloorResult.IsWalkableFloor() && bIsFloorWithinReach)
+	{
+		OutputState.MovementEndState.NextModeName = KinematicModeNames::Walking;
+		TargetPos -= UpDir * (FloorResult.FloorDist - TargetHeight);
+	}
+	else
+	{
+		OutputState.MovementEndState.NextModeName = KinematicModeNames::Falling;
+	}
+
+	OutputState.MovementEndState.RemainingMs = 0.0f;
+	OutputSyncState.MoveDirectionIntent = ProposedMove.bHasDirIntent ? ProposedMove.DirectionIntent : FVector::ZeroVector;
+	OutputSyncState.SetTransforms_WorldSpace(
+		TargetPos,
+		TargetOrient,
+		TargetVel);
+}
\ No newline at end of file
diff --git a/Engine/Plugins/Experimental/Mover/Source/Mover/Private/PhysicsMover/Modes/PhysicsDrivenFlyingMode.cpp b/Engine/Plugins/Experimental/Mover/Source/Mover/Private/PhysicsMover/Modes/PhysicsDrivenFlyingMode.cpp
new file mode 100644
--- /dev/null
+++ b/Engine/Plugins/Experimental/Mover/Source/Mover/Private/PhysicsMover/Modes/PhysicsDrivenFlyingMode.cpp
@@ -0,0 +1,84 @@
+// Copyright Epic Games, Inc. All Rights Reserved.
+
+#include "PhysicsMover/Modes/PhysicsDrivenFlyingMode.h"
+
+#include "Chaos/Character/CharacterGroundConstraint.h"
+#include "GameFramework/PhysicsVolume.h"
+#include "MoverComponent.h"
+#include "Kinematic/Settings/CommonLegacyMovementSettings.h"
+#include "Math/UnitConversion.h"
+#include "MoveLibrary/MovementUtils.h"
+#include "PhysicsMover/PhysicsMovementUtils.h"
+
+#include UE_INLINE_GENERATED_CPP_BY_NAME(PhysicsDrivenFlyingMode)
+
+UPhysicsDrivenFlyingMode::UPhysicsDrivenFlyingMode(const FObjectInitializer& ObjectInitializer)
+	: Super(ObjectInitializer)
+{
+}
+
+void UPhysicsDrivenFlyingMode::UpdateConstraintSettings(Chaos::FCharacterGroundConstraint& Constraint) const
+{
+	Constraint.SetRadialForceLimit(300000.0); // TEMP - Move radial force limit to shared mode data
+	Constraint.SetTwistTorqueLimit(FUnitConversion::Convert(TwistTorqueLimit, EUnit::NewtonMeters, EUnit::KilogramCentimetersSquaredPerSecondSquared));
+	Constraint.SetSwingTorqueLimit(FUnitConversion::Convert(SwingTorqueLimit, EUnit::NewtonMeters, EUnit::KilogramCentimetersSquaredPerSecondSquared));
+	Constraint.SetTargetHeight(0.0f);
+}
+
+void UPhysicsDrivenFlyingMode::OnSimulationTick(const FSimulationTickParams& Params, FMoverTickEndData& OutputState)
+{
+	const FMoverTickStartData& StartState = Params.StartState;
+	USceneComponent* UpdatedComponent = Params.UpdatedComponent;
+	UPrimitiveComponent* UpdatedPrimitive = Params.UpdatedPrimitive;
+	FProposedMove ProposedMove = Params.ProposedMove;
+
+	const FKinematicDefaultInputs* KinematicInputs = StartState.InputCmd.InputCollection.FindDataByType<FKinematicDefaultInputs>();
+	const FMoverDefaultSyncState* StartingSyncState = StartState.SyncState.SyncStateCollection.FindDataByType<FMoverDefaultSyncState>();
+	check(StartingSyncState);
+
+	FMoverDefaultSyncState& OutputSyncState = OutputState.SyncState.SyncStateCollection.FindOrAddMutableDataByType<FMoverDefaultSyncState>();
+
+	const float DeltaSeconds = Params.TimeStep.StepMs * 0.001f;
+	const FVector UpDir = GetMoverComponent()->GetUpDirection();
+
+	// Instantaneous movement changes that are executed and we exit before consuming any time
+	if (ProposedMove.bHasTargetLocation && AttemptTeleport(UpdatedComponent, ProposedMove.TargetLocation, UpdatedComponent->GetComponentRotation(), *StartingSyncState, OutputState))
+	{
+		OutputState.MovementEndState.RemainingMs = Params.TimeStep.StepMs; 	// Give back all the time
+		return;
+	}
+
+	// Don't need a floor query - just invalidate the blackboard to ensure we don't use an old result elsewhere
+
+	if (UMoverBlackboard* SimBlackboard = GetBlackboard_Mutable())
+	{
+		SimBlackboard->Invalidate(KinematicBlackboard::LastFloorResult);
+		SimBlackboard->Invalidate(KinematicBlackboard::LastWaterResult);
+	}
+
+	// In air steering
+
+	FRotator TargetOrient = StartingSyncState->GetOrientation_WorldSpace();
+	if (!ProposedMove.AngularVelocity.IsZero())
+	{
+		TargetOrient += (ProposedMove.AngularVelocity * DeltaSeconds);
+	}
+
+	FVector TargetVel = ProposedMove.LinearVelocity;
+	if (const APhysicsVolume* CurPhysVolume = UpdatedComponent->GetPhysicsVolume())
+	{
+		// The physics simulation applies Z-only gravity acceleration via physics volumes, so we need to account for it here 
+		TargetVel -= (CurPhysVolume->GetGravityZ() * FVector::UpVector * DeltaSeconds);
+	}
+
+	FVector TargetPos = StartingSyncState->GetLocation_WorldSpace() + TargetVel * DeltaSeconds;
+
+	OutputState.MovementEndState.NextModeName = KinematicModeNames::Flying;
+
+	OutputState.MovementEndState.RemainingMs = 0.0f;
+	OutputSyncState.MoveDirectionIntent = ProposedMove.bHasDirIntent ? ProposedMove.DirectionIntent : FVector::ZeroVector;
+	OutputSyncState.SetTransforms_WorldSpace(
+		TargetPos,
+		TargetOrient,
+		TargetVel);
+}
\ No newline at end of file
