Moving Mover, MoverExamples, MoverTests out of Restricted engine folder

diff --git a/Engine/Plugins/Experimental/Mover/Source/Mover/Private/PhysicsMover/Modes/PhysicsDrivenSwimmingMode.cpp b/Engine/Plugins/Experimental/Mover/Source/Mover/Private/PhysicsMover/Modes/PhysicsDrivenSwimmingMode.cpp
new file mode 100644
--- /dev/null
+++ b/Engine/Plugins/Experimental/Mover/Source/Mover/Private/PhysicsMover/Modes/PhysicsDrivenSwimmingMode.cpp
@@ -0,0 +1,157 @@
+// Copyright Epic Games, Inc. All Rights Reserved.
+
+#include "PhysicsMover/Modes/PhysicsDrivenSwimmingMode.h"
+
+#include "Chaos/Character/CharacterGroundConstraint.h"
+#include "Math/UnitConversion.h"
+#include "MoveLibrary/FloorQueryUtils.h"
+#include "MoveLibrary/WaterMovementUtils.h"
+#include "Kinematic/LayeredMoves/BasicLayeredMoves.h"
+#include "MoverComponent.h"
+#include "Kinematic/Settings/CommonLegacyMovementSettings.h"
+#include "GameFramework/PhysicsVolume.h"
+#include "PhysicsMover/PhysicsMovementUtils.h"
+
+#include UE_INLINE_GENERATED_CPP_BY_NAME(PhysicsDrivenSwimmingMode)
+
+UPhysicsDrivenSwimmingMode::UPhysicsDrivenSwimmingMode(const FObjectInitializer& ObjectInitializer)
+	: Super(ObjectInitializer)
+{
+}
+
+void UPhysicsDrivenSwimmingMode::UpdateConstraintSettings(Chaos::FCharacterGroundConstraint& Constraint) const
+{
+	Constraint.SetSwingTorqueLimit(FUnitConversion::Convert(3000.0f, EUnit::NewtonMeters, EUnit::KilogramCentimetersSquaredPerSecondSquared));
+	Constraint.SetRadialForceLimit(0.0f);
+	Constraint.SetFrictionForceLimit(0.0f);
+	Constraint.SetTwistTorqueLimit(0.0f);
+}
+
+void UPhysicsDrivenSwimmingMode::OnSimulationTick(const FSimulationTickParams& Params, FMoverTickEndData& OutputState)
+{
+	const UMoverComponent* MoverComp = GetMoverComponent();
+
+	const FMoverTickStartData& StartState = Params.StartState;
+	USceneComponent* UpdatedComponent = Params.UpdatedComponent;
+	UPrimitiveComponent* UpdatedPrimitive = Params.UpdatedPrimitive;
+	FProposedMove ProposedMove = Params.ProposedMove;
+	
+	const FVector UpDir = GetMoverComponent()->GetUpDirection();
+	
+	const FKinematicDefaultInputs* KinematicInputs = StartState.InputCmd.InputCollection.FindDataByType<FKinematicDefaultInputs>();
+
+	const FMoverDefaultSyncState* StartingSyncState = StartState.SyncState.SyncStateCollection.FindDataByType<FMoverDefaultSyncState>();
+	check(StartingSyncState);
+
+	FMoverDefaultSyncState& OutputSyncState = OutputState.SyncState.SyncStateCollection.FindOrAddMutableDataByType<FMoverDefaultSyncState>();
+
+	const float DeltaSeconds = Params.TimeStep.StepMs * 0.001f;
+
+	if ((ProposedMove.bHasTargetLocation && AttemptTeleport(UpdatedComponent, ProposedMove.TargetLocation, UpdatedComponent->GetComponentRotation(), StartingSyncState->GetVelocity_WorldSpace(), OutputState)) ||	// Teleport
+		(KinematicInputs->bIsJumpJustPressed && AttemptJump(SurfaceSwimmingWaterControlSettings.JumpMultiplier*CommonLegacySettings->JumpUpwardsSpeed, OutputState)))
+	{
+		OutputState.MovementEndState.RemainingMs = Params.TimeStep.StepMs;
+		return;
+	}
+	
+	UMoverBlackboard* SimBlackboard = GetBlackboard_Mutable();
+	if (!SimBlackboard)
+	{
+		OutputSyncState = *StartingSyncState;
+		return;
+	}
+	SimBlackboard->Invalidate(KinematicBlackboard::LastFloorResult);
+	SimBlackboard->Invalidate(KinematicBlackboard::LastWaterResult);
+	
+	OutputSyncState.MoveDirectionIntent = ProposedMove.bHasDirIntent ? ProposedMove.DirectionIntent : FVector::ZeroVector;
+
+	// Floor query
+	FFloorCheckResult FloorResult;
+	FWaterCheckResult WaterResult;
+	
+	float PawnHalfHeight;
+	float PawnRadius;
+	UpdatedPrimitive->CalcBoundingCylinder(PawnRadius, PawnHalfHeight);
+	
+	UPhysicsMovementUtils::FindFloor(StartingSyncState->GetLocation_WorldSpace(), StartingSyncState->GetVelocity_WorldSpace() * DeltaSeconds,
+		UpdatedPrimitive, UpDir, PawnRadius, TargetHeight, CommonLegacySettings->MaxStepHeight,
+		CommonLegacySettings->MaxWalkSlopeCosine, FloorResult, WaterResult);
+	
+	SimBlackboard->Set(KinematicBlackboard::LastFloorResult, FloorResult);
+	SimBlackboard->Set(KinematicBlackboard::LastWaterResult, WaterResult);
+	
+	if (WaterResult.IsSwimmableVolume())
+	{
+		const bool bIsWithinReach = FloorResult.FloorDist <= TargetHeight;
+		const bool bWalkTrigger = WaterResult.WaterSplineData.ImmersionDepth < CommonLegacySettings->SwimmingStopImmersionDepth;
+		const bool bFallTrigger = FMath::Clamp((WaterResult.WaterSplineData.ImmersionDepth + TargetHeight) / (2 * TargetHeight), -2.f, 2.f) < -1.f;
+	
+		FRotator TargetOrient = StartingSyncState->GetOrientation_WorldSpace();
+		if (!ProposedMove.AngularVelocity.IsZero())
+		{
+			TargetOrient += (ProposedMove.AngularVelocity * DeltaSeconds);
+		}
+
+		FVector TargetVel = ProposedMove.LinearVelocity;		
+		if (const APhysicsVolume* CurPhysVolume = UpdatedComponent->GetPhysicsVolume())
+		{
+			// Discount G Forces as Buoyancy accounts for it
+			TargetVel -= (CurPhysVolume->GetGravityZ() * FVector::UpVector * DeltaSeconds);
+		}
+		
+		FVector TargetPos = StartingSyncState->GetLocation_WorldSpace();
+		TargetPos += TargetVel * DeltaSeconds;
+	
+		OutputSyncState.SetTransforms_WorldSpace(
+			TargetPos,
+			TargetOrient,
+			TargetVel);
+	
+		if (bWalkTrigger && bIsWithinReach)
+		{
+			OutputState.MovementEndState.NextModeName = KinematicModeNames::Walking;
+		}
+		else if (bFallTrigger)
+		{
+			OutputState.MovementEndState.NextModeName = KinematicModeNames::Falling;
+		}
+		else
+		{
+			OutputState.MovementEndState.NextModeName = KinematicModeNames::Swimming;
+		}
+	}
+	else
+	{
+		OutputState.MovementEndState.NextModeName = KinematicModeNames::Falling;
+	}
+	
+	OutputState.MovementEndState.RemainingMs = 0.0f;
+}
+
+bool UPhysicsDrivenSwimmingMode::AttemptJump(float UpwardsSpeed, FMoverTickEndData& OutputState)
+{
+	FMoverDefaultSyncState& OutputSyncState = OutputState.SyncState.SyncStateCollection.FindOrAddMutableDataByType<FMoverDefaultSyncState>();
+
+	// TODO: This should check if a jump is even allowed
+ 	TSharedPtr<FLayeredMove_JumpImpulse> JumpMove = MakeShared<FLayeredMove_JumpImpulse>();
+	JumpMove->UpwardsSpeed = UpwardsSpeed;
+	OutputSyncState.LayeredMoves.QueueLayeredMove(JumpMove);
+	OutputState.MovementEndState.NextModeName = KinematicModeNames::Falling;
+	return true;
+}
+
+bool UPhysicsDrivenSwimmingMode::AttemptTeleport(USceneComponent* UpdatedComponent, const FVector& TeleportPos, const FRotator& TeleportRot, const FVector& PriorVelocity, FMoverTickEndData& Output)
+{
+	FMoverDefaultSyncState& OutputSyncState = Output.SyncState.SyncStateCollection.FindOrAddMutableDataByType<FMoverDefaultSyncState>();
+
+	OutputSyncState.SetTransforms_WorldSpace(TeleportPos,
+		TeleportRot,
+		PriorVelocity,
+		nullptr); // no movement base
+
+	// TODO: instead of invalidating it, consider checking for a floor. Possibly a dynamic base?
+	GetBlackboard_Mutable()->Invalidate(KinematicBlackboard::LastFloorResult);
+	GetBlackboard_Mutable()->Invalidate(KinematicBlackboard::LastMovementBase);
+
+	return true;
+}
\ No newline at end of file
diff --git a/Engine/Plugins/Experimental/Mover/Source/Mover/Private/PhysicsMover/Modes/PhysicsDrivenWalkingMode.cpp b/Engine/Plugins/Experimental/Mover/Source/Mover/Private/PhysicsMover/Modes/PhysicsDrivenWalkingMode.cpp
new file mode 100644
--- /dev/null
+++ b/Engine/Plugins/Experimental/Mover/Source/Mover/Private/PhysicsMover/Modes/PhysicsDrivenWalkingMode.cpp
@@ -0,0 +1,243 @@
+// Copyright Epic Games, Inc. All Rights Reserved.
+
+#include "PhysicsMover/Modes/PhysicsDrivenWalkingMode.h"
+
+#include "Chaos/Character/CharacterGroundConstraint.h"
+#include "Chaos/PhysicsObject.h"
+#include "Chaos/PhysicsObjectInternalInterface.h"
+#include "GameFramework/PhysicsVolume.h"
+#include "Kinematic/LayeredMoves/BasicLayeredMoves.h"
+#include "MoverComponent.h"
+#include "Kinematic/Settings/CommonLegacyMovementSettings.h"
+#include "Math/UnitConversion.h"
+#include "MoveLibrary/WaterMovementUtils.h"
+#include "PhysicsMover/PhysicsMovementUtils.h"
+#include "PhysicsMover/PhysicsMoverSimulationTypes.h"
+
+#include UE_INLINE_GENERATED_CPP_BY_NAME(PhysicsDrivenWalkingMode)
+
+extern FPhysicsDrivenMotionDebugParams GPhysicsDrivenMotionDebugParams;
+
+UPhysicsDrivenWalkingMode::UPhysicsDrivenWalkingMode(const FObjectInitializer& ObjectInitializer)
+	: Super(ObjectInitializer)
+{
+}
+
+void UPhysicsDrivenWalkingMode::UpdateConstraintSettings(Chaos::FCharacterGroundConstraint& Constraint) const
+{
+	Constraint.SetRadialForceLimit(FUnitConversion::Convert(RadialForceLimit, EUnit::Newtons, EUnit::KilogramCentimetersPerSecondSquared));
+	Constraint.SetFrictionForceLimit(FUnitConversion::Convert(FrictionForceLimit, EUnit::Newtons, EUnit::KilogramCentimetersPerSecondSquared));
+	Constraint.SetTwistTorqueLimit(FUnitConversion::Convert(TwistTorqueLimit, EUnit::NewtonMeters, EUnit::KilogramCentimetersSquaredPerSecondSquared));
+	Constraint.SetSwingTorqueLimit(FUnitConversion::Convert(SwingTorqueLimit, EUnit::NewtonMeters, EUnit::KilogramCentimetersSquaredPerSecondSquared));
+	Constraint.SetTargetHeight(TargetHeight);
+}
+
+void UPhysicsDrivenWalkingMode::OnSimulationTick(const FSimulationTickParams& Params, FMoverTickEndData& OutputState)
+{
+	const UMoverComponent* MoverComp = GetMoverComponent();
+	const FMoverTickStartData& StartState = Params.StartState;
+	USceneComponent* UpdatedComponent = Params.UpdatedComponent;
+	UPrimitiveComponent* UpdatedPrimitive = Params.UpdatedPrimitive;
+	FProposedMove ProposedMove = Params.ProposedMove;
+
+	const FVector UpDir = GetMoverComponent()->GetUpDirection();
+
+	const FKinematicDefaultInputs* KinematicInputs = StartState.InputCmd.InputCollection.FindDataByType<FKinematicDefaultInputs>();
+	const FMoverDefaultSyncState* StartingSyncState = StartState.SyncState.SyncStateCollection.FindDataByType<FMoverDefaultSyncState>();
+	check(StartingSyncState);
+
+	FMoverDefaultSyncState& OutputSyncState = OutputState.SyncState.SyncStateCollection.FindOrAddMutableDataByType<FMoverDefaultSyncState>();
+
+	const float DeltaSeconds = Params.TimeStep.StepMs * 0.001f;
+
+	// Instantaneous movement changes that are executed and we exit before consuming any time
+	if ((ProposedMove.bHasTargetLocation && AttemptTeleport(UpdatedComponent, ProposedMove.TargetLocation, UpdatedComponent->GetComponentRotation(), StartingSyncState->GetVelocity_WorldSpace(), OutputState)))
+	{
+		OutputState.MovementEndState.RemainingMs = Params.TimeStep.StepMs; 	// Give back all the time
+		return;
+	}
+
+	UMoverBlackboard* SimBlackboard = GetBlackboard_Mutable();
+	if (!SimBlackboard)
+	{
+		OutputSyncState = *StartingSyncState;
+		return;
+	}
+
+	// Store the previous ground normal that was used to compute the proposed move
+	FFloorCheckResult PrevFloorResult;
+	FVector PrevGroundNormal = UpDir;
+	if (SimBlackboard->TryGet(KinematicBlackboard::LastFloorResult, PrevFloorResult))
+	{
+		PrevGroundNormal = PrevFloorResult.HitResult.ImpactNormal;
+	}
+
+	// Floor query
+	float PawnHalfHeight;
+	float PawnRadius;
+	UpdatedPrimitive->CalcBoundingCylinder(PawnRadius, PawnHalfHeight);
+
+	FFloorCheckResult FloorResult;
+	FWaterCheckResult WaterResult;
+	UPhysicsMovementUtils::FindFloor(StartingSyncState->GetLocation_WorldSpace(), StartingSyncState->GetVelocity_WorldSpace() * DeltaSeconds,
+		UpdatedPrimitive, UpDir, PawnRadius, TargetHeight, CommonLegacySettings->MaxStepHeight,
+		CommonLegacySettings->MaxWalkSlopeCosine, FloorResult, WaterResult);
+
+	SimBlackboard->Set(KinematicBlackboard::LastFloorResult, FloorResult);
+	SimBlackboard->Set(KinematicBlackboard::LastWaterResult, WaterResult);
+
+	const bool bStartSwimming = (WaterResult.WaterSplineData.ImmersionDepth+TargetHeight) > CommonLegacySettings->SwimmingIdealImmersionDepth;
+	
+	if (WaterResult.IsSwimmableVolume() && bStartSwimming)
+	{
+		SwitchToState(KinematicModeNames::Swimming, Params, OutputState);
+	}
+	else if (FloorResult.IsWalkableFloor())
+	{
+		const FVector StartGroundVelocity = UPhysicsMovementUtils::ComputeGroundVelocityFromHitResult(StartingSyncState->GetLocation_WorldSpace(), FloorResult.HitResult, DeltaSeconds);
+		
+		FVector TargetVelocity = StartingSyncState->GetVelocity_WorldSpace();
+		FVector TargetPosition = StartingSyncState->GetLocation_WorldSpace();
+		if (FloorResult.bWalkableFloor)
+		{
+			const FVector ProposedMovePlaneVelocity = ProposedMove.LinearVelocity - ProposedMove.LinearVelocity.ProjectOnToNormal(PrevGroundNormal);
+			const FVector StartingMovePlaneVelocity = StartingSyncState->GetVelocity_WorldSpace() - StartingSyncState->GetVelocity_WorldSpace().ProjectOnToNormal(PrevGroundNormal);
+
+			// If there is velocity intent in the normal direction then use the velocity from the proposed move. Otherwise
+			// retain the previous vertical velocity
+			FVector ProposedNormalVelocity = ProposedMove.LinearVelocity - ProposedMovePlaneVelocity;
+			if (ProposedNormalVelocity.SizeSquared() <= UE_SMALL_NUMBER)
+			{
+				ProposedNormalVelocity = StartingSyncState->GetVelocity_WorldSpace() - StartingMovePlaneVelocity;
+			}
+
+			TargetVelocity = FMath::Lerp(StartingMovePlaneVelocity, ProposedMovePlaneVelocity, FractionalVelocityToTarget) + ProposedNormalVelocity;
+			TargetPosition += ProposedMovePlaneVelocity * DeltaSeconds;
+		}
+
+		// Check if the proposed velocity would lift off the movement surface.
+		bool bIsLiftingOffSurface = false;
+
+		float CharacterGravity = 0.0f;
+		if (const APhysicsVolume* PhysVolume = UpdatedComponent->GetPhysicsVolume())
+		{
+			CharacterGravity = PhysVolume->GetGravityZ();
+		}
+		const FVector ProjectedVelocity = TargetVelocity + CharacterGravity * UpDir * DeltaSeconds;
+
+		float GroundGravity = 0.0f;
+		if (const Chaos::FPBDRigidParticleHandle* Rigid = UPhysicsMovementUtils::GetRigidParticelHandleFromHitResult(FloorResult.HitResult))
+		{
+			if (Rigid->IsDynamic() && Rigid->GravityEnabled())
+			{
+				if (const UPrimitiveComponent* GroundComp = FloorResult.HitResult.GetComponent())
+				{
+					if (const APhysicsVolume* PhysVolume = GroundComp->GetPhysicsVolume())
+					{
+						GroundGravity = PhysVolume->GetGravityZ();
+					}
+				}
+			}
+		}
+		const FVector ProjectedGroundVelocity = StartGroundVelocity + GroundGravity * UpDir * DeltaSeconds;
+
+		const float ProjectedRelativeVerticalVelocity = FloorResult.HitResult.ImpactNormal.Dot(ProjectedVelocity - ProjectedGroundVelocity);
+		const float VerticalVelocityLimit = 2.0f / DeltaSeconds;
+		if (ProjectedRelativeVerticalVelocity > VerticalVelocityLimit)
+		{
+			bIsLiftingOffSurface = true;
+		}
+
+		// Determine if the character is stepping up or stepping down.
+		// If stepping up make sure that the step height is less than the max step height
+		// and the new surface has CanCharacterStepUpOn set to true.
+		// If stepping down make sure the step height is less than the max step height.
+		const float StartHeightAboveGround = FloorResult.FloorDist - TargetHeight;
+		const float EndHeightAboutGround = StartHeightAboveGround + UpDir.Dot(ProjectedVelocity - ProjectedGroundVelocity) * DeltaSeconds ;
+		const bool bIsSteppingDown = StartHeightAboveGround > GPhysicsDrivenMotionDebugParams.MinStepUpDistance;
+		const bool bIsWithinReach = EndHeightAboutGround <= CommonLegacySettings->MaxStepHeight;
+
+		// If the character is unsupported allow some grace period before falling
+		bool bIsSupported = bIsWithinReach && !bIsLiftingOffSurface;
+		float TimeSinceSupported = MaxUnsupportedTimeBeforeFalling;
+		SimBlackboard->TryGet(KinematicBlackboard::TimeSinceSupported, TimeSinceSupported);
+		if (bIsSupported)
+		{
+			SimBlackboard->Set(KinematicBlackboard::TimeSinceSupported, 0.0f);
+		}
+		else
+		{
+			TimeSinceSupported += DeltaSeconds;
+			SimBlackboard->Set(KinematicBlackboard::TimeSinceSupported, TimeSinceSupported);
+			bIsSupported = TimeSinceSupported < MaxUnsupportedTimeBeforeFalling;
+		}
+
+		// Apply vertical velocity to target if stepping down
+		const bool bNeedsVerticalVelocityToTarget = bIsSupported && bIsSteppingDown && (EndHeightAboutGround > 0.0f) && !bIsLiftingOffSurface;
+		if (bNeedsVerticalVelocityToTarget)
+		{
+			TargetVelocity -= FractionalDownwardVelocityToTarget * (EndHeightAboutGround / DeltaSeconds) * UpDir;
+		}
+
+		// Target orientation
+		// This is always applied regardless of whether the character is supported
+		FRotator TargetOrientation = StartingSyncState->GetOrientation_WorldSpace();
+		if (!ProposedMove.AngularVelocity.IsZero())
+		{
+			TargetOrientation += (ProposedMove.AngularVelocity * DeltaSeconds);
+		}
+
+		if (bIsSupported)
+		{
+			OutputState.MovementEndState.NextModeName = KinematicModeNames::Walking;
+			OutputState.MovementEndState.RemainingMs = 0.0f;
+
+			OutputSyncState.MoveDirectionIntent = ProposedMove.bHasDirIntent ? ProposedMove.DirectionIntent : FVector::ZeroVector;
+			OutputSyncState.SetTransforms_WorldSpace(
+				TargetPosition,
+				TargetOrientation,
+				TargetVelocity,
+				nullptr);
+		}
+		else
+		{
+			// Blocking hit but not supported
+			SwitchToState(KinematicModeNames::Falling, Params, OutputState);
+		}
+	}
+	else
+	{
+		// No water or floor not found
+		SwitchToState(KinematicModeNames::Falling, Params, OutputState);
+	}
+}
+
+bool UPhysicsDrivenWalkingMode::AttemptTeleport(USceneComponent* UpdatedComponent, const FVector& TeleportPos, const FRotator& TeleportRot, const FVector& PriorVelocity, FMoverTickEndData& Output)
+{
+	FMoverDefaultSyncState& OutputSyncState = Output.SyncState.SyncStateCollection.FindOrAddMutableDataByType<FMoverDefaultSyncState>();
+
+	OutputSyncState.SetTransforms_WorldSpace(TeleportPos,
+		TeleportRot,
+		PriorVelocity,
+		nullptr); // no movement base
+
+	// TODO: instead of invalidating it, consider checking for a floor. Possibly a dynamic base?
+	GetBlackboard_Mutable()->Invalidate(KinematicBlackboard::LastFloorResult);
+	GetBlackboard_Mutable()->Invalidate(KinematicBlackboard::LastMovementBase);
+
+	return true;
+}
+
+void UPhysicsDrivenWalkingMode::SwitchToState(const FName& StateName, const FSimulationTickParams& Params, FMoverTickEndData& OutputState)
+{
+	OutputState.MovementEndState.RemainingMs = Params.TimeStep.StepMs;
+	OutputState.MovementEndState.NextModeName = StateName;
+
+	const FMoverDefaultSyncState* StartingSyncState = Params.StartState.SyncState.SyncStateCollection.FindDataByType<FMoverDefaultSyncState>();
+	FMoverDefaultSyncState& OutputSyncState = OutputState.SyncState.SyncStateCollection.FindOrAddMutableDataByType<FMoverDefaultSyncState>();
+	OutputSyncState.SetTransforms_WorldSpace(
+		StartingSyncState->GetLocation_WorldSpace(),
+		StartingSyncState->GetOrientation_WorldSpace(),
+		StartingSyncState->GetVelocity_WorldSpace(),
+		nullptr);
+}
\ No newline at end of file
diff --git a/Engine/Plugins/Experimental/Mover/Source/Mover/Private/PhysicsMover/PhysicsMovementUtils.cpp b/Engine/Plugins/Experimental/Mover/Source/Mover/Private/PhysicsMover/PhysicsMovementUtils.cpp
new file mode 100644
--- /dev/null
+++ b/Engine/Plugins/Experimental/Mover/Source/Mover/Private/PhysicsMover/PhysicsMovementUtils.cpp
@@ -0,0 +1,302 @@
+// Copyright Epic Games, Inc. All Rights Reserved.
+
+#include "PhysicsMover/PhysicsMovementUtils.h"
+#include "Chaos/PhysicsObject.h"
+#include "Chaos/PhysicsObjectInternalInterface.h"
+#include "Components/PrimitiveComponent.h"
+#include "Engine/World.h"
+#include "MoveLibrary/FloorQueryUtils.h"
+#include "MoveLibrary/GroundMovementUtils.h"
+#include "MoveLibrary/WaterMovementUtils.h"
+#include "PhysicsMover/PhysicsMoverSimulationTypes.h"
+#include "WaterBodyActor.h"
+
+#if PHYSICSDRIVENMOTION_DEBUG_DRAW
+#include "Chaos/DebugDrawQueue.h"
+#endif
+
+extern FPhysicsDrivenMotionDebugParams GPhysicsDrivenMotionDebugParams;
+
+void UPhysicsMovementUtils::FindFloor(
+	const FVector& Location,
+	const FVector& DeltaPos,
+	const UPrimitiveComponent* UpdatedPrimitive,
+	const FVector& UpDir,
+	float QueryRadius,
+	float TargetHeight,
+	float MaxStepHeight,
+	float MaxWalkSlopeCosine,
+	FFloorCheckResult& OutFloorResult,
+	FWaterCheckResult& OutWaterResult)
+{
+	if (const UWorld* World = UpdatedPrimitive->GetWorld())
+	{
+		FCollisionQueryParams QueryParams(SCENE_QUERY_STAT(PhysicsFloorTest), false, UpdatedPrimitive->GetOwner());
+		QueryParams.bTraceIntoSubComponents = false;
+		const ECollisionChannel CollisionChannel = UpdatedPrimitive->GetCollisionObjectType();
+		FCollisionResponseParams ResponseParams(ECR_Overlap);
+		ResponseParams.CollisionResponse.SetResponse(ECC_WorldStatic, ECR_Block);
+		ResponseParams.CollisionResponse.SetResponse(ECC_WorldDynamic, ECR_Block);
+		ResponseParams.CollisionResponse.SetResponse(ECC_Vehicle, ECR_Block);
+		ResponseParams.CollisionResponse.SetResponse(ECC_Destructible, ECR_Block);
+		ResponseParams.CollisionResponse.SetResponse(ECC_PhysicsBody, ECR_Block);
+
+		const FVector DeltaPosVert = DeltaPos.ProjectOnTo(UpDir);
+		const FVector DeltaPosHoriz = DeltaPos - DeltaPosVert;
+
+		const float StartOffset = TargetHeight - MaxStepHeight - QueryRadius;
+		const float SweepDistance = MaxStepHeight + QueryRadius + FMath::Max(MaxStepHeight, -DeltaPosVert.Dot(UpDir));
+
+		FVector Start = Location + DeltaPosHoriz - StartOffset * UpDir;
+		FVector End = Start - SweepDistance * UpDir;
+		TArray<FHitResult> Hits;
+		FHitResult OutHit;
+		if (World->SweepMultiByChannel(Hits, Start, End, FQuat::Identity, CollisionChannel, FCollisionShape::MakeSphere(QueryRadius), QueryParams, ResponseParams))
+		{
+			OutHit = Hits.Last();
+		}
+
+#if PHYSICSDRIVENMOTION_DEBUG_DRAW
+		// Draw full length of query
+		if (GPhysicsDrivenMotionDebugParams.DebugDrawGroundQueries)
+		{
+			const FVector Center = 0.5f * (Start + End);
+			const float Dist = (Start - End).Size();
+			Chaos::FDebugDrawQueue::GetInstance().DrawDebugCapsule(Center, 0.5f * Dist + QueryRadius, QueryRadius, FQuat::Identity, FColor::Silver, false, -1.f, 10, 1.0f);
+		}
+#endif
+
+		if (OutHit.bBlockingHit)
+		{
+			const float Distance = UpDir.Dot(Location - OutHit.ImpactPoint);
+			const float StepHeight = TargetHeight - Distance;
+			const float MinStepHeight = 5.0f;
+
+			bool bWalkable = IsHitSurfaceWalkableWithStepUpCheck(OutHit, StepHeight, MaxStepHeight, MinStepHeight, MaxWalkSlopeCosine);
+
+#if PHYSICSDRIVENMOTION_DEBUG_DRAW
+			if (GPhysicsDrivenMotionDebugParams.DebugDrawGroundQueries)
+			{
+				const FVector Center = Start - 0.5f * OutHit.Distance * UpDir;
+				const FColor Color = bWalkable ? FColor::Green : FColor::Red;
+				Chaos::FDebugDrawQueue::GetInstance().DrawDebugCapsule(Center, 0.5f * OutHit.Distance + QueryRadius, QueryRadius, FQuat::Identity, Color, false, -1.f, 10, 1.0f);
+			}
+#endif
+
+			GetWaterResultFromHitResults(Hits, Location, OutWaterResult);
+
+			if (bWalkable)
+			{
+				// Found a walkable surface
+				OutFloorResult.bBlockingHit = true;
+				OutFloorResult.bWalkableFloor = true;
+				OutFloorResult.FloorDist = UpDir.Dot(Location - OutHit.ImpactPoint);;
+				OutFloorResult.HitResult = OutHit;
+			}
+			else
+			{
+				// Hit something but not walkable. Fire some line queries to try and find a walkable surface
+
+				bool bLineQueryWalkable = false;
+				float LineQueryDistance = 0.0f;
+
+				FVector MovementDir = FVector::ForwardVector;
+				FVector PerpDir = FVector::RightVector;
+				FVector Fwd = DeltaPos - DeltaPos.Dot(UpDir) * UpDir;
+				const float FwdSizeSq = Fwd.SizeSquared();
+				if (FwdSizeSq > UE_SMALL_NUMBER)
+				{
+					MovementDir = Fwd * FMath::InvSqrtEst(FwdSizeSq);
+					PerpDir = UpDir.Cross(MovementDir);
+				}
+
+				const int NumLineTraces = 4;
+				FVector StartPoints[NumLineTraces];
+
+				StartPoints[0] = Start + QueryRadius * MovementDir;
+				StartPoints[1] = Start - QueryRadius * MovementDir;
+				StartPoints[2] = Start - QueryRadius * PerpDir;
+				StartPoints[3] = Start + QueryRadius * PerpDir;
+
+				FHitResult LineQueryHit[NumLineTraces];
+
+				float MinDistance = SweepDistance + QueryRadius;
+				int32 MinDistanceIdx = INDEX_NONE;
+				for (int32 Idx = 0; Idx < NumLineTraces; ++Idx)
+				{
+					Start = StartPoints[Idx];
+					End = Start - UpDir * SweepDistance;
+					World->LineTraceSingleByChannel(LineQueryHit[Idx], Start, End, CollisionChannel, QueryParams);
+
+					if (LineQueryHit[Idx].IsValidBlockingHit())
+					{
+						LineQueryDistance = UpDir.Dot(Location - LineQueryHit[Idx].ImpactPoint);
+						const float LineQueryStepHeight = TargetHeight - LineQueryDistance;
+
+						bLineQueryWalkable = IsHitSurfaceWalkableWithStepUpCheck(LineQueryHit[Idx], LineQueryStepHeight, MaxStepHeight, MinStepHeight, MaxWalkSlopeCosine);
+
+						if (bLineQueryWalkable && LineQueryDistance < MinDistance)
+						{
+							MinDistance = LineQueryDistance;
+							MinDistanceIdx = Idx;
+						}
+					}
+
+#if PHYSICSDRIVENMOTION_DEBUG_DRAW
+					if (GPhysicsDrivenMotionDebugParams.DebugDrawGroundQueries)
+					{
+						FVector Dir = End - Start;
+						const float SizeSq = Dir.SizeSquared();
+						if (SizeSq > UE_SMALL_NUMBER)
+						{
+							Dir *= FMath::InvSqrt(SizeSq);
+						}
+
+						if (LineQueryHit[Idx].bBlockingHit)
+						{
+							const FVector Center = Start + 0.5f * LineQueryHit[Idx].Distance * Dir;
+							const FColor Color = bLineQueryWalkable ? FColor::Green : FColor::Red;
+							Chaos::FDebugDrawQueue::GetInstance().DrawDebugLine(Start, LineQueryHit[Idx].ImpactPoint, Color, false, -1.f, 10, 1.0f);
+						}
+						else
+						{
+							const FVector Center = 0.5f * (Start + End);
+							float Dist = SizeSq > UE_SMALL_NUMBER ? FMath::Sqrt(SizeSq) : 0.0f;
+							Chaos::FDebugDrawQueue::GetInstance().DrawDebugLine(Start, End, FColor::Silver, false, -1.f, 10, 1.0f);
+						}
+					}
+#endif
+				}
+
+				if (MinDistanceIdx != INDEX_NONE)
+				{
+					// Found a walkable surface
+					OutFloorResult.bBlockingHit = true;
+					OutFloorResult.bWalkableFloor = true;
+					OutFloorResult.FloorDist = MinDistance;
+					OutFloorResult.HitResult = LineQueryHit[MinDistanceIdx];
+				}
+				else
+				{
+					// Didn't find a walkable surface. Use original sphere query result
+					OutFloorResult.bBlockingHit = true;
+					OutFloorResult.bWalkableFloor = false;
+					OutFloorResult.FloorDist = Distance;
+					OutFloorResult.HitResult = OutHit;
+				}
+
+				// May have missed a water result due to the original query blocking on a non-walkable surface
+				// so run again with no blocking results
+				ResponseParams.CollisionResponse.SetAllChannels(ECR_Overlap);
+				World->SweepMultiByChannel(Hits, Start, End, FQuat::Identity, CollisionChannel, FCollisionShape::MakeSphere(QueryRadius), QueryParams, ResponseParams);
+				GetWaterResultFromHitResults(Hits, Location, OutWaterResult);
+			}
+		}
+		else
+		{
+			// Sweep didn't hit anything blocking
+			OutFloorResult.Clear();
+			OutFloorResult.FloorDist = 1.0e10f;
+
+			// In deep waters i.e., WaterDepth > SweepDistance, we may hit only water and no blocking surfaces
+			GetWaterResultFromHitResults(Hits, Location, OutWaterResult);
+		}
+	}
+}
+
+bool UPhysicsMovementUtils::IsHitSurfaceWalkableWithStepUpCheck(const FHitResult& Hit, float StepHeight, float MaxStepHeight, float MinStepUpHeight, float MaxWalkSlopeCosine)
+{
+	bool bWalkable = false;
+	if (StepHeight <= MaxStepHeight)
+	{
+		bWalkable = UFloorQueryUtils::IsHitSurfaceWalkable(Hit, MaxWalkSlopeCosine);
+		if (bWalkable)
+		{
+			const float MinStepHeight = 1.0f;
+			const bool bSteppingUp = StepHeight > MinStepHeight;
+			if (bSteppingUp && !UGroundMovementUtils::CanStepUpOnHitSurface(Hit))
+			{
+				bWalkable = false;
+			}
+		}
+	}
+
+	return bWalkable;
+}
+
+const Chaos::FPBDRigidParticleHandle* UPhysicsMovementUtils::GetRigidParticelHandleFromHitResult(const FHitResult& HitResult)
+{
+	if (IPhysicsComponent* PhysicsComp = Cast<IPhysicsComponent>(HitResult.Component))
+	{
+		if (Chaos::FPhysicsObjectHandle PhysicsObject = PhysicsComp->GetPhysicsObjectById(HitResult.Item))
+		{
+			Chaos::FReadPhysicsObjectInterface_Internal Interface = Chaos::FPhysicsObjectInternalInterface::GetRead();
+			if (const Chaos::FGeometryParticleHandle* Particle = Interface.GetParticle(PhysicsObject))
+			{
+				return Particle->CastToRigidParticle();
+			}
+		}
+	}
+
+	return nullptr;
+}
+
+FVector UPhysicsMovementUtils::ComputeGroundVelocityFromHitResult(const FVector& CharacterPosition, const FHitResult& FloorHit, const float DeltaSeconds)
+{
+	FVector GroundVelocity = FVector::ZeroVector;
+	if (const Chaos::FPBDRigidParticleHandle* Rigid = GetRigidParticelHandleFromHitResult(FloorHit))
+	{
+		FVector Offset = CharacterPosition - Rigid->X();
+		Offset -= Offset.ProjectOnToNormal(FloorHit.ImpactNormal);
+
+		if (Rigid->KinematicTarget().IsSet())
+		{
+			const FVector LinearDisplacement = Rigid->KinematicTarget().GetTargetPosition() - Rigid->X();
+			const FQuat RelativeQuat = Rigid->R().Inverse() * Rigid->KinematicTarget().GetTargetRotation();
+			const FVector AngularDisplacement = RelativeQuat.ToRotationVector();
+			GroundVelocity = (LinearDisplacement + AngularDisplacement.Cross(Offset)) / DeltaSeconds;
+		}
+		else
+		{
+			GroundVelocity = Rigid->V() + Rigid->W().Cross(Offset);
+		}
+	}
+	return GroundVelocity;
+}
+
+bool UPhysicsMovementUtils::GetWaterResultFromHitResults(const TArray<FHitResult>& Hits, const FVector& Location, FWaterCheckResult& OutWaterResult)
+{
+	// Find the closet hit that is a water body
+	// Note: Relies on ordering of hit results
+	for (int32 Idx = 0; Idx < Hits.Num(); ++Idx)
+	{
+		const FHitResult& Hit = Hits[Idx];
+
+		if (Hit.Component.IsValid())
+		{
+			if (const AActor* Actor = Hit.Component->GetOwner())
+			{
+				if (Actor->IsA(AWaterBody::StaticClass()))
+				{
+					OutWaterResult.HitResult = Hit;
+					OutWaterResult.bSwimmableVolume = true;
+
+					FWaterBodyQueryResult QueryResult = Cast<AWaterBody>(Actor)->GetWaterBodyComponent()->QueryWaterInfoClosestToWorldLocation(
+						Location, EWaterBodyQueryFlags::ComputeDepth | EWaterBodyQueryFlags::ComputeVelocity | EWaterBodyQueryFlags::ComputeImmersionDepth);
+
+					OutWaterResult.WaterSplineData.WaterDepth = QueryResult.GetWaterSurfaceDepth();
+					OutWaterResult.WaterSplineData.RawWaterVelocity = QueryResult.GetVelocity();
+					OutWaterResult.WaterSplineData.ImmersionDepth = QueryResult.GetImmersionDepth();
+
+#if PHYSICSDRIVENMOTION_DEBUG_DRAW
+					Chaos::FDebugDrawQueue::GetInstance().DrawDebugLine(Location, Location - FVector::UpVector * OutWaterResult.WaterSplineData.ImmersionDepth, FColor::Blue, false, -1.f, 10, 1.0f);
+					Chaos::FDebugDrawQueue::GetInstance().DrawDebugPoint(OutWaterResult.HitResult.Location, FColor::Blue, false, -1.f, 10, 1.0f);
+#endif
+					return true;
+				}
+			}
+		}
+	}
+
+	return false;
+}
\ No newline at end of file
diff --git a/Engine/Plugins/Experimental/Mover/Source/Mover/Private/PhysicsMover/PhysicsMoverManager.cpp b/Engine/Plugins/Experimental/Mover/Source/Mover/Private/PhysicsMover/PhysicsMoverManager.cpp
new file mode 100644
--- /dev/null
+++ b/Engine/Plugins/Experimental/Mover/Source/Mover/Private/PhysicsMover/PhysicsMoverManager.cpp
@@ -0,0 +1,100 @@
+// Copyright Epic Games, Inc. All Rights Reserved.
+
+#include "PhysicsMover/PhysicsMoverManager.h"
+
+#include "Backends/MoverNetworkPhysicsLiaison.h"
+#include "Chaos/Character/CharacterGroundConstraint.h"
+#include "PBDRigidsSolver.h"
+#include "Physics/Experimental/PhysScene_Chaos.h"
+#include "PhysicsMover/PhysicsMoverManagerAsyncCallback.h"
+#include "PhysicsMover/PhysicsMoverSimulationTypes.h"
+
+#include UE_INLINE_GENERATED_CPP_BY_NAME(PhysicsMoverManager)
+
+//////////////////////////////////////////////////////////////////////////
+
+void UPhysicsMoverManager::OnWorldBeginPlay(UWorld& InWorld)
+{
+	Super::OnWorldBeginPlay(InWorld);
+
+	if (FPhysScene* PhysScene = InWorld.GetPhysicsScene())
+	{
+		PhysScenePreTickCallbackHandle = PhysScene->OnPhysScenePreTick.AddUObject(this, &ThisClass::PrePhysicsUpdate);
+		AsyncCallback = PhysScene->GetSolver()->CreateAndRegisterSimCallbackObject_External<FPhysicsMoverManagerAsyncCallback>();
+	}
+}
+
+void UPhysicsMoverManager::Deinitialize()
+{
+	if (UWorld* World = GetWorld())
+	{
+		if (FPhysScene* PhysScene = World->GetPhysicsScene())
+		{
+			PhysScene->OnPhysScenePreTick.Remove(PhysScenePreTickCallbackHandle);
+			if (AsyncCallback)
+			{
+				PhysScene->GetSolver()->UnregisterAndFreeSimCallbackObject_External(AsyncCallback);
+			}
+
+			PhysicsMoverComponents.Empty();
+		}
+	}
+
+	Super::Deinitialize();
+}
+
+bool UPhysicsMoverManager::DoesSupportWorldType(const EWorldType::Type WorldType) const
+{
+	return WorldType == EWorldType::Game || WorldType == EWorldType::PIE;
+}
+
+//////////////////////////////////////////////////////////////////////////
+
+void UPhysicsMoverManager::RegisterPhysicsMoverComponent(TWeakObjectPtr<UMoverNetworkPhysicsLiaisonComponent> InPhysicsMoverComp)
+{
+	PhysicsMoverComponents.AddUnique(InPhysicsMoverComp);
+}
+
+void UPhysicsMoverManager::UnregisterPhysicsMoverComponent(TWeakObjectPtr<UMoverNetworkPhysicsLiaisonComponent> InPhysicsMoverComp)
+{
+	PhysicsMoverComponents.Remove(InPhysicsMoverComp);
+}
+
+//////////////////////////////////////////////////////////////////////////
+
+void UPhysicsMoverManager::PrePhysicsUpdate(FPhysScene* PhysScene, float DeltaTime)
+{
+	// Clear invalid data before we attempt to use it.
+	PhysicsMoverComponents.RemoveAllSwap([](const TWeakObjectPtr<UMoverNetworkPhysicsLiaisonComponent> MoverComp) {
+		return (MoverComp.IsValid() == false) || (MoverComp->GetUniqueIdx().IsValid() == false);
+		});
+
+	while (Chaos::TSimCallbackOutputHandle<FPhysicsMoverManagerAsyncOutput> ManagerAsyncOutput = AsyncCallback->PopFutureOutputData_External())
+	{
+		for (TWeakObjectPtr<UMoverNetworkPhysicsLiaisonComponent> PhysicsMoverComp : PhysicsMoverComponents)
+		{
+			Chaos::FUniqueIdx Idx = PhysicsMoverComp->GetUniqueIdx();
+			if (Idx.IsValid())
+			{
+				if (TUniquePtr<FPhysicsMoverAsyncOutput>* OutputData = ManagerAsyncOutput->PhysicsMoverToAsyncOutput.Find(Idx))
+				{
+					PhysicsMoverComp->GameThread_ConsumeOutput(**OutputData);
+				}
+			}
+		}
+	}
+
+	FPhysicsMoverManagerAsyncInput* ManagerAsyncInput = AsyncCallback->GetProducerInputData_External();
+	ManagerAsyncInput->Reset();
+	ManagerAsyncInput->AsyncInput.Reserve(PhysicsMoverComponents.Num());
+
+	if (PhysScene && PhysScene->GetOwningWorld())
+	{
+		for (TWeakObjectPtr<UMoverNetworkPhysicsLiaisonComponent> PhysicsMoverComp : PhysicsMoverComponents)
+		{
+			TUniquePtr<FPhysicsMoverAsyncInput> InputData = MakeUnique<FPhysicsMoverAsyncInput>();
+			PhysicsMoverComp->GameThread_ProduceInput(DeltaTime, *InputData);
+			ManagerAsyncInput->AsyncInput.Add(MoveTemp(InputData));
+		}
+	}
+}
\ No newline at end of file
diff --git a/Engine/Plugins/Experimental/Mover/Source/Mover/Private/PhysicsMover/PhysicsMoverManagerAsyncCallback.cpp b/Engine/Plugins/Experimental/Mover/Source/Mover/Private/PhysicsMover/PhysicsMoverManagerAsyncCallback.cpp
new file mode 100644
--- /dev/null
+++ b/Engine/Plugins/Experimental/Mover/Source/Mover/Private/PhysicsMover/PhysicsMoverManagerAsyncCallback.cpp
@@ -0,0 +1,104 @@
+// Copyright Epic Games, Inc. All Rights Reserved.
+
+#include "PhysicsMover/PhysicsMoverManagerAsyncCallback.h"
+
+#include "Backends/MoverNetworkPhysicsLiaison.h"
+#include "Chaos/Character/CharacterGroundConstraint.h"
+#include "Chaos/Framework/Parallel.h"
+
+//////////////////////////////////////////////////////////////////////////
+
+extern FPhysicsDrivenMotionDebugParams GPhysicsDrivenMotionDebugParams;
+
+//////////////////////////////////////////////////////////////////////////
+
+void FPhysicsMoverManagerAsyncCallback::ProcessInputs_Internal(int32 PhysicsStep)
+{
+	const FPhysicsMoverManagerAsyncInput* ManagerAsyncInput = GetConsumerInput_Internal();
+	if (ManagerAsyncInput == nullptr)
+	{
+		return;
+	}
+
+	if (ManagerAsyncInput->AsyncInput.Num() == 0)
+	{
+		return;
+	}
+
+	const TArray<TUniquePtr<FPhysicsMoverAsyncInput>>& InputDataArray = ManagerAsyncInput->AsyncInput;
+	const float DeltaTime = GetDeltaTime_Internal();
+
+	auto LambdaParallelUpdate = [PhysicsStep, DeltaTime, &InputDataArray](int32 Idx)
+	{
+		const FPhysicsMoverAsyncInput& AsyncInput = *InputDataArray[Idx];
+		if (AsyncInput.IsValid())
+		{
+			AsyncInput.MoverSimulation->AsyncPhysics_ProcessInputs(PhysicsStep, DeltaTime, AsyncInput);
+		}
+	};
+
+	bool ForceSingleThread = GPhysicsDrivenMotionDebugParams.EnableMultithreading;
+	Chaos::PhysicsParallelFor(InputDataArray.Num(), LambdaParallelUpdate, ForceSingleThread);
+}
+
+void FPhysicsMoverManagerAsyncCallback::OnPreSimulate_Internal()
+{
+	const FPhysicsMoverManagerAsyncInput* ManagerAsyncInput = GetConsumerInput_Internal();
+
+	if (!ManagerAsyncInput)
+	{
+		return;
+	}
+
+	if (ManagerAsyncInput->AsyncInput.Num() == 0)
+	{
+		return;
+	}
+
+	const TArray<TUniquePtr<FPhysicsMoverAsyncInput>>& InputDataArray = ManagerAsyncInput->AsyncInput;
+	const FPhysicsMoverSimulationTickParams TickParams = { GetSimTime_Internal(), GetDeltaTime_Internal() };
+	FPhysicsMoverManagerAsyncOutput& ManagerAsyncOutput = GetProducerOutputData_Internal();
+
+	auto LambdaParallelUpdate = [TickParams, &InputDataArray, &ManagerAsyncOutput](int32 Idx)
+	{
+		const FPhysicsMoverAsyncInput& AsyncInput = *InputDataArray[Idx];
+		if (AsyncInput.IsValid())
+		{
+			TUniquePtr<FPhysicsMoverAsyncOutput> AsyncOutput = MakeUnique<FPhysicsMoverAsyncOutput>();
+			AsyncInput.MoverSimulation->AsyncPhysics_OnPreSimulate(TickParams, AsyncInput, *AsyncOutput);
+			ManagerAsyncOutput.PhysicsMoverToAsyncOutput.Add(AsyncInput.MoverIdx, MoveTemp(AsyncOutput));
+		}
+	};
+
+	bool ForceSingleThread = GPhysicsDrivenMotionDebugParams.EnableMultithreading;
+	Chaos::PhysicsParallelFor(InputDataArray.Num(), LambdaParallelUpdate, ForceSingleThread);
+}
+
+void FPhysicsMoverManagerAsyncCallback::OnContactModification_Internal(Chaos::FCollisionContactModifier& Modifier)
+{
+	const FPhysicsMoverManagerAsyncInput* ManagerAsyncInput = GetConsumerInput_Internal();
+
+	if (!ManagerAsyncInput)
+	{
+		return;
+	}
+
+	if (ManagerAsyncInput->AsyncInput.Num() == 0)
+	{
+		return;
+	}
+
+	const TArray<TUniquePtr<FPhysicsMoverAsyncInput>>& InputDataArray = ManagerAsyncInput->AsyncInput;
+
+	auto LambdaParallelUpdate = [&InputDataArray, &Modifier](int32 Idx)
+	{
+		const FPhysicsMoverAsyncInput& AsyncInput = *InputDataArray[Idx];
+		if (AsyncInput.IsValid())
+		{
+			AsyncInput.MoverSimulation->ASyncPhysics_OnContactModification(AsyncInput, Modifier);
+		}
+	};
+
+	bool ForceSingleThread = GPhysicsDrivenMotionDebugParams.EnableMultithreading;
+	Chaos::PhysicsParallelFor(InputDataArray.Num(), LambdaParallelUpdate, ForceSingleThread);
+}
\ No newline at end of file
diff --git a/Engine/Plugins/Experimental/Mover/Source/Mover/Private/PhysicsMover/PhysicsMoverSimulationTypes.cpp b/Engine/Plugins/Experimental/Mover/Source/Mover/Private/PhysicsMover/PhysicsMoverSimulationTypes.cpp
new file mode 100644
--- /dev/null
+++ b/Engine/Plugins/Experimental/Mover/Source/Mover/Private/PhysicsMover/PhysicsMoverSimulationTypes.cpp
@@ -0,0 +1,21 @@
+// Copyright Epic Games, Inc. All Rights Reserved.
+
+#include "PhysicsMover/PhysicsMoverSimulationTypes.h"
+#include "HAL/IConsoleManager.h"
+
+//////////////////////////////////////////////////////////////////////////
+// Debug
+
+FPhysicsDrivenMotionDebugParams GPhysicsDrivenMotionDebugParams;
+
+FAutoConsoleVariableRef CVarPhysicsDrivenMotionEnableMultithreading(TEXT("p.mover.physics.EnableMultithreading"),
+	GPhysicsDrivenMotionDebugParams.EnableMultithreading, TEXT("Enable multi-threading of physics driven motion updates."));
+
+FAutoConsoleVariableRef CVarPhysicsDrivenMotionDebugDrawFloorTest(TEXT("p.mover.physics.DebugDrawFloorQueries"),
+	GPhysicsDrivenMotionDebugParams.DebugDrawGroundQueries, TEXT("Debug draw floor test queries."));
+
+FAutoConsoleVariableRef CVarPhysicsDrivenMotionTeleportThreshold(TEXT("p.mover.physics.TeleportThreshold"),
+	GPhysicsDrivenMotionDebugParams.TeleportThreshold, TEXT("Single frame movement threshold in cm that will trigger a teleport."));
+
+FAutoConsoleVariableRef CVarPhysicsDrivenMotionMinStepUpDistance(TEXT("p.mover.physics.MinStepUpDistance"),
+	GPhysicsDrivenMotionDebugParams.MinStepUpDistance, TEXT("Minimum distance that will be considered a step up."));
\ No newline at end of file
diff --git a/Engine/Plugins/Experimental/Mover/Source/Mover/Private/PhysicsMover/Transitions/PhysicsJumpCheck.cpp b/Engine/Plugins/Experimental/Mover/Source/Mover/Private/PhysicsMover/Transitions/PhysicsJumpCheck.cpp
new file mode 100644
--- /dev/null
+++ b/Engine/Plugins/Experimental/Mover/Source/Mover/Private/PhysicsMover/Transitions/PhysicsJumpCheck.cpp
@@ -0,0 +1,47 @@
+// Copyright Epic Games, Inc. All Rights Reserved.
+
+#include "PhysicsMover/Transitions/PhysicsJumpCheck.h"
+#include "Kinematic/LayeredMoves/BasicLayeredMoves.h"
+#include "MoverComponent.h"
+
+#if WITH_EDITOR
+#include "Misc/DataValidation.h"
+#endif
+
+UPhysicsJumpCheck::UPhysicsJumpCheck(const FObjectInitializer& ObjectInitializer)
+	: Super(ObjectInitializer)
+{
+	TransitionToMode = KinematicModeNames::Falling;
+}
+
+FTransitionEvalResult UPhysicsJumpCheck::OnEvaluate(const FSimulationTickParams& Params) const
+{
+	FTransitionEvalResult EvalResult; 
+
+	const FMoverTickStartData& StartState = Params.StartState;
+	const FKinematicDefaultInputs* KinematicInputs = StartState.InputCmd.InputCollection.FindDataByType<FKinematicDefaultInputs>();
+	check(KinematicInputs);
+
+	if (KinematicInputs->bIsJumpJustPressed)
+	{
+		EvalResult.NextMode = TransitionToMode;
+	}
+
+	return EvalResult;
+}
+
+void UPhysicsJumpCheck::OnTrigger(const FSimulationTickParams& Params)
+{
+	TSharedPtr<FLayeredMove_JumpImpulse> JumpMove = MakeShared<FLayeredMove_JumpImpulse>();
+	JumpMove->UpwardsSpeed = JumpUpwardsSpeed;
+
+	Params.MoverComponent->QueueLayeredMove(JumpMove);
+
+}
+
+#if WITH_EDITOR
+EDataValidationResult UPhysicsJumpCheck::IsDataValid(FDataValidationContext& Context) const
+{
+	return Super::IsDataValid(Context);
+}
+#endif // WITH_EDITOR
diff --git a/Engine/Plugins/Experimental/Mover/Source/Mover/Public/Backends/MoverNetworkPhysicsLiaison.h b/Engine/Plugins/Experimental/Mover/Source/Mover/Public/Backends/MoverNetworkPhysicsLiaison.h
new file mode 100644
--- /dev/null
+++ b/Engine/Plugins/Experimental/Mover/Source/Mover/Public/Backends/MoverNetworkPhysicsLiaison.h
@@ -0,0 +1,162 @@
+// Copyright Epic Games, Inc. All Rights Reserved.
+
+#pragma once
+
+#include "Chaos/Character/CharacterGroundConstraint.h"
+#include "Chaos/GeometryParticlesfwd.h"
+#include "Backends/MoverBackendLiaison.h"
+#include "MoverComponent.h"
+#include "Physics/NetworkPhysicsComponent.h"
+#include "PhysicsMover/PhysicsMoverSimulationTypes.h"
+#include "MoverNetworkPhysicsLiaison.generated.h"
+
+class UCommonLegacyMovementSettings;
+
+
+//////////////////////////////////////////////////////////////////////////
+// Physics networking
+
+USTRUCT()
+struct MOVER_API FNetworkPhysicsMoverInputs : public FNetworkPhysicsData
+{
+	GENERATED_USTRUCT_BODY()
+
+	UPROPERTY()
+	FMoverInputCmdContext InputCmdContext;
+
+	/**  Apply the data onto the network physics component */
+	virtual void ApplyData(UActorComponent* NetworkComponent) const override;
+
+	/**  Build the data from the network physics component */
+	virtual void BuildData(const UActorComponent* NetworkComponent) override;
+
+	/**  Serialize data function that will be used to transfer the struct across the network */
+	bool NetSerialize(FArchive& Ar, class UPackageMap* Map, bool& bOutSuccess);
+
+	/** Interpolate the data in between two inputs data */
+	virtual void InterpolateData(const FNetworkPhysicsData& MinData, const FNetworkPhysicsData& MaxData) override;
+	
+	/** Merge data into this input */
+	virtual void MergeData(const FNetworkPhysicsData& FromData) override;
+};
+
+template<>
+struct TStructOpsTypeTraits<FNetworkPhysicsMoverInputs> : public TStructOpsTypeTraitsBase2<FNetworkPhysicsMoverInputs>
+{
+	enum
+	{
+		WithNetSerializer = true,
+	};
+};
+
+USTRUCT()
+struct MOVER_API FNetworkPhysicsMoverState : public FNetworkPhysicsData
+{
+	GENERATED_USTRUCT_BODY()
+
+	UPROPERTY()
+	FMoverSyncState SyncStateContext;
+
+	UPROPERTY()
+	FMoverAuxStateContext AuxStateContext;
+
+	/**  Apply the data onto the network physics component */
+	virtual void ApplyData(UActorComponent* NetworkComponent) const override;
+
+	/**  Build the data from the network physics component */
+	virtual void BuildData(const UActorComponent* NetworkComponent) override;
+
+	/**  Serialize data function that will be used to transfer the struct across the network */
+	bool NetSerialize(FArchive& Ar, class UPackageMap* Map, bool& bOutSuccess);
+
+	/** Interpolate the data in between two inputs data */
+	virtual void InterpolateData(const FNetworkPhysicsData& MinData, const FNetworkPhysicsData& MaxData) override;
+};
+
+template<>
+struct TStructOpsTypeTraits<FNetworkPhysicsMoverState> : public TStructOpsTypeTraitsBase2<FNetworkPhysicsMoverState>
+{
+	enum
+	{
+		WithNetSerializer = true,
+	};
+};
+
+struct MOVER_API FNetworkPhysicsMoverTraits
+{
+	using InputsType = FNetworkPhysicsMoverInputs;
+	using StatesType = FNetworkPhysicsMoverState;
+};
+
+//////////////////////////////////////////////////////////////////////////
+// UMoverNetworkPhysicsLiaisonComponent
+
+/**
+ * MoverPhysicsLiaisonComponent: This component acts as a middleman between an actor's Mover component and the chaos physics network prediction system.
+ * This class is set on a Mover component as the "back end".
+ */
+UCLASS()
+class MOVER_API UMoverNetworkPhysicsLiaisonComponent : public UActorComponent, public IMoverBackendLiaisonInterface
+{
+	GENERATED_BODY()
+
+public:
+	UMoverNetworkPhysicsLiaisonComponent();
+
+	// IMoverBackendLiaisonInterface
+	virtual float GetCurrentSimTimeMs() override;
+	virtual int32 GetCurrentSimFrame() override;
+
+	// UObject interface
+	virtual void OnRegister() override;
+	virtual void OnUnregister() override;
+	virtual void InitializeComponent() override;
+	virtual void UninitializeComponent() override;
+	virtual bool ShouldCreatePhysicsState() const override;
+	virtual bool HasValidPhysicsState() const override;
+	virtual void OnCreatePhysicsState() override;
+	virtual void OnDestroyPhysicsState() override;
+	virtual bool CanCreatePhysics() const;
+
+	// Required for Network Physics Rewind/Resim data
+	void GetCurrentInputData(OUT FMoverInputCmdContext& InputCmd) const;
+	void SetCurrentInputData(const FMoverInputCmdContext& InputCmd);
+	void GetCurrentStateData(OUT FMoverSyncState& SyncState, OUT FMoverAuxStateContext& AuxState) const;
+	void SetCurrentStateData(const FMoverSyncState& SyncState, const FMoverAuxStateContext& AuxState);
+
+	// Used by the manager to uniquely identify the component
+	Chaos::FUniqueIdx GetUniqueIdx() const;
+
+	void GameThread_ProduceInput(float DeltaSeconds, OUT FPhysicsMoverAsyncInput& Input);
+	void GameThread_ConsumeOutput(const FPhysicsMoverAsyncOutput& Output);
+
+	void AsyncPhysics_ProcessInputs(int32 PhysicsStep, float DeltaTime, const FPhysicsMoverAsyncInput& Input) const;
+	void AsyncPhysics_OnPreSimulate(const FPhysicsMoverSimulationTickParams& TickParams, const FPhysicsMoverAsyncInput& Input, OUT FPhysicsMoverAsyncOutput& SimOutput) const;
+	void ASyncPhysics_OnContactModification(const FPhysicsMoverAsyncInput& Input, Chaos::FCollisionContactModifier& Modifier) const;
+
+protected:
+	UFUNCTION()
+	void OnComponentPhysicsStateChanged(UPrimitiveComponent* ChangedComponent, EComponentPhysicsStateChange StateChange);
+
+	bool HasValidState() const;
+
+	void DestroyConstraint();
+	void SetupConstraint();
+
+	void UpdateConstraintSettings(const FMoverAuxStateContext& AuxState);
+	void TeleportParticle(Chaos::FGeometryParticleHandle* Particle, const FVector& Position, const FQuat& Rotation) const;
+	void WakeParticleIfSleeping(Chaos::FGeometryParticleHandle* Particle) const;
+
+	FMoverTimeStep GetCurrentMoverTimeStep() const;
+
+	// These are written to by the network input and state data
+	FMoverInputCmdContext NetInputCmd;
+	FMoverSyncState NetSyncState;
+	FMoverAuxStateContext NetAuxState;
+
+	TUniquePtr<Chaos::FCharacterGroundConstraint> Constraint;
+	TObjectPtr<UNetworkPhysicsComponent> NetworkPhysicsComponent;
+	TObjectPtr<UMoverComponent> MoverComp;
+	TObjectPtr<const UCommonLegacyMovementSettings> CommonMovementSettings;
+	bool bCachedInputIsValid = false;
+};
\ No newline at end of file
diff --git a/Engine/Plugins/Experimental/Mover/Source/Mover/Public/Backends/MoverNetworkPredictionLiaison.h b/Engine/Plugins/Experimental/Mover/Source/Mover/Public/Backends/MoverNetworkPredictionLiaison.h
new file mode 100644
--- /dev/null
+++ b/Engine/Plugins/Experimental/Mover/Source/Mover/Public/Backends/MoverNetworkPredictionLiaison.h
@@ -0,0 +1,72 @@
+// Copyright Epic Games, Inc. All Rights Reserved.
+
+#pragma once
+
+#include "CoreMinimal.h"
+#include "Components/ActorComponent.h"
+#include "Backends/MoverBackendLiaison.h"
+#include "NetworkPredictionComponent.h"
+#include "NetworkPredictionSimulation.h"
+#include "MovementMode.h"
+#include "MoverTypes.h"
+
+#include "MoverNetworkPredictionLiaison.generated.h"
+
+class UMoverComponent;
+
+
+using KinematicMoverStateTypes = TNetworkPredictionStateTypes<FMoverInputCmdContext, FMoverSyncState, FMoverAuxStateContext>;
+
+/**
+ * MoverNetworkPredictionLiaisonComponent: this component acts as a middleman between an actor's Mover component and the Network Prediction plugin.
+ * This class is set on a Mover component as the "back end".
+ */
+UCLASS()
+class MOVER_API UMoverNetworkPredictionLiaisonComponent : public UNetworkPredictionComponent, public IMoverBackendLiaisonInterface
+{
+	GENERATED_BODY()
+
+public:
+	// Begin NP Driver interface
+	// Get latest local input prior to simulation step. Called by Network Prediction system on owner's instance (autonomous or authority).
+	void ProduceInput(const int32 DeltaTimeMS, FMoverInputCmdContext* Cmd);
+
+	// Restore a previous frame prior to resimulating. Called by Network Prediction system.
+	void RestoreFrame(const FMoverSyncState* SyncState, const FMoverAuxStateContext* AuxState);
+
+	// Take output for simulation. Called by Network Prediction system.
+	void FinalizeFrame(const FMoverSyncState* SyncState, const FMoverAuxStateContext* AuxState);
+
+	// Seed initial values based on component's state. Called by Network Prediction system.
+	void InitializeSimulationState(FMoverSyncState* OutSync, FMoverAuxStateContext* OutAux);
+
+	// Primary movement simulation update. Given an starting state and timestep, produce a new state. Called by Network Prediction system.
+	void SimulationTick(const FNetSimTimeStep& TimeStep, const TNetSimInput<KinematicMoverStateTypes>& SimInput, const TNetSimOutput<KinematicMoverStateTypes>& SimOutput);
+	// End NP Driver interface
+
+	// IMoverBackendLiaisonInterface
+	virtual float GetCurrentSimTimeMs() override;
+	virtual int32 GetCurrentSimFrame() override;
+	// End IMoverBackendLiaisonInterface
+
+	virtual void BeginPlay() override;
+
+	// UObject interface
+	void InitializeComponent() override;
+	void OnRegister() override;
+	void RegisterComponentTickFunctions(bool bRegister) override;
+	// End UObject interface
+
+	// UNetworkPredictionComponent interface
+	virtual void InitializeNetworkPredictionProxy() override;
+	// End UNetworkPredictionComponent interface
+
+
+public:
+	UMoverNetworkPredictionLiaisonComponent();
+
+protected:
+	TObjectPtr<UMoverComponent> MoverComp;	// the component that we're in charge of driving
+	FMoverSyncState* StartingOutSync;
+	FMoverAuxStateContext* StartingOutAux;
+};
\ No newline at end of file
diff --git a/Engine/Plugins/Experimental/Mover/Source/Mover/Public/MoveLibrary/AirMovementUtils.h b/Engine/Plugins/Experimental/Mover/Source/Mover/Public/MoveLibrary/AirMovementUtils.h
new file mode 100644
--- /dev/null
+++ b/Engine/Plugins/Experimental/Mover/Source/Mover/Public/MoveLibrary/AirMovementUtils.h
@@ -0,0 +1,73 @@
+﻿// Copyright Epic Games, Inc. All Rights Reserved.
+
+#pragma once
+
+#include "CoreMinimal.h"
+#include "MoverDataModelTypes.h"
+#include "Kismet/BlueprintFunctionLibrary.h"
+#include "AirMovementUtils.generated.h"
+
+struct FProposedMove;
+
+// Input parameters for controlled free movement function
+USTRUCT(BlueprintType)
+struct MOVER_API FFreeMoveParams
+{
+	GENERATED_BODY()
+
+	UPROPERTY(BlueprintReadOnly, EditAnywhere, Category = Mover)
+	EMoveInputType MoveInputType = EMoveInputType::DirectionalIntent;
+
+	UPROPERTY(BlueprintReadOnly, EditAnywhere, Category = Mover)
+	FVector MoveInput = FVector::ZeroVector;
+
+	UPROPERTY(BlueprintReadOnly, EditAnywhere, Category = Mover)
+	FRotator OrientationIntent = FRotator::ZeroRotator;
+
+	UPROPERTY(BlueprintReadOnly, EditAnywhere, Category = Mover)
+	FVector PriorVelocity = FVector::ZeroVector;
+
+	UPROPERTY(BlueprintReadOnly, EditAnywhere, Category = Mover)
+	FRotator PriorOrientation = FRotator::ZeroRotator;
+
+	UPROPERTY(BlueprintReadOnly, EditAnywhere, Category = Mover)
+	float MaxSpeed = 800.f;
+
+	UPROPERTY(BlueprintReadOnly, EditAnywhere, Category = Mover)
+	float Acceleration = 4000.f;
+
+	UPROPERTY(BlueprintReadOnly, EditAnywhere, Category = Mover)
+	float Deceleration = 8000.f;
+
+	UPROPERTY(BlueprintReadOnly, EditAnywhere, Category = Mover)
+	float TurningBoost = 8.f;
+
+	UPROPERTY(BlueprintReadOnly, EditAnywhere, Category = Mover)
+	float TurningRate = 500.f;
+
+	UPROPERTY(BlueprintReadOnly, EditAnywhere, Category = Mover)
+	float DeltaSeconds = 0.f;
+};
+
+/**
+ * AirMovementUtils: a collection of stateless static BP-accessible functions for a variety of air movement-related operations
+ */
+UCLASS()
+class MOVER_API UAirMovementUtils : public UBlueprintFunctionLibrary
+{
+	GENERATED_BODY()
+	
+public:
+	/** Generate a new movement based on move/orientation intents and the prior state, unconstrained like when flying */
+	UFUNCTION(BlueprintCallable, Category = Mover)
+	static FProposedMove ComputeControlledFreeMove(const FFreeMoveParams& InParams);
+	
+    // Checks if a hit result represents a walkable location that an actor can land on
+    UFUNCTION(BlueprintCallable, Category=Mover)
+    static bool IsValidLandingSpot(USceneComponent* UpdatedComponent, UPrimitiveComponent* UpdatedPrimitive, const FVector& Location, const FHitResult& Hit, float FloorSweepDistance, float MaxWalkSlopeCosine, FFloorCheckResult& OutFloorResult);
+    
+    /** Attempts to move a component along a surface, while checking for landing on a walkable surface. Intended for use while falling. Returns the percent of time applied, with 0.0 meaning no movement occurred. */
+    UFUNCTION(BlueprintCallable, Category=Mover)
+    static float TryMoveToFallAlongSurface(USceneComponent* UpdatedComponent, UPrimitiveComponent* UpdatedPrimitive, UMoverComponent* MoverComponent, const FVector& Delta, float PctOfDeltaToMove, const FQuat Rotation, const FVector& Normal, FHitResult& Hit, bool bHandleImpact, float FloorSweepDistance, float MaxWalkSlopeCosine, FFloorCheckResult& OutFloorResult, FMovementRecord& MoveRecord);
+
+};
diff --git a/Engine/Plugins/Experimental/Mover/Source/Mover/Public/MoveLibrary/BasedMovementUtils.h b/Engine/Plugins/Experimental/Mover/Source/Mover/Public/MoveLibrary/BasedMovementUtils.h
new file mode 100644
--- /dev/null
+++ b/Engine/Plugins/Experimental/Mover/Source/Mover/Public/MoveLibrary/BasedMovementUtils.h
@@ -0,0 +1,116 @@
+// Copyright Epic Games, Inc. All Rights Reserved.
+
+#pragma once
+
+#include "Kismet/BlueprintFunctionLibrary.h"
+#include "BasedMovementUtils.generated.h"
+
+struct FMovementRecord;
+struct FFloorCheckResult;
+
+/** Data about the object a Mover actor is basing its movement on, such as when standing on a moving platform */
+USTRUCT(BlueprintType)
+struct MOVER_API FRelativeBaseInfo
+{
+	GENERATED_USTRUCT_BODY()
+
+	/** Component we are moving relative to */
+	UPROPERTY(VisibleInstanceOnly, BlueprintReadOnly, Category="Mover")
+	TObjectPtr<UPrimitiveComponent> MovementBase = nullptr;
+
+	/** Bone name on component, for skeletal meshes. NAME_None if not a skeletal mesh or if bone is invalid. */
+	UPROPERTY(VisibleInstanceOnly, BlueprintReadOnly, Category = "Mover")
+	FName BoneName = NAME_None;
+
+	/** Last captured worldspace location of MovementBase / Bone */
+	UPROPERTY(VisibleInstanceOnly, BlueprintReadOnly, Category = "Mover")
+	FVector Location = FVector::ZeroVector;
+
+	/** Last captured worldspace orientation of MovementBase / Bone */
+	UPROPERTY(VisibleInstanceOnly, BlueprintReadOnly, Category = "Mover")
+	FQuat Rotation = FQuat::Identity;
+
+	/** Last captured location of the tethering point where the Mover actor is "attached". */
+	UPROPERTY(VisibleInstanceOnly, BlueprintReadOnly, Category = "Mover")
+	FVector ContactLocalPosition = FVector::ZeroVector;;
+
+public:
+
+	void Clear();
+
+	bool HasRelativeInfo() const;
+	bool UsesSameBase(const FRelativeBaseInfo& Other) const;
+	bool UsesSameBase(const UPrimitiveComponent* OtherComp, FName OtherBoneName=NAME_None) const;
+
+	void SetFromFloorResult(const FFloorCheckResult& FloorTestResult);
+	void SetFromComponent(UPrimitiveComponent* InRelativeComp, FName InBoneName=NAME_None);
+};
+
+/**
+ * MovementBaseUtils: a collection of stateless static BP-accessible functions for based movement
+ */
+UCLASS()
+class MOVER_API UBasedMovementUtils : public UBlueprintFunctionLibrary
+{
+	GENERATED_BODY()
+
+public:
+	/** Determine whether MovementBase can possibly move. */
+	UFUNCTION(BlueprintCallable, Category = "Mover/MovementBases")
+	static bool IsADynamicBase(const UPrimitiveComponent* MovementBase);
+
+	/** Attempts to move a component relative to any change of the movement base. OldBaseInfo will typically be captured from the prior simulation frame. */
+	static bool TryMoveToStayWithBase(USceneComponent* UpdatedComponent, UPrimitiveComponent* UpdatedPrimitive, const FRelativeBaseInfo& OldBaseInfo, FMovementRecord& MoveRecord, const bool& bIgnoreBaseRotation);
+	
+	/** Get the transform (local-to-world) for the given MovementBase, optionally at the location of a bone. Returns false if MovementBase is nullptr, or if BoneName is not a valid bone. */
+	UFUNCTION(BlueprintCallable, Category="Mover/MovementBases")
+	static bool GetMovementBaseTransform(const UPrimitiveComponent* MovementBase, const FName BoneName, FVector& OutLocation, FQuat& OutQuat);
+
+	/** Convert a local location to a world location for a given MovementBase. Returns false if MovementBase is nullptr, or if BoneName is not a valid bone. Scaling is ignored. */
+	UFUNCTION(BlueprintCallable, Category = "Mover/MovementBases")
+	static bool TransformBasedLocationToWorld(const UPrimitiveComponent* MovementBase, const FName BoneName, FVector LocalLocation, FVector& OutLocationWorldSpace);
+
+	/** Convert a world location to a local location for a given MovementBase, optionally at the location of a bone. Returns false if MovementBase is nullptr, or if BoneName is not a valid bone. Scaling is ignored. */
+	UFUNCTION(BlueprintCallable, Category = "Mover/MovementBases")
+	static bool TransformWorldLocationToBased(const UPrimitiveComponent* MovementBase, const FName BoneName, FVector WorldSpaceLocation, FVector& OutLocalLocation);
+
+	/** Convert a local direction to a world direction for a given MovementBase. Returns false if MovementBase is nullptr, or if BoneName is not a valid bone. Scaling is ignored. */
+	UFUNCTION(BlueprintCallable, Category = "Mover/MovementBases")
+	static bool TransformBasedDirectionToWorld(const UPrimitiveComponent* MovementBase, const FName BoneName, FVector LocalDirection, FVector& OutDirectionWorldSpace);
+
+	/** Convert a world direction to a local direction for a given MovementBase, optionally relative to the orientation of a bone. Returns false if MovementBase is nullptr, or if BoneName is not a valid bone. Scaling is ignored. */
+	UFUNCTION(BlueprintCallable, Category = "Mover/MovementBases")
+	static bool TransformWorldDirectionToBased(const UPrimitiveComponent* MovementBase, const FName BoneName, FVector WorldSpaceDirection, FVector& OutLocalDirection);
+
+	/** Convert a local rotator to world space for a given MovementBase. Returns false if MovementBase is nullptr, or if BoneName is not a valid bone. Scaling is ignored. */
+	UFUNCTION(BlueprintCallable, Category = "Mover/MovementBases")
+	static bool TransformBasedRotatorToWorld(const UPrimitiveComponent* MovementBase, const FName BoneName, FRotator LocalRotator, FRotator& OutWorldSpaceRotator);
+
+	/** Convert a world space rotator to a local rotator for a given MovementBase, optionally relative to the orientation of a bone. Returns false if MovementBase is nullptr, or if BoneName is not a valid bone. Scaling is ignored. */
+	UFUNCTION(BlueprintCallable, Category = "Mover/MovementBases")
+	static bool TransformWorldRotatorToBased(const UPrimitiveComponent* MovementBase, const FName BoneName, FRotator WorldSpaceRotator, FRotator& OutLocalRotator);
+
+	/** Convert a local location to a world location for a given MovementBase. Returns false if MovementBase is nullptr, or if BoneName is not a valid bone. Scaling is ignored. */
+	UFUNCTION(BlueprintCallable, Category = "Mover/MovementBases")
+	static void TransformLocationToWorld(FVector BasePos, FQuat BaseQuat, FVector LocalLocation, FVector& OutLocationWorldSpace);
+
+	/** Convert a world location to a local location for a given MovementBase, optionally at the location of a bone. Returns false if MovementBase is nullptr, or if BoneName is not a valid bone. Scaling is ignored. */
+	UFUNCTION(BlueprintCallable, Category = "Mover/MovementBases")
+	static void TransformLocationToLocal(FVector BasePos, FQuat BaseQuat, FVector WorldSpaceLocation, FVector& OutLocalLocation);
+
+	/** Convert a local direction to a world direction for a given MovementBase. Returns false if MovementBase is nullptr, or if BoneName is not a valid bone. Scaling is ignored. */
+	UFUNCTION(BlueprintCallable, Category = "Mover/MovementBases")
+	static void TransformDirectionToWorld(FQuat BaseQuat, FVector LocalDirection, FVector& OutDirectionWorldSpace);
+
+	/** Convert a world direction to a local direction for a given MovementBase, optionally relative to the orientation of a bone. Returns false if MovementBase is nullptr, or if BoneName is not a valid bone. Scaling is ignored. */
+	UFUNCTION(BlueprintCallable, Category = "Mover/MovementBases")
+	static void TransformDirectionToLocal(FQuat BaseQuat, FVector WorldSpaceDirection, FVector& OutLocalDirection);
+
+	/** Convert a local rotator to world space for a given MovementBase. Returns false if MovementBase is nullptr, or if BoneName is not a valid bone. Scaling is ignored. */
+	UFUNCTION(BlueprintCallable, Category = "Mover/MovementBases")
+	static void TransformRotatorToWorld(FQuat BaseQuat, FRotator LocalRotator, FRotator& OutWorldSpaceRotator);
+
+	/** Convert a world space rotator to a local rotator for a given MovementBase, optionally relative to the orientation of a bone. Returns false if MovementBase is nullptr, or if BoneName is not a valid bone. Scaling is ignored. */
+	UFUNCTION(BlueprintCallable, Category = "Mover/MovementBases")
+	static void TransformRotatorToLocal(FQuat BaseQuat, FRotator WorldSpaceRotator, FRotator& OutLocalRotator);
+};
diff --git a/Engine/Plugins/Experimental/Mover/Source/Mover/Public/MoveLibrary/FloorQueryUtils.h b/Engine/Plugins/Experimental/Mover/Source/Mover/Public/MoveLibrary/FloorQueryUtils.h
new file mode 100644
--- /dev/null
+++ b/Engine/Plugins/Experimental/Mover/Source/Mover/Public/MoveLibrary/FloorQueryUtils.h
@@ -0,0 +1,125 @@
+﻿// Copyright Epic Games, Inc. All Rights Reserved.
+
+#pragma once
+
+#include "CoreMinimal.h"
+#include "Kismet/BlueprintFunctionLibrary.h"
+#include "Components/PrimitiveComponent.h"
+#include "FloorQueryUtils.generated.h"
+
+namespace UE::FloorQueryUtility
+{
+	extern const float MIN_FLOOR_DIST;
+	extern const float MAX_FLOOR_DIST;
+	extern const float SWEEP_EDGE_REJECT_DISTANCE;
+}
+
+/** Data about the floor for walking movement, used by Mover simulations */
+USTRUCT(BlueprintType)
+struct MOVER_API FFloorCheckResult
+{
+	GENERATED_USTRUCT_BODY()
+
+	/**
+	* True if there was a blocking hit in the floor test that was NOT in initial penetration.
+	* The HitResult can give more info about other circumstances.
+	*/
+	UPROPERTY(VisibleInstanceOnly, BlueprintReadOnly, Category = CharacterFloor)
+	uint32 bBlockingHit : 1;
+
+	/** True if the hit found a valid walkable floor. */
+	UPROPERTY(VisibleInstanceOnly, BlueprintReadOnly, Category = CharacterFloor)
+	uint32 bWalkableFloor : 1;
+
+	/** True if the hit found a valid walkable floor using a line trace (rather than a sweep test, which happens when the sweep test fails to yield a walkable surface). */
+	UPROPERTY(VisibleInstanceOnly, BlueprintReadOnly, Category = CharacterFloor)
+	uint32 bLineTrace : 1;
+
+	/** The distance to the floor, computed from the trace. Only valid if bLineTrace is true. */
+	UPROPERTY(VisibleInstanceOnly, BlueprintReadOnly, Category = CharacterFloor)
+    float LineDist;
+
+	/** The distance to the floor, computed from the swept capsule trace. */
+	UPROPERTY(VisibleInstanceOnly, BlueprintReadOnly, Category = CharacterFloor)
+	float FloorDist;
+
+	/** Hit result of the test that found a floor. Includes more specific data about the point of impact and surface normal at that point. */
+	UPROPERTY(VisibleInstanceOnly, BlueprintReadOnly, Category = CharacterFloor)
+	FHitResult HitResult;
+
+public:
+
+	FFloorCheckResult()
+		: bBlockingHit(false)
+		, bWalkableFloor(false)
+		, bLineTrace(false)
+		, LineDist(0.f)
+		, FloorDist(0.f)
+		, HitResult(1.f)
+	{
+	}
+
+	/** Returns true if the floor result hit a walkable surface. */
+	bool IsWalkableFloor() const
+	{
+		return bBlockingHit && bWalkableFloor;
+	}
+
+	void Clear()
+	{
+		bBlockingHit = false;
+		bWalkableFloor = false;
+		FloorDist = 0.f;
+		bLineTrace = false;
+		LineDist = 0.f;
+		HitResult.Reset(1.f, false);
+	}
+
+	/** Gets the distance to floor, either LineDist or FloorDist. */
+	float GetDistanceToFloor() const
+	{
+		return FloorDist;
+	}
+
+	/** Sets the passed in hit result with data from the floor Sweep */
+	void SetFromSweep(const FHitResult& InHit, const float InSweepFloorDist, const bool bIsWalkableFloor);
+	/** Sets the passed in hit result with data from the line trace */
+	void SetFromLineTrace(const FHitResult& InHit, const float InSweepFloorDist, const float InLineDist, const bool bIsWalkableFloor);
+};
+
+/** Used by some movement operations to conditionally return a floor check result, if one was performed */
+struct FOptionalFloorCheckResult
+{
+	uint8 bHasFloorResult : 1;			// If true, the floor test result has been computed
+	FFloorCheckResult FloorTestResult;
+
+	FOptionalFloorCheckResult()
+	{
+		bHasFloorResult = false;
+	}
+};
+
+/**
+ * FloorQueryUtils: a collection of stateless static BP-accessible functions for a variety of operations involving floor checks
+ */
+UCLASS()
+class MOVER_API UFloorQueryUtils : public UBlueprintFunctionLibrary
+{
+	GENERATED_BODY()
+
+public:
+	static void FindFloor(const USceneComponent* UpdatedComponent, const UPrimitiveComponent* UpdatedPrimitive, float FloorSweepDistance, float MaxWalkSlopeCosine, const FVector& Location, FFloorCheckResult& OutFloorResult);
+
+	static void ComputeFloorDist(const USceneComponent* UpdatedComponent, const UPrimitiveComponent* UpdatedPrimitive, float LineTraceDistance, float FloorSweepDistance, float MaxWalkSlopeCosine, const FVector& Location, FFloorCheckResult& OutFloorResult);
+
+	static bool FloorSweepTest(const UPrimitiveComponent* UpdatedPrimitive, FHitResult& OutHit, const FVector& Start, const FVector& End, ECollisionChannel TraceChannel, const struct FCollisionShape& CollisionShape, const struct FCollisionQueryParams& Params, const struct FCollisionResponseParams& ResponseParam);
+
+	UFUNCTION(BlueprintCallable, Category=Mover)
+	static bool IsHitSurfaceWalkable(const FHitResult& Hit, float MaxWalkSlopeCosine);
+
+	/**
+	 * Return true if the 2D distance to the impact point is inside the edge tolerance (CapsuleRadius minus a small rejection threshold).
+	 * Useful for rejecting adjacent hits when finding a floor or landing spot.
+	 */
+	static bool IsWithinEdgeTolerance(const FVector& CapsuleLocation, const FVector& TestImpactPoint, float CapsuleRadius);
+};
diff --git a/Engine/Plugins/Experimental/Mover/Source/Mover/Public/MoveLibrary/GroundMovementUtils.h b/Engine/Plugins/Experimental/Mover/Source/Mover/Public/MoveLibrary/GroundMovementUtils.h
new file mode 100644
--- /dev/null
+++ b/Engine/Plugins/Experimental/Mover/Source/Mover/Public/MoveLibrary/GroundMovementUtils.h
@@ -0,0 +1,95 @@
+﻿// Copyright Epic Games, Inc. All Rights Reserved.
+
+#pragma once
+
+#include "CoreMinimal.h"
+#include "MoverDataModelTypes.h"
+#include "Kismet/BlueprintFunctionLibrary.h"
+#include "GroundMovementUtils.generated.h"
+
+class UMoverComponent;
+struct FMovementRecord;
+struct FFloorCheckResult;
+struct FOptionalFloorCheckResult;
+struct FHitResult;
+struct FProposedMove;
+
+// Input parameters for controlled ground movement function
+USTRUCT(BlueprintType)
+struct MOVER_API FGroundMoveParams
+{
+	GENERATED_BODY()
+
+	UPROPERTY(BlueprintReadOnly, EditAnywhere, Category = Mover)
+	EMoveInputType MoveInputType = EMoveInputType::DirectionalIntent;
+
+	UPROPERTY(BlueprintReadOnly, EditAnywhere, Category = Mover)
+	FVector MoveInput = FVector::ZeroVector;
+
+	UPROPERTY(BlueprintReadOnly, EditAnywhere, Category = Mover)
+	FRotator OrientationIntent = FRotator::ZeroRotator;
+
+	UPROPERTY(BlueprintReadOnly, EditAnywhere, Category = Mover)
+	FVector PriorVelocity = FVector::ZeroVector;
+
+	UPROPERTY(BlueprintReadOnly, EditAnywhere, Category = Mover)
+	FRotator PriorOrientation = FRotator::ZeroRotator;
+
+	UPROPERTY(BlueprintReadOnly, EditAnywhere, Category = Mover)
+	float MaxSpeed = 800.f;
+
+	UPROPERTY(BlueprintReadOnly, EditAnywhere, Category = Mover)
+	float Acceleration = 4000.f;
+
+	UPROPERTY(BlueprintReadOnly, EditAnywhere, Category = Mover)
+	float Deceleration = 8000.f;
+
+	UPROPERTY(BlueprintReadOnly, EditAnywhere, Category = Mover)
+	float Friction = 0.f;
+
+	UPROPERTY(BlueprintReadOnly, EditAnywhere, Category = Mover)
+	float TurningRate = 500.f;
+
+	UPROPERTY(BlueprintReadOnly, EditAnywhere, Category = Mover)
+	float TurningBoost = 8.f;
+
+	UPROPERTY(BlueprintReadOnly, EditAnywhere, Category = Mover)
+	FVector GroundNormal = FVector::UpVector;
+
+	UPROPERTY(BlueprintReadOnly, EditAnywhere, Category = Mover)
+	float DeltaSeconds = 0.f;
+};
+
+/**
+ * GroundMovementUtils: a collection of stateless static BP-accessible functions for a variety of ground movement-related operations
+ */
+UCLASS()
+class MOVER_API UGroundMovementUtils : public UBlueprintFunctionLibrary
+{
+	GENERATED_BODY()
+	
+public:
+	/** Generate a new movement based on move/orientation intents and the prior state, constrained to the ground movement plane. Also applies deceleration friction as necessary. */
+	UFUNCTION(BlueprintCallable, Category = Mover)
+	static FProposedMove ComputeControlledGroundMove(const FGroundMoveParams& InParams);
+	
+	// TODO: Refactor this API for fewer parameters
+	/** Move up steps or slope. Does nothing and returns false if hit surface is invalid for step-up use */
+	static bool TryMoveToStepUp(USceneComponent* UpdatedComponent, UPrimitiveComponent* UpdatedPrimitive, UMoverComponent* MoverComponent, const FVector& GravDir, float MaxStepHeight, float MaxWalkSlopeCosine, float FloorSweepDistance, const FVector& MoveDelta, const FHitResult& Hit, const FFloorCheckResult& CurrentFloor, bool bIsFalling, FOptionalFloorCheckResult* OutStepDownResult, FMovementRecord& MoveRecord);
+
+	/** Moves vertically to stay within range of the walkable floor. Does nothing and returns false if floor is unwalkable or if already in range. */ 
+	static bool TryMoveToAdjustHeightAboveFloor(USceneComponent* UpdatedComponent, UPrimitiveComponent* UpdatedPrimitive, FFloorCheckResult& CurrentFloor, float MaxWalkSlopeCosine, FMovementRecord& MoveRecord);
+
+	/** Attempts to move a component along a surface in the walking mode. Returns the percent of time applied, with 0.0 meaning no movement occurred.
+     *  Note: This modifies the normal and calls UMovementUtils::TryMoveToSlideAlongSurface
+     */
+    UFUNCTION(BlueprintCallable, Category=Mover)
+    static float TryWalkToSlideAlongSurface(USceneComponent* UpdatedComponent, UPrimitiveComponent* UpdatedPrimitive, UMoverComponent* MoverComponent, const FVector& Delta, float PctOfDeltaToMove, const FQuat Rotation, const FVector& Normal, FHitResult& Hit, bool bHandleImpact, FMovementRecord& MoveRecord, float MaxWalkSlopeCosine, float MaxStepHeight);
+
+	/** Used to change a movement to be along a ramp's surface, typically to prevent slow movement when running up/down a ramp */
+	UFUNCTION(BlueprintCallable, Category = Mover)
+	static FVector ComputeDeflectedMoveOntoRamp(const FVector& OrigMoveDelta, const FHitResult& RampHitResult, float MaxWalkSlopeCosine, const bool bHitFromLineTrace);
+
+	UFUNCTION(BlueprintCallable, Category = Mover)
+	static bool CanStepUpOnHitSurface(const FHitResult& Hit);
+};
diff --git a/Engine/Plugins/Experimental/Mover/Source/Mover/Public/MoveLibrary/ModularMovement.h b/Engine/Plugins/Experimental/Mover/Source/Mover/Public/MoveLibrary/ModularMovement.h
new file mode 100644
--- /dev/null
+++ b/Engine/Plugins/Experimental/Mover/Source/Mover/Public/MoveLibrary/ModularMovement.h
@@ -0,0 +1,90 @@
+// Copyright Epic Games, Inc. All Rights Reserved.
+
+#pragma once
+
+#include "UObject/Interface.h"
+#include "ModularMovement.generated.h"
+
+struct FMoverTickStartData;
+struct FMoverTimeStep;
+struct FProposedMove;
+class UMoverBlackboard;
+
+
+
+UINTERFACE(BlueprintType)
+class MOVER_API UTurnGeneratorInterface : public UInterface
+{
+	GENERATED_BODY()
+};
+
+/**
+ * TurnGeneratorInterface: API for generating an in-place orientation change, based on a starting state and a target orientation
+ */
+class ITurnGeneratorInterface : public IInterface
+{
+	GENERATED_BODY()
+
+public:
+	// TODO: consider making GetTurn take a params struct instead of a fixed arg list
+
+	/** Returns an additive angular velocity (per-axis degrees/second) based on the starting state and timestep */
+	UFUNCTION(BlueprintNativeEvent)
+	FRotator GetTurn(FRotator TargetOrientation, const FMoverTickStartData& FullStartState, const FMoverDefaultSyncState& MoverState, const FMoverTimeStep& TimeStep, const FProposedMove& ProposedMove, UMoverBlackboard* SimBlackboard);
+};
+
+
+/**
+ * Modular mechanism for turning a Mover actor in place using linear angular velocity.
+ */
+UCLASS(EditInlineNew)
+class MOVER_API ULinearTurnGenerator : public UObject, public ITurnGeneratorInterface
+{
+	GENERATED_BODY()
+
+public:
+
+	/** Maximum angular velocity of heading changes (degrees per second). AKA Yaw, AKA Z. Negative numbers will cause rotation to snap instantly to desired orientation. */
+	UPROPERTY(BlueprintReadWrite, EditAnywhere, Category = "Linear Turning", meta = (ClampMin = "-1", UIMin = "0", ForceUnits = "deg/s"))
+	float HeadingRate = 500.f;
+
+	/** Maximum angular velocity of pitch changes (degrees per second). Negative numbers will cause rotation to snap instantly to desired orientation. */
+	UPROPERTY(BlueprintReadWrite, EditAnywhere, Category = "Linear Turning", meta = (ClampMin = "-1", UIMin = "0", ForceUnits = "deg/s"))
+	float PitchRate = -1.f;
+
+	/** Maximum angular velocity of roll changes (degrees per second). Negative numbers will cause rotation to snap instantly to desired orientation. */
+	UPROPERTY(BlueprintReadWrite, EditAnywhere, Category = "Linear Turning", meta = (ClampMin = "-1", UIMin = "0", ForceUnits = "deg/s"))
+	float RollRate = -1.f;
+
+
+	virtual FRotator GetTurn_Implementation(FRotator TargetOrientation, const FMoverTickStartData& FullStartState, const FMoverDefaultSyncState& MoverState, const FMoverTimeStep& TimeStep, const FProposedMove& ProposedMove, UMoverBlackboard* SimBlackboard) override;
+};
+
+
+/**
+ * Modular mechanism for turning a Mover actor in place using an exact damped spring
+ */
+UCLASS(EditInlineNew)
+class MOVER_API UExactDampedTurnGenerator : public UObject, public ITurnGeneratorInterface
+{
+	GENERATED_BODY()
+
+public:
+
+	/** Time required to reach halfway to the target orientation (smaller = quicker) */
+	UPROPERTY(BlueprintReadWrite, EditAnywhere, Category = "Exact Damped Turning", meta = (ClampMin = "0.01", UIMin = "0", ForceUnits = "seconds"))
+	float HalfLifeSeconds = 0.15f;
+
+	virtual FRotator GetTurn_Implementation(FRotator TargetOrientation, const FMoverTickStartData& FullStartState, const FMoverDefaultSyncState& MoverState, const FMoverTimeStep& TimeStep, const FProposedMove& ProposedMove, UMoverBlackboard* SimBlackboard) override;
+};
+
+
+/**
+ * Base class for blueprint-implemented turn generators. This is necessary due to the lack of support for specifying 
+ * EditInlineNew on a BP class, so it has to inherit the flag from a native parent.
+ */
+UCLASS(Blueprintable, EditInlineNew, Abstract)
+class MOVER_API UBlueprintableTurnGenerator : public UObject, public ITurnGeneratorInterface
+{
+	GENERATED_BODY()
+};
diff --git a/Engine/Plugins/Experimental/Mover/Source/Mover/Public/MoveLibrary/MovementUtils.h b/Engine/Plugins/Experimental/Mover/Source/Mover/Public/MoveLibrary/MovementUtils.h
new file mode 100644
--- /dev/null
+++ b/Engine/Plugins/Experimental/Mover/Source/Mover/Public/MoveLibrary/MovementUtils.h
@@ -0,0 +1,189 @@
+// Copyright Epic Games, Inc. All Rights Reserved.
+
+#pragma once
+
+#include "Kismet/BlueprintFunctionLibrary.h"
+#include "MoverSimulationTypes.h"
+#include "MoverDataModelTypes.h"
+#include "Components/PrimitiveComponent.h"
+
+#include "MovementUtils.generated.h"
+
+struct FMovementRecord;
+class UMoverComponent;
+
+namespace UE::MoverUtils
+{
+	extern const double SMALL_MOVE_DISTANCE;
+}
+
+/** Encapsulates detailed trajectory sample info, from a move that has already occurred or one projected into the future */
+USTRUCT(BlueprintType)
+struct MOVER_API FTrajectorySampleInfo
+{
+	GENERATED_USTRUCT_BODY()
+
+	FTransform Transform;					// Position and orientation (world space)
+
+	FVector LinearVelocity;					// Velocity at the time of this sample (world space, units/sec)
+	FVector InstantaneousAcceleration;		// Acceleration at the time of this sample (world space, units/sec^2)
+
+	FRotator AngularVelocity;				// Rotational velocity (world space, degrees/sec)
+
+	float SimTimeMs;						// Time stamp of this sample, in server simulation time
+};
+
+// Input parameters for compute velocity function
+USTRUCT(BlueprintType)
+struct MOVER_API FComputeVelocityParams
+{
+	GENERATED_BODY()
+	
+	UPROPERTY(BlueprintReadOnly, EditAnywhere, Category = Mover)
+	float DeltaSeconds = 0.f;
+	
+	UPROPERTY(BlueprintReadOnly, EditAnywhere, Category = Mover)
+	FVector InitialVelocity = FVector::ZeroVector;
+	
+	UPROPERTY(BlueprintReadOnly, EditAnywhere, Category = Mover)
+	FVector MoveDirectionIntent = FVector::ZeroVector;
+
+	// AuxState variables
+	UPROPERTY(BlueprintReadOnly, EditAnywhere, Category = Mover)
+	float MaxSpeed = 0.f;
+	
+	UPROPERTY(BlueprintReadOnly, EditAnywhere, Category = Mover)
+	float TurningBoost = 0.f;
+	
+	UPROPERTY(BlueprintReadOnly, EditAnywhere, Category = Mover)
+	float Friction = 0.f;
+	
+	UPROPERTY(BlueprintReadOnly, EditAnywhere, Category = Mover)
+	float Deceleration = 0.f;
+	
+	UPROPERTY(BlueprintReadOnly, EditAnywhere, Category = Mover)
+	float Acceleration = 0.f;
+};
+
+// Input parameters for ComputeCombinedVelocity()
+USTRUCT(BlueprintType)
+struct MOVER_API FComputeCombinedVelocityParams
+{
+	GENERATED_BODY()
+	
+	UPROPERTY(BlueprintReadOnly, EditAnywhere, Category = Mover)
+	float DeltaSeconds = 0.f;
+	
+	UPROPERTY(BlueprintReadOnly, EditAnywhere, Category = Mover)
+	FVector InitialVelocity = FVector::ZeroVector;
+	
+	UPROPERTY(BlueprintReadOnly, EditAnywhere, Category = Mover)
+	FVector MoveDirectionIntent = FVector::ZeroVector;
+
+	// AuxState variables
+	UPROPERTY(BlueprintReadOnly, EditAnywhere, Category = Mover)
+	float MaxSpeed = 0.f;
+	
+	UPROPERTY(BlueprintReadOnly, EditAnywhere, Category = Mover)
+	float TurningBoost = 0.f;
+	
+	UPROPERTY(BlueprintReadOnly, EditAnywhere, Category = Mover)
+	float Friction = 0.f;
+	
+	UPROPERTY(BlueprintReadOnly, EditAnywhere, Category = Mover)
+	float Deceleration = 0.f;
+	
+	UPROPERTY(BlueprintReadOnly, EditAnywhere, Category = Mover)
+	float Acceleration = 0.f;
+
+	UPROPERTY(BlueprintReadOnly, EditAnywhere, Category = Mover)
+	FVector ExternalAcceleration = FVector::ZeroVector;
+	
+	UPROPERTY(BlueprintReadOnly, EditAnywhere, Category = Mover)
+	float OverallMaxSpeed = 0.f;
+};
+
+/**
+ * MovementUtils: a collection of stateless static BP-accessible functions for a variety of movement-related operations
+ */
+UCLASS()
+class MOVER_API UMovementUtils : public UBlueprintFunctionLibrary
+{
+	GENERATED_BODY()
+
+public:
+	// JAH TODO: Ideally, none of these functions should deal with simulation input/state types. Rework them to take only the core types they actually need
+	// JAH TODO: Make sure all 'out' parameters are last in the param list and marked as "Out"
+	// JAH TODO: separate out the public-facing ones from the internally-used ones and make all public-facing ones BlueprintCallable
+
+	/** Checks whether a given velocity is exceeding a maximum speed, with some leeway to account for numeric imprecision */
+	UFUNCTION(BlueprintCallable, Category = Mover)
+	static bool IsExceedingMaxSpeed(const FVector& Velocity, float InMaxSpeed);
+
+	/** Returns new ground-based velocity (worldspace) based on previous state, movement intent (worldspace), and movement settings */
+	UFUNCTION(BlueprintCallable, Category=Mover)
+    static FVector ComputeVelocity(const FComputeVelocityParams& InParams);
+
+	/** Returns new velocity based on previous state, movement intent, movement mode's influence and movement settings */
+	UFUNCTION(BlueprintCallable, Category=Mover)
+    static FVector ComputeCombinedVelocity(const FComputeCombinedVelocityParams& InParams);
+
+	/** Returns velocity (units per second) contributed by gravitational acceleration over a given time */
+	UFUNCTION(BlueprintCallable, Category=Mover)
+	static FVector ComputeVelocityFromGravity(const FVector& GravityAccel, float DeltaSeconds) { return GravityAccel * DeltaSeconds; }
+
+	/** Checks whether a given velocity is strong enough to lift off against gravity */
+	UFUNCTION(BlueprintCallable, Category=Mover)
+	static bool CanEscapeGravity(const FVector& PriorVelocity, const FVector& NewVelocity, const FVector& GravityAccel, float DeltaSeconds);
+
+	/** Ensures input Vector (typically a velocity, acceleration, or move delta) is limited to a movement plane. 
+	* @param bMaintainMagnitude - if true, vector will be scaled after projection in an attempt to keep magnitude the same 
+	*/
+	UFUNCTION(BlueprintCallable, Category = Mover)
+	static FVector ConstrainToPlane(const FVector& Vector, const FPlane& MovementPlane, bool bMaintainMagnitude=true);
+
+	// Surface sliding
+
+	/** Returns an alternative move delta to slide along a surface, based on parameters describing a blocked attempted move */
+	UFUNCTION(BlueprintCallable, Category=Mover)
+	static FVector ComputeSlideDelta(const FVector& Delta, const float PctOfDeltaToMove, const FVector& Normal, const FHitResult& Hit);
+
+	/** Returns an alternative move delta when we are in contact with 2 surfaces */
+	static FVector ComputeTwoWallAdjustedDelta(const FVector& MoveDelta, const FHitResult& Hit, const FVector& OldHitNormal);
+
+	/** Attempts to move a component along a surface. Returns the percent of time applied, with 0.0 meaning no movement occurred. */
+	UFUNCTION(BlueprintCallable, Category=Mover)
+	static float TryMoveToSlideAlongSurface(USceneComponent* UpdatedComponent, UPrimitiveComponent* UpdatedPrimitive, UMoverComponent* MoverComponent, const FVector& Delta, float PctOfDeltaToMove, const FQuat Rotation, const FVector& Normal, FHitResult& Hit, bool bHandleImpact, FMovementRecord& MoveRecord);
+	// Component movement
+
+	/** Attempts to move a component and resolve any penetration issues with the proposed move Delta */
+	UFUNCTION(BlueprintCallable, Category=Mover)
+	static bool TrySafeMoveUpdatedComponent(USceneComponent* UpdatedComponent, UPrimitiveComponent* UpdatedPrimitive, const FVector& Delta, const FQuat& NewRotation, bool bSweep, FHitResult& OutHit, ETeleportType Teleport, FMovementRecord& MoveRecord);
+
+	/** Returns a movement step that should get the subject of the hit result out of an initial penetration condition */
+	static FVector ComputePenetrationAdjustment(const FHitResult& Hit);
+	
+	/** Attempts to move out of a situation where the component is stuck in geometry, using a suggested adjustment to start. */
+	static bool TryMoveToResolvePenetration(USceneComponent* UpdatedComponent, UPrimitiveComponent* UpdatedPrimitive, EMoveComponentFlags MoveComponentFlags, const FVector& ProposedAdjustment, const FHitResult& Hit, const FQuat& NewRotationQuat, FMovementRecord& MoveRecord);
+	static void InitCollisionParams(const UPrimitiveComponent* UpdatedPrimitive, FCollisionQueryParams& OutParams, FCollisionResponseParams& OutResponseParam);
+	static bool OverlapTest(const USceneComponent* UpdatedComponent, const UPrimitiveComponent* UpdatedPrimitive, const FVector& Location, const FQuat& RotationQuat, const ECollisionChannel CollisionChannel, const FCollisionShape& CollisionShape, const AActor* IgnoreActor);
+
+	/** Computes velocity based on start and end positions over time */
+	UFUNCTION(BlueprintCallable, Category = Mover)
+	static FVector ComputeVelocityFromPositions(const FVector& FromPos, const FVector& ToPos, float DeltaSeconds);
+
+	/** Computes the angular velocity needed to change from one orientation to another within a time frame. Use the optional TurningRateLimit to clamp to a maximum step (negative=unlimited). */
+	UFUNCTION(BlueprintCallable, Category = Mover)
+	static FRotator ComputeAngularVelocity(const FRotator& From, const FRotator& To, float DeltaSeconds, float TurningRateLimit=-1.f);
+
+	/** Computes the directional movement intent based on input vector and associated type */
+	UFUNCTION(BlueprintCallable, Category = Mover)
+	static FVector ComputeDirectionIntent(const FVector& MoveInput, EMoveInputType MoveInputType);
+
+	// Internal functions - not meant to be called outside of this library
+	
+	/** Internal function that other move functions use to perform all actual component movement and retrieve results
+	 *  Note: This function moves the character directly and should only be used if needed. Consider using something like TrySafeMoveUpdatedComponent.
+	 */
+	static bool TryMoveUpdatedComponent_Internal(USceneComponent* UpdatedComponent, const FVector& Delta, const FQuat& NewRotation, bool bSweep, EMoveComponentFlags MoveComponentFlags, FHitResult* OutHit, ETeleportType Teleport);
+};
diff --git a/Engine/Plugins/Experimental/Mover/Source/Mover/Public/MoveLibrary/MovementUtilsTypes.h b/Engine/Plugins/Experimental/Mover/Source/Mover/Public/MoveLibrary/MovementUtilsTypes.h
new file mode 100644
--- /dev/null
+++ b/Engine/Plugins/Experimental/Mover/Source/Mover/Public/MoveLibrary/MovementUtilsTypes.h
@@ -0,0 +1,56 @@
+// Copyright Epic Games, Inc. All Rights Reserved.
+
+#pragma once
+
+#include "MovementUtilsTypes.generated.h"
+
+
+
+UENUM()
+enum class EMoveMixMode : uint8
+{
+	/** Velocity (linear and angular) is intended to be added with other sources */
+	AdditiveVelocity = 0,
+	/** Velocity (linear and angular) should override others */
+	OverrideVelocity = 1,
+	/** All move parameters should override others */
+	OverrideAll      = 2,
+};
+
+
+/** Encapsulates info about an intended move that hasn't happened yet */
+USTRUCT(BlueprintType)
+struct MOVER_API FProposedMove
+{
+	GENERATED_USTRUCT_BODY()
+
+	FProposedMove() : 
+		bHasDirIntent(false),
+		bHasTargetLocation(false)
+	{}
+
+	UPROPERTY(EditAnywhere, BlueprintReadWrite, Category = Mover)
+	EMoveMixMode MixMode = EMoveMixMode::AdditiveVelocity;		// Determines how this move should resolve with other moves
+
+	UPROPERTY(EditAnywhere, BlueprintReadWrite, Category = Mover)
+	FName	PreferredMode = NAME_None;					// Indicates that we should switch to a particular movement mode before the next simulation step is performed.
+
+	UPROPERTY(EditAnywhere, BlueprintReadWrite, Category = Mover)
+	uint8	 bHasDirIntent : 1;							// Signals whether there was any directional intent specified
+	UPROPERTY(EditAnywhere, BlueprintReadWrite, Category = Mover)
+	uint8	 bHasTargetLocation : 1;					// Signals whether the proposed move should move to a target location, regardless of other fields
+
+	UPROPERTY(EditAnywhere, BlueprintReadWrite, Category = Mover)
+	FVector  DirectionIntent = FVector::ZeroVector;		// Directional, per-axis magnitude [-1, 1] in world space (length of 1 indicates max speed intent). Only valid if bHasDirIntent is set.
+
+	UPROPERTY(EditAnywhere, BlueprintReadWrite, Category = Mover)
+	FVector  LinearVelocity = FVector::ZeroVector;		// Units per second, world space, possibly mapped onto walking surface
+	UPROPERTY(EditAnywhere, BlueprintReadWrite, Category = Mover)
+	FRotator AngularVelocity = FRotator::ZeroRotator;	// Degrees per second, local space
+
+	UPROPERTY(EditAnywhere, BlueprintReadWrite, Category = Mover)
+	FVector  MovePlaneVelocity = FVector::ZeroVector;	// Units per second, world space, always along the movement plane
+
+	UPROPERTY(EditAnywhere, BlueprintReadWrite, Category = Mover)
+	FVector  TargetLocation = FVector::ZeroVector;		// World space go-to position. Only valid if bHasTargetLocation is set. Used for movement like teleportation.
+};
\ No newline at end of file
diff --git a/Engine/Plugins/Experimental/Mover/Source/Mover/Public/MoveLibrary/WaterMovementUtils.h b/Engine/Plugins/Experimental/Mover/Source/Mover/Public/MoveLibrary/WaterMovementUtils.h
new file mode 100644
--- /dev/null
+++ b/Engine/Plugins/Experimental/Mover/Source/Mover/Public/MoveLibrary/WaterMovementUtils.h
@@ -0,0 +1,194 @@
+﻿// Copyright Epic Games, Inc. All Rights Reserved.
+
+#pragma once
+
+#include "CoreMinimal.h"
+#include "MoverDataModelTypes.h"
+#include "Kismet/BlueprintFunctionLibrary.h"
+#include "WaterBodyActor.h"
+
+#include "WaterMovementUtils.generated.h"
+
+struct FProposedMove;
+
+// Input parameters for ComputeControlledWaterMove()
+USTRUCT(BlueprintType)
+struct MOVER_API FWaterMoveParams
+{
+	GENERATED_BODY()
+
+	UPROPERTY(BlueprintReadOnly, EditAnywhere, Category = Mover)
+	EMoveInputType MoveInputType = EMoveInputType::DirectionalIntent;
+
+	UPROPERTY(BlueprintReadOnly, EditAnywhere, Category = Mover)
+	FVector MoveInput = FVector::ZeroVector;
+
+	UPROPERTY(BlueprintReadOnly, EditAnywhere, Category = Mover)
+	FRotator OrientationIntent = FRotator::ZeroRotator;
+
+	UPROPERTY(BlueprintReadOnly, EditAnywhere, Category = Mover)
+	FVector PriorVelocity = FVector::ZeroVector;
+
+	UPROPERTY(BlueprintReadOnly, EditAnywhere, Category = Mover)
+	FRotator PriorOrientation = FRotator::ZeroRotator;
+
+	UPROPERTY(BlueprintReadOnly, EditAnywhere, Category = Mover)
+	float MaxSpeed = 800.f;
+
+	UPROPERTY(BlueprintReadOnly, EditAnywhere, Category = Mover)
+	float Acceleration = 4000.f;
+
+	UPROPERTY(BlueprintReadOnly, EditAnywhere, Category = Mover)
+	float Deceleration = 8000.f;
+
+	UPROPERTY(BlueprintReadOnly, EditAnywhere, Category = Mover)
+	float Friction = 0.f;
+
+	UPROPERTY(BlueprintReadOnly, EditAnywhere, Category = Mover)
+	float TurningRate = 500.f;
+
+	UPROPERTY(BlueprintReadOnly, EditAnywhere, Category = Mover)
+	float TurningBoost = 8.f;
+
+	UPROPERTY(BlueprintReadOnly, EditAnywhere, Category = Mover)
+	float DeltaSeconds = 0.f;
+	
+	UPROPERTY(BlueprintReadOnly, EditAnywhere, Category = Mover)
+	FVector MoveAcceleration = FVector::ZeroVector;
+
+	UPROPERTY(BlueprintReadOnly, EditAnywhere, Category = Mover)
+	float MoveSpeed = 0.f;
+};
+
+// Data about the water volume and its interaction with the pawn used in calculating swimming movement
+USTRUCT(BlueprintType)
+struct MOVER_API FWaterFlowSplineData
+{
+	GENERATED_USTRUCT_BODY()
+	
+	FWaterFlowSplineData()
+	{
+		Reset();
+	}
+
+	float SplineInputKey;
+	float ImmersionDepth;
+	float ImmersionPercent;
+	float WaterDepth;
+	float WaterVelocityDepthMultiplier;
+	float WaveAttenuationFactor;
+	float WaveReferenceTime;
+	FVector RawWaterVelocity;
+	FVector WaterVelocity;
+	FVector PlayerRelativeVelocityToWater;
+	FVector WaterSurfaceLocation;
+	FVector WaterSurfaceNormal;
+	FVector WaterSurfaceOffset;
+	FVector WaterPlaneLocation;
+	FVector WaterPlaneNormal;
+	TWeakObjectPtr<class AWaterBody> WaterBody;
+
+	bool IsValid() const { return SplineInputKey >= 0; }
+
+	void Reset()
+	{
+		SplineInputKey = -1.0f;
+		ImmersionDepth = -1.0f;
+		ImmersionPercent = 0.f;
+		WaterDepth = 0.f;
+		WaterVelocityDepthMultiplier = 1.0f;
+		WaveAttenuationFactor = 1.0f;
+		WaveReferenceTime = 0.0f;
+		RawWaterVelocity = FVector::ZeroVector;
+		WaterVelocity = FVector::ZeroVector;
+		PlayerRelativeVelocityToWater = FVector::ZeroVector;
+		WaterSurfaceLocation = FVector::ZeroVector;
+		WaterSurfaceNormal = FVector::UpVector;
+		WaterSurfaceOffset = FVector::ZeroVector;
+		WaterPlaneLocation = FVector::ZeroVector;
+		WaterPlaneNormal = FVector(0.0f, 0.0f, 1.0f);
+		WaterBody = nullptr;
+	}
+};
+
+/** Data about the overlapping volume typically used for swimming */
+USTRUCT(BlueprintType)
+struct MOVER_API FWaterCheckResult
+{
+	GENERATED_USTRUCT_BODY()
+	
+	/** True if the hit found a valid swimmable volume. */
+	UPROPERTY(VisibleInstanceOnly, BlueprintReadOnly, Category = CharacterVolume)
+	uint32 bSwimmableVolume : 1;
+	
+	/** Hit result of the test that found a volume. Includes more specific data about the point of impact and surface normal at that point. */
+	UPROPERTY(VisibleInstanceOnly, BlueprintReadOnly, Category = CharacterVolume)
+	FHitResult HitResult;
+
+	/** Water Spline data to be used in calculating swimming movement, FX, etc. */ 
+	UPROPERTY(VisibleInstanceOnly, BlueprintReadOnly, Category = CharacterVolume)
+	FWaterFlowSplineData WaterSplineData; 
+	
+public:
+	FWaterCheckResult()
+		: bSwimmableVolume(false)
+		, HitResult(-1.f)
+		, WaterSplineData()
+	{
+	}
+	
+	/** Returns true if the volume result hit a walkable surface. */
+	bool IsSwimmableVolume() const
+	{
+		return bSwimmableVolume;
+	}
+	
+	void Clear()
+	{
+		bSwimmableVolume = false;
+		HitResult.Reset(-1.f, false);
+	}
+};
+
+// Input parameters for Updating WaterSplineData
+USTRUCT(BlueprintType)
+struct MOVER_API FUpdateWaterSplineDataParams
+{
+	GENERATED_BODY()
+
+	UPROPERTY(BlueprintReadOnly, EditAnywhere, Category = Mover)
+	float TargetImmersionDepth = 0.f;
+	
+	UPROPERTY(BlueprintReadOnly, EditAnywhere, Category = Mover)
+	float WaterVelocityDepthForMax = 0.f;
+	
+	UPROPERTY(BlueprintReadOnly, EditAnywhere, Category = Mover)
+	float WaterVelocityMinMultiplier = 0.f;
+
+	UPROPERTY(BlueprintReadOnly, EditAnywhere, Category = Mover)
+	FVector PlayerVelocity = FVector::ZeroVector;
+	
+	UPROPERTY(BlueprintReadOnly, EditAnywhere, Category = Mover)
+	FVector PlayerLocation = FVector::ZeroVector;
+
+	UPROPERTY(BlueprintReadOnly, EditAnywhere, Category = Mover)
+	float CapsuleHalfHeight = 0.f;	
+};
+
+/**
+ * WaterMovementUtils: a collection of stateless static BP-accessible functions for a variety of water movement-related operations
+ */
+UCLASS()
+class MOVER_API UWaterMovementUtils : public UBlueprintFunctionLibrary
+{
+	GENERATED_BODY()
+
+public:
+	/** Generate a new movement based on move/orientation intents and the prior state for the swimming move */
+	UFUNCTION(BlueprintCallable, Category = Mover)
+	static FProposedMove ComputeControlledWaterMove(const FWaterMoveParams& InParams);
+	
+	/** Updates the water spline data used in calculated swimming movement */
+	UFUNCTION(BlueprintCallable, Category=Mover)
+	static void UpdateWaterSplineData(const FUpdateWaterSplineDataParams& UpdateWaterSplineDataParams, FWaterCheckResult& OutWaterResult);
+};
