Moving Mover, MoverExamples, MoverTests out of Restricted engine folder

diff --git a/Engine/Plugins/Experimental/Mover/Source/Mover/Public/MovementMode.h b/Engine/Plugins/Experimental/Mover/Source/Mover/Public/MovementMode.h
new file mode 100644
--- /dev/null
+++ b/Engine/Plugins/Experimental/Mover/Source/Mover/Public/MovementMode.h
@@ -0,0 +1,100 @@
+// Copyright Epic Games, Inc. All Rights Reserved.
+
+#pragma once
+
+#include "CoreMinimal.h"
+#include "MoverSimulationTypes.h"
+#include "MoverTypes.h"
+#include "MoveLibrary/MoverBlackboard.h"
+#include "MovementModeTransition.h"
+#include "UObject/Interface.h"
+#include "Templates/SubclassOf.h"
+#include "MovementMode.generated.h"
+
+
+/**
+ * UMovementSettingsInterface: interface that must be implemented for any settings object to be shared between modes
+ */
+UINTERFACE(MinimalAPI, BlueprintType)
+class UMovementSettingsInterface : public UInterface
+{
+	GENERATED_BODY()
+};
+
+class IMovementSettingsInterface
+{
+	GENERATED_BODY()
+
+public:
+	virtual FString GetDisplayName() const = 0;
+};
+
+
+
+/**
+ * Base class for all movement modes, exposing simulation update methods for both C++ and blueprint extension
+ */
+UCLASS(Abstract, Blueprintable, BlueprintType, EditInlineNew)
+class MOVER_API UBaseMovementMode : public UObject
+{
+	GENERATED_UCLASS_BODY()
+
+public:
+	
+	void DoRegister(const FName ModeName);
+	void DoUnregister();
+	void DoGenerateMove(const FMoverTickStartData& StartState, const FMoverTimeStep& TimeStep, FProposedMove& OutProposedMove) const;
+	void DoSimulationTick(const FSimulationTickParams& Params, FMoverTickEndData& OutputState);
+
+	UFUNCTION(BlueprintCallable, Category=Mover, meta=(DisplayName="Get Mover Component"))
+	UMoverComponent* GetMoverComponent() const;
+
+	UFUNCTION(BlueprintCallable, Category=Mover, meta=(DisplayName="Get Blackboard"))
+	const UMoverBlackboard* GetBlackboard() const;
+
+	UFUNCTION(BlueprintCallable, Category=Mover, meta=(DisplayName="Get Blackboard"))
+	UMoverBlackboard* GetBlackboard_Mutable() const;
+
+
+#if WITH_EDITOR
+	virtual EDataValidationResult IsDataValid(class FDataValidationContext& Context) const override;
+#endif // WITH_EDITOR
+
+	/** Settings object type that this mode depends on. May be shared with other movement modes. When the mode is added to a Mover Component, it will create a shared instance of this settings class. */
+	UPROPERTY(EditDefaultsOnly, Category = Mover, meta = (MustImplement = "/Script/Mover.MovementSettingsInterface"))
+	TSubclassOf<UObject> SharedSettingsClass = nullptr;
+
+	/** Transition checks for the current mode. Evaluated in order, stopping at the first successful transition check */
+	UPROPERTY(EditAnywhere, BlueprintReadOnly, Instanced, Category = Mover, meta = (FullyExpand = true))
+	TArray<TObjectPtr<UBaseMovementModeTransition>> Transitions;
+
+protected:
+
+	virtual void OnRegistered(const FName ModeName);
+
+	UFUNCTION(BlueprintImplementableEvent, meta = (DisplayName = "OnRegistered", ScriptName = "OnRegistered"))
+	void K2_OnRegistered(const FName ModeName);
+
+
+	virtual void OnUnregistered();
+
+	UFUNCTION(BlueprintImplementableEvent, meta = (DisplayName = "OnUnregistered", ScriptName = "OnUnregistered"))
+	void K2_OnUnregistered();
+
+
+	virtual void OnGenerateMove(const FMoverTickStartData& StartState, const FMoverTimeStep& TimeStep, FProposedMove& OutProposedMove) const;
+
+	UFUNCTION(BlueprintImplementableEvent, DisplayName = "OnGenerateMove", meta = (ScriptName = "OnGenerateMove"))
+	FProposedMove K2_OnGenerateMove(const FMoverTickStartData& StartState, const FMoverTimeStep& TimeStep) const;
+
+	virtual void OnSimulationTick(const FSimulationTickParams& Params, FMoverTickEndData& OutputState);
+	
+	UFUNCTION(BlueprintImplementableEvent, DisplayName = "OnSimulationTick", meta = (ScriptName = "OnSimulationTick"))
+	FMoverTickEndData K2_OnSimulationTick(const FSimulationTickParams& Params);
+
+
+
+private:
+	bool bHasBlueprintGenerateMove = false;
+	bool bHasBlueprintSimulationTick = false;
+};
\ No newline at end of file
diff --git a/Engine/Plugins/Experimental/Mover/Source/Mover/Public/MovementModeTransition.h b/Engine/Plugins/Experimental/Mover/Source/Mover/Public/MovementModeTransition.h
new file mode 100644
--- /dev/null
+++ b/Engine/Plugins/Experimental/Mover/Source/Mover/Public/MovementModeTransition.h
@@ -0,0 +1,73 @@
+// Copyright Epic Games, Inc. All Rights Reserved.
+
+#pragma once
+
+#include "CoreMinimal.h"
+#include "UObject/Interface.h"
+#include "MovementModeTransition.generated.h"
+
+struct FSimulationTickParams;
+
+/** 
+ * Results from a movement mode transition evaluation
+ */
+USTRUCT(BlueprintType)
+struct MOVER_API FTransitionEvalResult
+{
+	GENERATED_BODY()
+
+	FTransitionEvalResult() {}
+	FTransitionEvalResult(FName InNextMode) { NextMode = InNextMode; }
+
+	// Mode name that should be transitioned to. NAME_None indicates no transition.
+	UPROPERTY(BlueprintReadWrite, Category=Mover)
+	FName NextMode = NAME_None;
+
+
+	static const FTransitionEvalResult NoTransition;
+};
+
+
+/**
+ * Base class for all transitions
+ */
+UCLASS(Abstract, Blueprintable, BlueprintType, EditInlineNew)
+class MOVER_API UBaseMovementModeTransition : public UObject
+{
+	GENERATED_UCLASS_BODY()
+
+public:
+	// TODO : add hooks for shared settings, so that transitions can refer to the same settings as the modes without needing to know about the mode's implementation
+
+	void DoRegister();
+	void DoUnregister();
+
+	FTransitionEvalResult DoEvaluate(const FSimulationTickParams& Params) const;
+
+	void DoTrigger(const FSimulationTickParams& Params);
+
+#if WITH_EDITOR
+	virtual EDataValidationResult IsDataValid(class FDataValidationContext& Context) const override;
+#endif // WITH_EDITOR
+
+protected:
+
+	virtual FTransitionEvalResult OnEvaluate(const FSimulationTickParams& Params) const;
+
+	UFUNCTION(BlueprintImplementableEvent, DisplayName = "OnEvaluate", meta = (ScriptName = "OnEvaluate"))
+	FTransitionEvalResult K2_OnEvaluate(const FSimulationTickParams& Params) const;
+
+
+	virtual void OnTrigger(const FSimulationTickParams& Params);
+
+	UFUNCTION(BlueprintImplementableEvent, DisplayName = "OnTrigger", meta = (ScriptName = "OnTrigger"))
+	void K2_OnTrigger(const FSimulationTickParams& Params);
+
+private:
+	bool bHasBlueprintEvaluate = false;
+	bool bHasBlueprintTrigger = false;
+};
+
+
+
+
diff --git a/Engine/Plugins/Experimental/Mover/Source/Mover/Public/MoverBasePawn.h b/Engine/Plugins/Experimental/Mover/Source/Mover/Public/MoverBasePawn.h
new file mode 100644
--- /dev/null
+++ b/Engine/Plugins/Experimental/Mover/Source/Mover/Public/MoverBasePawn.h
@@ -0,0 +1,32 @@
+// Copyright Epic Games, Inc. All Rights Reserved.
+
+#pragma once
+
+#include "CoreMinimal.h"
+#include "GameFramework/Pawn.h"
+#include "MoverBasePawn.generated.h"
+
+class UMoverComponent;
+
+UCLASS()
+class MOVER_API AMoverBasePawn : public APawn
+{
+	GENERATED_BODY()
+
+public:
+	// Sets default values for this pawn's properties
+	AMoverBasePawn(const FObjectInitializer& ObjectInitializer);
+
+	/** Accessor for the actor's movement component */
+	UFUNCTION(BlueprintPure, Category = Mover)
+	UMoverComponent* GetMoverComponent() const { return CharacterMotionComponent; }
+
+protected:
+
+	virtual UPrimitiveComponent* GetMovementBase() const override;
+
+protected:
+	UPROPERTY(Category = Movement, VisibleAnywhere, BlueprintReadOnly, meta = (AllowPrivateAccess = "true"))
+	TObjectPtr<UMoverComponent> CharacterMotionComponent;
+
+};
diff --git a/Engine/Plugins/Experimental/Mover/Source/Mover/Public/MoverComponent.h b/Engine/Plugins/Experimental/Mover/Source/Mover/Public/MoverComponent.h
new file mode 100644
--- /dev/null
+++ b/Engine/Plugins/Experimental/Mover/Source/Mover/Public/MoverComponent.h
@@ -0,0 +1,379 @@
+// Copyright Epic Games, Inc. All Rights Reserved.
+
+#pragma once
+
+#include "CoreMinimal.h"
+#include "Components/ActorComponent.h"
+#include "MovementMode.h"
+#include "MoverTypes.h"
+#include "LayeredMove.h"
+#include "Engine/HitResult.h"
+#include "Backends/MoverBackendLiaison.h"
+#include "UObject/WeakInterfacePtr.h"
+#include "MoverComponent.generated.h"
+
+struct FMoverTimeStep;
+class UMovementModeStateMachine;
+
+namespace MoverComponentConstants
+{
+	extern const FVector DefaultGravityAccel;		// Fallback gravity if not determined by the component or world (cm/s^2)
+	extern const FVector DefaultUpDir;				// Fallback up direction if not determined by the component or world (normalized)
+}
+
+
+// Fired just before a simulation tick, regardless of being a re-simulated frame or not.
+DECLARE_DYNAMIC_MULTICAST_DELEGATE_TwoParams(FMover_OnPreSimTick, const FMoverTimeStep&, TimeStep, const FMoverInputCmdContext&, InputCmd);
+
+// Fired during a simulation tick, after movement has occurred but before the state is finalized, allowing changes to the output state.
+DECLARE_DYNAMIC_MULTICAST_DELEGATE_ThreeParams(FMover_OnPostMovement, const FMoverTimeStep&, TimeStep, FMoverSyncState&, SyncState, FMoverAuxStateContext&, AuxState);
+
+// Fired after a simulation tick, regardless of being a re-simulated frame or not.
+DECLARE_DYNAMIC_MULTICAST_DELEGATE_OneParam(FMover_OnPostSimTick, const FMoverTimeStep&, TimeStep);
+
+// Fired after a rollback. First param is the time step we've rolled back to. Second param is when we rolled back from, and represents a later frame that is no longer valid.
+DECLARE_DYNAMIC_MULTICAST_DELEGATE_TwoParams(FMover_OnPostSimRollback, const FMoverTimeStep&, CurrentTimeStep, const FMoverTimeStep&, ExpungedTimeStep);
+
+// Fired after changing movement modes. First param is the name of the previous movement mode. Second is the name of the new movement mode. 
+DECLARE_DYNAMIC_MULTICAST_DELEGATE_TwoParams(FMover_OnMovementModeChanged, const FName&, PreviousMovementModeName, const FName&, NewMovementModeName);
+
+// Fired after the actor hits a surface. Param is FMoverOnImpactParams struct that holds all of the data/params for an impact see @FMoverOnImpactParams for more details of it's properties.
+DECLARE_DYNAMIC_MULTICAST_DELEGATE_OneParam(FMover_OnImpact, const FMoverOnImpactParams&, ImpactParams);
+
+/**
+ * 
+ */
+UCLASS(BlueprintType, meta = (BlueprintSpawnableComponent))
+class MOVER_API UMoverComponent : public UActorComponent
+{
+	GENERATED_BODY()
+
+
+public:	
+	UMoverComponent();
+
+	virtual void InitializeComponent() override;
+	virtual void UninitializeComponent() override;
+	virtual void OnRegister() override;
+	virtual void RegisterComponentTickFunctions(bool bRegister) override;
+	virtual void PostLoad() override;
+	virtual void BeginPlay() override;
+
+	virtual void TickComponent(float DeltaTime, enum ELevelTick TickType, FActorComponentTickFunction* ThisTickFunction) override;
+	
+	// Broadcast before each simulation tick
+	UPROPERTY(BlueprintAssignable, Category = Mover)
+	FMover_OnPreSimTick OnPreSimulationTick;
+
+	// Broadcast at the end of a simulation tick after movement has occurred, but allowing additions/modifications to the state. Not assignable as a BP event due to it having mutable parameters.
+	UPROPERTY()
+	FMover_OnPostMovement OnPostMovement;
+
+	// Broadcast after each simulation tick and the state is finalized
+	UPROPERTY(BlueprintAssignable, Category = Mover)
+	FMover_OnPostSimTick OnPostSimulationTick;
+
+	// Broadcast when a rollback has occurred, just before the next simulation tick occurs
+	UPROPERTY(BlueprintAssignable, Category = Mover)
+	FMover_OnPostSimRollback OnPostSimulationRollback;
+	
+	// Broadcast when this actor impacts a surface.
+	UPROPERTY(BlueprintAssignable, Category = Mover)
+	FMover_OnImpact OnImpact;
+
+	// Broadcast when a MovementMode has changed. Happens during a simulation tick if the mode changed that tick or when SetModeImmediately is used to change modes.
+	UPROPERTY(BlueprintAssignable, Category = Mover)
+	FMover_OnMovementModeChanged OnMovementModeChanged;
+
+	// Callbacks
+	UFUNCTION()
+	virtual void OnBeginOverlap(UPrimitiveComponent* OverlappedComp, AActor* Other, UPrimitiveComponent* OtherComp, int32 OtherBodyIndex, bool bFromSweep, const FHitResult& SweepResult) { }
+
+	// --------------------------------------------------------------------------------
+	// NP Driver
+	// --------------------------------------------------------------------------------
+
+	// Get latest local input prior to simulation step. Called by Network Prediction system on owner's instance (autonomous or authority).
+	void ProduceInput(const int32 DeltaTimeMS, FMoverInputCmdContext* Cmd);
+
+	// Restore a previous frame prior to resimulating. Called by Network Prediction system.
+	void RestoreFrame(const FMoverSyncState* SyncState, const FMoverAuxStateContext* AuxState);
+
+	// Take output for simulation. Called by Network Prediction system.
+	void FinalizeFrame(const FMoverSyncState* SyncState, const FMoverAuxStateContext* AuxState);
+
+	// Seed initial values based on component's state. Called by Network Prediction system.
+	void InitializeSimulationState(FMoverSyncState* OutSync, FMoverAuxStateContext* OutAux);
+
+	// Primary movement simulation update. Given an starting state and timestep, produce a new state. Called by Network Prediction system.
+	void SimulationTick(const FMoverTimeStep& InTimeStep, const FMoverTickStartData& SimInput, OUT FMoverTickEndData& SimOutput);
+
+	// Specifies which supporting back end class should drive this Mover actor
+	UPROPERTY(EditDefaultsOnly, Category = Mover, meta = (MustImplement = "/Script/Mover.MoverBackendLiaisonInterface"))
+	TSubclassOf<UActorComponent> BackendClass;
+
+	UPROPERTY(EditAnywhere, BlueprintReadOnly, Instanced, Category = Mover, meta=(FullyExpand=true))
+	TMap<FName, TObjectPtr<UBaseMovementMode>> MovementModes;
+
+	// Name of the first mode to start in when simulation begins. Must have a mapping in MovementModes. Only used during initialization.
+	UPROPERTY(EditAnywhere, BlueprintReadOnly, Category=Mover)
+	FName StartingMovementMode;
+
+	/** List of types that should always be present in this actor's sync state */
+	UPROPERTY(BlueprintReadWrite, EditAnywhere, Category=Mover)
+	TArray<FMoverDataPersistence> PersistentSyncStateDataTypes;
+
+	/** Optional object for producing input cmds. Typically set at BeginPlay time. If not specified, defaulted input will be used. */
+	UPROPERTY(BlueprintReadWrite, EditAnywhere, Category = Mover, meta = (MustImplement = "/Script/Mover.MoverInputProducerInterface"))
+	TObjectPtr<UObject> InputProducer;
+
+	/**
+	 * Queue a layered move to start during the next simulation frame. This will clone whatever move you pass in, so you'll need to fully set it up before queuing.
+	 * @param LayeredMove			The move to queue, which must be a LayeredMoveBase sub-type. 
+	 */
+	UFUNCTION(BlueprintCallable, CustomThunk, Category = Mover, meta = (CustomStructureParam = "MoveAsRawData", AllowAbstract = "false", DisplayName = "Queue Layered Move"))
+	void K2_QueueLayeredMove(UPARAM(DisplayName="Layered Move") const int32& MoveAsRawData);
+	DECLARE_FUNCTION(execK2_QueueLayeredMove);
+
+	// Queue a layered move to start during the next simulation frame
+	void QueueLayeredMove(TSharedPtr<FLayeredMoveBase> Move);
+	
+	// Queue a movement mode change to occur during the next simulation frame. If bShouldReenter is true, then a mode change will occur even if already in that mode.
+	UFUNCTION(BlueprintCallable, Category = Mover)
+	void QueueNextMode(FName DesiredModeName, bool bShouldReenter=false);
+
+	// Add a movement mode to available movement modes. Returns true if the movement mode was added successfully. Returns the mode that was made.
+	UFUNCTION(BlueprintCallable, Category = Mover, meta=(DeterminesOutputType="MovementMode"))
+	UBaseMovementMode* AddMovementModeFromClass(FName ModeName, TSubclassOf<UBaseMovementMode> MovementMode);
+
+	// Add a movement mode to available movement modes. Returns true if the movement mode was added successfully
+	UFUNCTION(BlueprintCallable, Category = Mover)
+	bool AddMovementModeFromObject(FName ModeName, UBaseMovementMode* MovementMode);
+	
+	// Removes a movement mode from available movement modes. Returns number of modes removed from the available movement modes.
+	UFUNCTION(BlueprintCallable, Category = Mover)
+	bool RemoveMovementMode(FName ModeName);
+	
+public:
+	// Set gravity override, as a directional acceleration in worldspace.  Gravity on Earth would be {x=0,y=0,z=-980}
+	UFUNCTION(BlueprintCallable, Category = Mover)
+	void SetGravityOverride(bool bOverrideGravity, FVector GravityAcceleration=FVector::ZeroVector);
+	
+	// Get the current acceleration due to gravity (cm/s^2) in worldspace
+	UFUNCTION(BlueprintCallable, BlueprintPure = false, Category = Mover)
+	FVector GetGravityAcceleration() const;
+
+	// Get the normalized direction considered "up" in worldspace. Typically aligned with gravity, and typically determines the plane an actor tries to move along.
+	UFUNCTION(BlueprintCallable, BlueprintPure = false, Category = Mover)
+	FVector GetUpDirection() const;
+
+public:	// Queries
+
+	// Get the transform of the root component that our Mover simulation is moving
+	FTransform GetUpdatedComponentTransform() const;
+
+	// Access the primary visual component of the actor
+	USceneComponent* GetPrimaryVisualComponent() const;
+
+	// Typed accessor to primary visual component
+	template<class T>
+	T* GetPrimaryVisualComponent() const
+	{
+		return Cast<T>(GetPrimaryVisualComponent());
+	}
+
+	// Get the current velocity (units per second, worldspace)
+	UFUNCTION(BlueprintPure, Category = Mover)
+	FVector GetVelocity() const;
+
+	// Get the intended movement direction in worldspace with magnitude (range 0-1)
+	UFUNCTION(BlueprintPure, Category = Mover)
+	FVector GetMovementIntent() const;
+
+	// Get the orientation that the actor is moving towards
+	UFUNCTION(BlueprintPure, Category = Mover)
+	FRotator GetTargetOrientation() const;
+
+	// Is this actor in a falling state? Note that this includes upwards motion induced by jumping.
+	UFUNCTION(BlueprintPure, Category = Mover)
+	virtual bool IsFalling() const;
+
+	// Is this actor in a airborne state? (e.g. Flying, Falling)
+	UFUNCTION(BlueprintPure, Category = Mover)
+	virtual bool IsAirborne() const;
+
+	// Is this actor in a grounded state? (e.g. Walking)
+	UFUNCTION(BlueprintPure, Category = Mover)
+	virtual bool IsOnGround() const;
+
+	// Is this actor sliding on an unwalkable slope
+	UFUNCTION(BlueprintPure, Category = Mover)
+	virtual bool IsSlopeSliding() const;
+
+	/** Get a sampling of where the actor is projected to be in the future, based on a current state. Note that this is projecting ideal movement without doing full simulation and collision. */
+	UFUNCTION(BlueprintCallable, BlueprintPure = false, Category = Mover)
+	TArray<FTrajectorySampleInfo> GetFutureTrajectory(float FutureSeconds, float SamplesPerSecond) const;
+
+	// Get the current movement mode name
+	UFUNCTION(BlueprintPure, Category = Mover)
+	FName GetMovementModeName() const;
+
+	// Get the current movement base. Null if there isn't one.
+	UFUNCTION(BlueprintPure, Category = Mover)
+	UPrimitiveComponent* GetMovementBase() const;
+
+	// Get the current movement base bone, NAME_None if there isn't one.
+	UFUNCTION(BlueprintPure, Category = Mover)
+	FName GetMovementBaseBoneName() const;
+
+	// Signals whether we have a sync state saved yet. If not, most queries will not be meaningful.
+	UFUNCTION(BlueprintPure, Category = Mover)
+	bool HasValidCachedState() const;
+
+	// Access the most recent captured sync state. Check @HasValidCachedState first.
+	UFUNCTION(BlueprintPure, Category = Mover)
+	const FMoverSyncState& GetSyncState() const;
+
+	// Signals whether we have input data saved yet. If not, input queries will not be meaningful.
+	UFUNCTION(BlueprintPure, Category = Mover)
+	bool HasValidCachedInputCmd() const;
+
+	// Access the most recently-used inputs. Check @HasValidCachedInputCmd first.
+	UFUNCTION(BlueprintPure, Category = Mover)
+	const FMoverInputCmdContext& GetLastInputCmd() const;
+
+	// Access the most recent floor check hit result.
+	UFUNCTION(BlueprintPure, Category = Mover)
+	bool TryGetFloorCheckHitResult(FHitResult& OutHitResult) const;
+
+	// Access the read-only version of the Mover's Blackboard
+	UFUNCTION(BlueprintPure, Category=Mover)
+	const UMoverBlackboard* GetSimBlackboard() const;
+
+	/** Find settings object by type. Returns null if there is none of that type */
+	const IMovementSettingsInterface* FindSharedSettings(const UClass* ByType) const { return FindSharedSettings_Mutable(ByType); }
+	template<class T>
+	const T* FindSharedSettings() const { return Cast<const T>(FindSharedSettings(T::StaticClass())); }
+
+	/** Find mutable settings object by type. Returns null if there is none of that type */
+	IMovementSettingsInterface* FindSharedSettings_Mutable(const UClass* ByType) const;
+	template<class T>
+	T* FindSharedSettings_Mutable() const { return Cast<T>(FindSharedSettings_Mutable(T::StaticClass())); }
+
+	/** Find mutable settings object by type. Returns null if there is none of that type */
+	UFUNCTION(BlueprintPure, Category = Mover,  meta=(DeterminesOutputType="SharedSetting", DisplayName="Find Shared Settings Mutable"))
+	UObject* FindSharedSettings_Mutable_BP(UPARAM(meta = (MustImplement = "MovementSettingsInterface")) TSubclassOf<UObject> SharedSetting) const;
+
+	/** Find settings object by type. Returns null if there is none of that type */
+	UFUNCTION(BlueprintPure, Category = Mover,  meta=(DeterminesOutputType="SharedSetting", DisplayName="Find Shared Settings"))
+	const UObject* FindSharedSettings_BP(UPARAM(meta = (MustImplement = "MovementSettingsInterface")) TSubclassOf<UObject> SharedSetting) const;
+
+	/** Find movement mode by type. Returns null if there is none of that type */
+	UBaseMovementMode* FindMode_Mutable(const UClass* ByType, bool bRequireExactClass=false) const;
+	template<class T>
+	T* FindMode_Mutable(bool bRequireExactClass=false) const { return Cast<T>(FindMode_Mutable(T::StaticClass(), bRequireExactClass)); }
+
+	UFUNCTION(BlueprintPure, Category = Mover,  meta=(DeterminesOutputType="MovementMode"))
+	UBaseMovementMode* FindMovementMode(TSubclassOf<UBaseMovementMode> MovementMode) const;
+
+protected:
+
+	/** Makes this component and owner actor reflect the state of a particular frame snapshot. This occurs after simulation ticking, as well as during a rollback before we resimulate forward.
+	  @param bRebaseBasedState	If true and the state was using based movement, it will use the current game world base pos/rot instead of the captured one. This is necessary during rollbacks.
+	*/
+	void SetFrameStateFromContext(const FMoverSyncState* SyncState, const FMoverAuxStateContext* AuxState, bool bRebaseBasedState);
+
+public:
+	bool InitMoverSimulation();
+
+	/** Handle a blocking impact.*/
+	UFUNCTION(BlueprintCallable, Category = Mover)
+	void HandleImpact(FMoverOnImpactParams& ImpactParams);
+
+protected:
+	// Basic "Update Component/Ticking"
+	void SetUpdatedComponent(USceneComponent* NewUpdatedComponent);
+	void UpdateTickRegistration();
+
+	/** Called when a rollback occurs, after the simulation state has been restored */
+	void OnSimulationRollback(const FMoverSyncState* SyncState, const FMoverAuxStateContext* AuxState);
+
+	void ProcessFirstSimTickAfterRollback(const FMoverTimeStep& TimeStep);
+
+	#if WITH_EDITOR
+	virtual void PostCDOCompiled(const FPostCDOCompiledContext& Context) override;
+	virtual void PostEditChangeProperty(FPropertyChangedEvent& PropertyChangedEvent) override;
+	virtual void PostTransacted(const FTransactionObjectEvent& TransactionEvent) override;
+	bool ValidateSetup(class FDataValidationContext& ValidationErrors) const;
+	virtual EDataValidationResult IsDataValid(class FDataValidationContext& Context) const override;
+	#endif // WITH_EDITOR
+
+	UFUNCTION()
+	virtual void PhysicsVolumeChanged(class APhysicsVolume* NewVolume);
+
+	virtual void OnHandleImpact(const FMoverOnImpactParams& ImpactParams);
+	
+	/** internal function to ensure SharedSettings array matches what's needed by the list of Movement Modes */
+	void RefreshSharedSettings();
+
+	/** This is the component that's actually being moved. Typically it is the Actor's root component and often a collidable primitive. */
+	UPROPERTY(Transient)
+	TObjectPtr<USceneComponent> UpdatedComponent = nullptr;
+
+	/** UpdatedComponent, cast as a UPrimitiveComponent. May be invalid if UpdatedComponent was null or not a UPrimitiveComponent. */
+	UPROPERTY(Transient)
+	TObjectPtr<UPrimitiveComponent> UpdatedCompAsPrimitive = nullptr;
+
+	/** The main visual component associated with this Mover actor, typically a mesh and typically parented to the UpdatedComponent. */
+	UPROPERTY(Transient)
+	TObjectPtr<USceneComponent> PrimaryVisualComponent;
+
+	bool bHasValidLastProducedInput = false;
+	FMoverInputCmdContext CachedLastProducedInputCmd;
+
+	bool bHasValidCachedUsedInput = false;
+	FMoverInputCmdContext CachedLastUsedInputCmd;
+	
+	bool bHasValidCachedState = false;
+	FMoverSyncState CachedLastSyncState;
+	FMoverAuxStateContext CachedLastAuxState;
+
+	FMoverTimeStep CachedLastSimTickTimeStep;	// Saved timestep info from our last simulation tick, used during rollback handling. This will rewind during corrections.
+	FMoverTimeStep CachedNewestSimTickTimeStep;	// Saved timestep info from the newest (farthest-advanced) simulation tick. This will not rewind during corrections.
+
+private:
+	/** Collection of settings objects that are shared between movement modes. This list is automatically managed based on the @MovementModes contents. */
+	UPROPERTY(EditDefaultsOnly, EditFixedSize, Instanced, Category = Mover, meta = (NoResetToDefault, MustImplement = "/Script/Mover.MovementSettingsInterface"))
+	TArray<TObjectPtr<UObject>> SharedSettings;
+
+	// Whether or not gravity is overridden on this actor. Otherwise, fall back on world settings. See @SetGravityOverride
+	UPROPERTY(EditDefaultsOnly, Category="Mover|Gravity")
+	bool bHasGravityOverride = false;
+	
+	// cm/s^2, only meaningful if @bHasGravityOverride is enabled. Set @SetGravityOverride
+	UPROPERTY(EditDefaultsOnly, Category="Mover|Gravity", meta=(ForceUnits = "cm/s^2"))
+	FVector GravityAccelOverride;
+
+	/** Transient flag indicating whether we are executing OnRegister(). */
+	bool bInOnRegister = false;
+
+	/** Transient flag indicating whether we are executing InitializeComponent(). */
+	bool bInInitializeComponent = false;
+
+	// Transient flag indicating we've had a rollback and haven't started simulating forward again yet
+	bool bHasRolledBack = false;
+
+	UPROPERTY(Transient)
+	TObjectPtr<UMovementModeStateMachine> ModeFSM;	// JAH TODO: Also consider allowing a type property on the component to allow an alternative machine implementation to be allocated/used
+
+	/** Used to store cached data & computations between decoupled systems, that can be referenced by name */
+	UPROPERTY(Transient)
+	TObjectPtr<UMoverBlackboard> SimBlackboard;
+
+	
+	TWeakInterfacePtr<IMoverBackendLiaisonInterface> BackendLiaisonComp;
+
+	friend class UBaseMovementMode;
+	friend class UMoverNetworkPhysicsLiaisonComponent;
+	friend class UMoverDebugComponent;
+};
diff --git a/Engine/Plugins/Experimental/Mover/Source/Mover/Public/MoverDataModelTypes.h b/Engine/Plugins/Experimental/Mover/Source/Mover/Public/MoverDataModelTypes.h
new file mode 100644
--- /dev/null
+++ b/Engine/Plugins/Experimental/Mover/Source/Mover/Public/MoverDataModelTypes.h
@@ -0,0 +1,276 @@
+// Copyright Epic Games, Inc. All Rights Reserved.
+
+#pragma once
+
+#include "MoverTypes.h"
+#include "LayeredMove.h"
+#include "MoverDataModelTypes.generated.h"
+
+
+
+
+
+// Used to identify how to interpret a movement input vector's values
+UENUM(BlueprintType)
+enum class EMoveInputType : uint8
+{
+	Invalid,
+
+	/** Move with intent, as a per-axis magnitude [-1,1] (E.g., "move straight forward as fast as possible" would be (1, 0, 0) and "move straight left at half speed" would be (0, -0.5, 0) regardless of frame time). Zero vector indicates intent to stop. */
+	DirectionalIntent,
+
+	/** Move with a given velocity (units per second) */
+	Velocity,
+};
+
+
+// Data block containing all inputs that need to be authored and consumed for the kinematic Mover simulation
+USTRUCT(BlueprintType)
+struct MOVER_API FKinematicDefaultInputs : public FMoverDataStructBase
+{
+	GENERATED_USTRUCT_BODY()
+
+	// Sets the directional move inputs for a simulation frame
+	void SetMoveInput(EMoveInputType InMoveInputType, const FVector& InMoveInput);
+
+	const FVector& GetMoveInput() const { return MoveInput; }
+	EMoveInputType GetMoveInputType() const { return MoveInputType; }
+
+protected:
+	UPROPERTY(BlueprintReadWrite, Category = Mover)
+	EMoveInputType MoveInputType;
+
+	// Representing the directional move input for this frame. Must be interpreted according to MoveInputType. Relative to MovementBase if set, world space otherwise. Will be truncated to match network serialization precision.
+	UPROPERTY(BlueprintReadWrite, Category = Mover)
+	FVector MoveInput;
+
+public:
+	// Facing direction intent, as a normalized forward-facing direction. A zero vector indicates no intent to change facing direction. Relative to MovementBase if set, world space otherwise.
+	UPROPERTY(BlueprintReadWrite, Category = Mover)
+	FVector OrientationIntent;
+
+	// World space orientation that the controls were based on. This is commonly a player's camera rotation.
+	UPROPERTY(BlueprintReadWrite, Category = Mover)
+	FRotator ControlRotation;
+
+	// Used to force the Mover actor into a different movement mode
+	UPROPERTY(BlueprintReadWrite, Category = Mover)
+	FName SuggestedMovementMode;
+
+	// Specifies whether we are using a movement base, which will affect how move inputs are interpreted
+	UPROPERTY(BlueprintReadWrite, Category = Mover)
+	bool bUsingMovementBase;
+
+	// Optional: when moving on a base, input may be relative to this object
+	UPROPERTY(BlueprintReadWrite, Category = Mover)
+	TObjectPtr<UPrimitiveComponent> MovementBase;
+
+	// Optional: for movement bases that are skeletal meshes, this is the bone we're based on. Only valid if MovementBase is set.
+	UPROPERTY(BlueprintReadWrite, Category = Mover)
+	FName MovementBaseBoneName;
+
+	UPROPERTY(BlueprintReadWrite, Category = Mover)
+	bool bIsJumpJustPressed;
+	UPROPERTY(BlueprintReadWrite, Category = Mover)
+	bool bIsJumpPressed;
+
+	FVector GetMoveInput_WorldSpace() const;
+	FVector GetOrientationIntentDir_WorldSpace() const;
+
+
+	FKinematicDefaultInputs()
+		: MoveInputType(EMoveInputType::Invalid)
+		, MoveInput(ForceInitToZero)
+		, OrientationIntent(ForceInitToZero)
+		, ControlRotation(ForceInitToZero)
+		, SuggestedMovementMode(NAME_None)
+		, bUsingMovementBase(false)
+		, MovementBase(nullptr)
+		, MovementBaseBoneName(NAME_None)
+		, bIsJumpJustPressed(false)
+		, bIsJumpPressed(false)
+	{
+	}
+
+	virtual ~FKinematicDefaultInputs() {}
+
+	// @return newly allocated copy of this FKinematicDefaultInputs. Must be overridden by child classes
+	virtual FMoverDataStructBase* Clone() const override;
+
+	virtual bool NetSerialize(FArchive& Ar, UPackageMap* Map, bool& bOutSuccess) override;
+
+	virtual UScriptStruct* GetScriptStruct() const override { return StaticStruct(); }
+
+	virtual void ToString(FAnsiStringBuilderBase& Out) const override;
+
+	virtual void AddReferencedObjects(FReferenceCollector& Collector) override { Super::AddReferencedObjects(Collector); }
+
+
+};
+
+template<>
+struct TStructOpsTypeTraits< FKinematicDefaultInputs > : public TStructOpsTypeTraitsBase2< FKinematicDefaultInputs >
+{
+	enum
+	{
+		WithNetSerializer = true,
+		WithCopy = true
+	};
+};
+
+
+
+
+// Data block containing basic sync state information
+USTRUCT(BlueprintType)
+struct MOVER_API FMoverDefaultSyncState : public FMoverDataStructBase
+{
+	GENERATED_USTRUCT_BODY()
+protected:
+	// Position relative to MovementBase if set, world space otherwise
+	UPROPERTY(BlueprintReadWrite, EditAnywhere, Category = Mover)
+	FVector Location;
+
+	// Forward-facing rotation relative to MovementBase if set, world space otherwise.
+	UPROPERTY(BlueprintReadWrite, EditAnywhere, Category = Mover)
+	FRotator Orientation;
+
+	// Linear velocity, units per second, relative to MovementBase if set, world space otherwise.
+	UPROPERTY(BlueprintReadWrite, EditAnywhere, Category = Mover)
+	FVector Velocity;
+
+public:
+	// Movement intent direction relative to MovementBase if set, world space otherwise. Magnitude of range (0-1)
+	UPROPERTY(BlueprintReadWrite, EditAnywhere, Category = Mover)
+	FVector MoveDirectionIntent;
+
+	// The mode we ended up in from the prior frame, and which we'll start in during the next frame
+	UPROPERTY(BlueprintReadWrite, EditAnywhere, Category = Mover)
+	FName MovementMode;
+
+protected:
+	// Optional: when moving on a base, input may be relative to this object
+	UPROPERTY(BlueprintReadOnly, Category = Mover)
+	TObjectPtr<UPrimitiveComponent> MovementBase;
+
+	// Optional: for movement bases that are skeletal meshes, this is the bone we're based on. Only valid if MovementBase is set.
+	UPROPERTY(BlueprintReadOnly, Category = Mover)
+	FName MovementBaseBoneName;
+
+	UPROPERTY(BlueprintReadOnly, Category = Mover)
+	FVector MovementBasePos;
+
+	UPROPERTY(BlueprintReadOnly, Category = Mover)
+	FQuat MovementBaseQuat;
+
+public:
+	// Additional moves influencing our proposed motion
+	UPROPERTY(BlueprintReadWrite, EditAnywhere, Category = Mover)
+	FLayeredMoveGroup LayeredMoves;
+
+public:
+
+	FMoverDefaultSyncState()
+		: Location(ForceInitToZero)
+		, Orientation(ForceInitToZero)
+		, Velocity(ForceInitToZero)
+		, MoveDirectionIntent(ForceInitToZero)
+		, MovementMode(NAME_None)
+		, MovementBase(nullptr)
+		, MovementBaseBoneName(NAME_None)
+		, MovementBasePos(ForceInitToZero)
+		, MovementBaseQuat(FQuat::Identity)
+	{
+	}
+
+	virtual ~FMoverDefaultSyncState() {}
+
+	// @return newly allocated copy of this FMoverDefaultSyncState. Must be overridden by child classes
+	virtual FMoverDataStructBase* Clone() const override;
+
+	virtual bool NetSerialize(FArchive& Ar, UPackageMap* Map, bool& bOutSuccess) override;
+
+	virtual UScriptStruct* GetScriptStruct() const override { return StaticStruct(); }
+
+	virtual void ToString(FAnsiStringBuilderBase& Out) const override;
+
+	virtual bool ShouldReconcile(const FMoverDataStructBase& AuthorityState) const override;
+
+	virtual void Interpolate(const FMoverDataStructBase& From, const FMoverDataStructBase& To, float Pct) override;
+
+	void SetTransforms_WorldSpace(FVector WorldLocation, FRotator WorldOrient, FVector WorldVelocity, UPrimitiveComponent* Base=nullptr, FName BaseBone = NAME_None);
+
+	// Returns whether the base setting succeeded
+	bool SetMovementBase(UPrimitiveComponent* Base, FName BaseBone=NAME_None);
+
+	// Refreshes captured movement base transform based on its current state, while maintaining the same base-relative transforms
+	bool UpdateCurrentMovementBase();
+
+	// Queries
+	UPrimitiveComponent* GetMovementBase() const { return MovementBase; }
+	FName GetMovementBaseBoneName() const { return MovementBaseBoneName; }
+	FVector GetCapturedMovementBasePos() const { return MovementBasePos; }
+	FQuat GetCapturedMovementBaseQuat() const { return MovementBaseQuat; }
+
+	FVector GetLocation_WorldSpace() const;
+	FVector GetLocation_BaseSpace() const;	// If there is no movement base set, these will be the same as world space
+
+	FVector GetIntent_WorldSpace() const;
+	FVector GetIntent_BaseSpace() const;
+
+	FVector GetVelocity_WorldSpace() const;
+	FVector GetVelocity_BaseSpace() const;
+
+	FRotator GetOrientation_WorldSpace() const;
+	FRotator GetOrientation_BaseSpace() const;
+
+};
+
+template<>
+struct TStructOpsTypeTraits< FMoverDefaultSyncState > : public TStructOpsTypeTraitsBase2< FMoverDefaultSyncState >
+{
+	enum
+	{
+		WithNetSerializer = true,
+		WithCopy = true
+	};
+};
+
+
+/**
+ * Blueprint function library to make it easier to work with Mover data structs, since we can't add UFUNCTIONs to structs
+ */
+UCLASS()
+class MOVER_API UMoverDataModelBlueprintLibrary : public UBlueprintFunctionLibrary
+{
+	GENERATED_BODY()
+
+public:	// FKinematicDefaultInputs
+
+	/** Sets move inputs from worldspace intent, as a per-axis magnitude in the range [-1,1] Zero vector indicates intent to stop. */
+	UFUNCTION(BlueprintCallable, Category = Mover)
+	static void SetMoveIntent(UPARAM(Ref) FKinematicDefaultInputs& Inputs, const FVector& WorldDirectionIntent);
+
+	/** Returns the move direction intent, if any, in world space */
+	UFUNCTION(BlueprintCallable, Category = Mover)
+	static FVector GetMoveDirectionIntentFromInputs(const FKinematicDefaultInputs& Inputs);
+
+
+public:	// FMoverDefaultSyncState
+
+	/** Returns the location in world space */
+	UFUNCTION(BlueprintCallable, Category = Mover)
+	static FVector GetLocationFromSyncState(const FMoverDefaultSyncState& SyncState);
+
+	/** Returns the move direction intent, if any, in world space */
+	UFUNCTION(BlueprintCallable, Category = Mover)
+	static FVector GetMoveDirectionIntentFromSyncState(const FMoverDefaultSyncState& SyncState);
+
+	/** Returns the velocity in world space */
+	UFUNCTION(BlueprintCallable, Category = Mover)
+	static FVector GetVelocityFromSyncState(const FMoverDefaultSyncState& SyncState);
+
+	/** Returns the orientation in world space */
+	UFUNCTION(BlueprintCallable, Category = Mover)
+	static FRotator GetOrientationFromSyncState(const FMoverDefaultSyncState& SyncState);
+};
\ No newline at end of file
diff --git a/Engine/Plugins/Experimental/Mover/Source/Mover/Public/MoverTypes.h b/Engine/Plugins/Experimental/Mover/Source/Mover/Public/MoverTypes.h
new file mode 100644
--- /dev/null
+++ b/Engine/Plugins/Experimental/Mover/Source/Mover/Public/MoverTypes.h
@@ -0,0 +1,307 @@
+// Copyright Epic Games, Inc. All Rights Reserved.
+
+#pragma once
+
+#include "Misc/StringBuilder.h"
+#include "Engine/HitResult.h"
+#include "NetworkPredictionTickState.h"
+#include "MoverLog.h"
+#include "Kismet/BlueprintFunctionLibrary.h"
+#include "MoverTypes.generated.h"
+
+// Struct to hold params for when an impact happens. This contains all of the data for impacts including what gets passed to the FMover_OnImpact delegate
+USTRUCT(BlueprintType, meta = (DisplayName = "Impact Data"))
+struct MOVER_API FMoverOnImpactParams
+{
+	GENERATED_BODY()
+	
+	FMoverOnImpactParams();
+	
+	FMoverOnImpactParams(const FName& ModeName, const FHitResult& Hit, const FVector& Delta);
+	
+	// Name of the movement mode this actor is currently in at the time of the impact
+	UPROPERTY(EditAnywhere, BlueprintReadWrite, Category = Mover)
+	FName MovementModeName;
+	
+	// The hit result of the impact
+	UPROPERTY(EditAnywhere, BlueprintReadWrite, Category = Mover)
+	FHitResult HitResult;
+
+	// The original move that was being performed when the impact happened
+	UPROPERTY(EditAnywhere, BlueprintReadWrite, Category = Mover)
+	FVector AttemptedMoveDelta;
+};
+
+
+USTRUCT(BlueprintType)
+struct MOVER_API FMoverTimeStep
+{
+	GENERATED_BODY()
+	
+	FMoverTimeStep() 
+	{ 
+		bIsResimulating=false; 
+	}
+	
+	FMoverTimeStep(const FNetSimTimeStep& InNetSimTimeStep, bool InIsResimulating)
+		: ServerFrame(InNetSimTimeStep.Frame)
+		, BaseSimTimeMs(InNetSimTimeStep.TotalSimulationTime)
+		, StepMs(InNetSimTimeStep.StepMS)
+		, bIsResimulating(InIsResimulating)
+	{
+	}
+
+	FMoverTimeStep(const FNetSimTimeStep& InNetSimTimeStep)
+		: FMoverTimeStep(InNetSimTimeStep, false)
+	{
+	}
+
+	// The server simulation frame this timestep is associated with
+	UPROPERTY(EditAnywhere, BlueprintReadWrite, Category=Mover)
+	int32 ServerFrame=-1;
+
+	// Starting simulation time (in server simulation timespace)
+	UPROPERTY(EditAnywhere, BlueprintReadWrite, Category=Mover)
+	float BaseSimTimeMs=-1.f;
+
+	// The delta time step for this tick
+	UPROPERTY(EditAnywhere, BlueprintReadWrite, Category=Mover)
+	float StepMs=1.f;
+
+	// Indicates whether this time step is re-simulating based on prior inputs, such as during a correction
+	uint8 bIsResimulating : 1;
+
+};
+
+
+
+// Base type for all data structs used to compose Mover simulation model definition dynamically (input cmd, sync state, aux state)
+// NOTE: for simulation state data (sync/aux), derive from FMoverStateData instead 
+USTRUCT(BlueprintType)
+struct MOVER_API FMoverDataStructBase
+{
+	GENERATED_USTRUCT_BODY()
+
+	FMoverDataStructBase();
+	virtual ~FMoverDataStructBase() {}
+
+	/** return newly allocated copy of this FMoverDataStructBase. MUST be overridden by derived types. */
+	virtual FMoverDataStructBase* Clone() const;
+
+	virtual bool NetSerialize(FArchive& Ar, UPackageMap* Map, bool& bOutSuccess)
+	{
+		bOutSuccess = true;
+		return true;
+	}
+
+	/** Gets the type info of this FMoverDataStructBase. MUST be overridden by derived types. */
+	virtual UScriptStruct* GetScriptStruct() const;
+
+	/** Get string representation of this struct instance */
+	virtual void ToString(FAnsiStringBuilderBase& Out) const {}
+
+	/** If derived classes hold any object references, override this function and add them to the collector. */
+	virtual void AddReferencedObjects(FReferenceCollector& Collector) {}
+
+	/** Checks if the contained data is equal, within reason. MUST be override by types that compose STATE data (sync or aux). 
+	 *   AuthorityState is guaranteed to be the same concrete type as 'this'.
+	 */
+	virtual bool ShouldReconcile(const FMoverDataStructBase& AuthorityState) const;
+
+	/** Interpolates contained data between a starting and ending block. MUST be override by types that compose STATE data (sync or aux). 
+	 * From and To are guaranteed to be the same concrete type as 'this'.
+	 */
+	virtual void Interpolate(const FMoverDataStructBase& From, const FMoverDataStructBase& To, float Pct);
+
+};
+
+template<>
+struct TStructOpsTypeTraits< FMoverDataStructBase > : public TStructOpsTypeTraitsBase2< FMoverDataStructBase >
+{
+	enum
+	{
+		WithNetSerializer = true,
+		WithCopy = true
+	};
+};
+
+
+
+
+// Contains a group of different FMoverDataStructBase-derived data, and supports net serialization of them. Note that
+//	each contained data must have a unique type.  This is to support dynamic composition of Mover simulation model
+//  definitions (input cmd, sync state, aux state).
+USTRUCT(BlueprintType)
+struct MOVER_API FMoverDataCollection
+{
+	GENERATED_USTRUCT_BODY()
+
+	FMoverDataCollection();
+	virtual ~FMoverDataCollection() {}
+
+	void Empty() { DataArray.Empty(); }
+
+	/** Serialize all data in this collection */
+	bool NetSerialize(FArchive& Ar, UPackageMap* Map, bool& bOutSuccess);
+
+	/** Copy operator - deep copy so it can be used for archiving/saving off data */
+	FMoverDataCollection& operator=(const FMoverDataCollection& Other);
+
+	/** Comparison operator (deep) - needs matching struct types along with identical states in those structs. See also ShouldReconcile */
+	bool operator==(const FMoverDataCollection& Other) const;
+
+	/** Comparison operator */
+	bool operator!=(const FMoverDataCollection& Other) const;
+
+	/** Checks if the collections are significantly different enough (piece-wise) to need reconciliation. NOT an equality check. */
+	bool ShouldReconcile(const FMoverDataCollection& Other) const;
+
+	/** Make this collection a piece-wise interpolation between 2 collections */
+	void Interpolate(const FMoverDataCollection& From, const FMoverDataCollection& To, float Pct);
+
+	/** Exposes references to GC system */
+	void AddStructReferencedObjects(FReferenceCollector& Collector) const;
+
+	/** Get string representation of all elements in this collection */
+	void ToString(FAnsiStringBuilderBase& Out) const;
+
+	/** Find data of a specific type in the collection (mutable version). If not found, null will be returned. */
+	template <typename T>
+	T* FindMutableDataByType() const
+	{
+		if (FMoverDataStructBase* FoundData = FindDataByType(T::StaticStruct()))
+		{
+			return static_cast<T*>(FoundData);
+		}
+
+		return nullptr;
+	}
+
+	/** Find data of a specific type in the collection. If not found, null will be returned. */
+	template <typename T>
+	const T* FindDataByType() const
+	{
+		if (const FMoverDataStructBase* FoundData = FindDataByType(T::StaticStruct()))
+		{
+			return static_cast<const T*>(FoundData);
+		}
+
+		return nullptr;
+	}
+
+	/** Find data of a specific type in the collection. If not found, a new default instance will be added. */
+	template <typename T>
+	const T& FindOrAddDataByType()
+	{
+		if (const T* ExistingData = FindDataByType<T>())
+		{
+			return *ExistingData;
+		}
+
+		return *(static_cast<const T*>(AddDataByType(T::StaticStruct())));
+	}
+
+	/** Find data of a specific type in the collection (mutable version). If not found, a new default instance will be added. */
+	template <typename T>
+	T& FindOrAddMutableDataByType()
+	{
+		if (T* ExistingData = FindMutableDataByType<T>())
+		{
+			return *ExistingData;
+		}
+
+		return *(static_cast<T*>(AddDataByType(T::StaticStruct())));
+	}
+
+	/** Adds data to the collection. If an existing data struct of the same type is already there, it will be removed first. */
+	void AddOrOverwriteData(const TSharedPtr<FMoverDataStructBase> DataInstance);
+
+protected:
+	static TSharedPtr<FMoverDataStructBase> CreateDataByType(const UScriptStruct* DataStructType);
+	FMoverDataStructBase* AddDataByType(const UScriptStruct* DataStructType);
+
+	FMoverDataStructBase* FindDataByType(const UScriptStruct* DataStructType) const;
+	FMoverDataStructBase* FindOrAddDataByType(const UScriptStruct* DataStructType);
+	bool RemoveDataByType(const UScriptStruct* DataStructType);
+
+	/** Helper function for serializing array of data */
+	static void NetSerializeDataArray(FArchive& Ar, TArray<TSharedPtr<FMoverDataStructBase>>& DataArray);
+
+	/** All data in this collection */
+	TArray< TSharedPtr<FMoverDataStructBase> > DataArray;
+
+
+friend class UMoverDataCollectionLibrary;
+friend class UMoverComponent;
+};
+
+
+template<>
+struct TStructOpsTypeTraits<FMoverDataCollection> : public TStructOpsTypeTraitsBase2<FMoverDataCollection>
+{
+	enum
+	{
+		WithCopy = true,		// Necessary so that TSharedPtr<FMoverDataStructBase> Data is copied around
+		WithNetSerializer = true,
+		WithIdenticalViaEquality = true,
+		WithAddStructReferencedObjects = true,
+	};
+};
+
+/** Info about data collection types that should always be present, and how they should propagate from one frame to the next */
+USTRUCT(BlueprintType, meta = (DisplayName = "Persistent Data Settings"))
+struct MOVER_API FMoverDataPersistence
+{
+	GENERATED_BODY()
+
+	FMoverDataPersistence() {}
+
+	FMoverDataPersistence(UScriptStruct* TypeToPersist, bool bShouldCopyBetweenFrames) 
+	{
+		RequiredType = TypeToPersist;
+		bCopyFromPriorFrame = bShouldCopyBetweenFrames;
+	}
+
+	// The type that should propagate between frames
+	UPROPERTY(EditAnywhere, BlueprintReadWrite, Category = Mover, meta=(MetaStruct="/Script/Mover.MoverDataStructBase"))
+	TObjectPtr<UScriptStruct> RequiredType = nullptr;
+
+	// If true, values will be copied from the prior frame. Otherwise, they will be default-initialized.
+	UPROPERTY(EditAnywhere, BlueprintReadWrite, Category = Mover)
+	bool bCopyFromPriorFrame = true;
+};
+
+
+// Blueprint helper functions for working with a Mover data collection
+UCLASS()
+class MOVER_API UMoverDataCollectionLibrary : public UBlueprintFunctionLibrary
+{
+	GENERATED_BODY()
+
+public:
+	/**
+	* Add a data struct to the collection, overwriting an existing one with the same type
+	* @param SourceAsRawBytes		The data struct instance to add by copy, which must be a FMoverDataStructBase sub-type
+	*/
+	UFUNCTION(BlueprintCallable, CustomThunk, Category = "Mover|Data", meta = (CustomStructureParam = "SourceAsRawBytes", AllowAbstract = "false", DisplayName = "Add Data To Collection"))
+	static void K2_AddDataToCollection(UPARAM(Ref) FMoverDataCollection& Collection, UPARAM(DisplayName="Struct To Add") const int32& SourceAsRawBytes);
+
+	/**
+	 * Retrieves data from a collection, by writing to a target instance if it contains one of the matching type.  Changes must be written back using AddDataToCollection.
+	 * @param DidSucceed			Flag indicating whether data was actually written to target struct instance
+	 * @param TargetAsRawBytes		The data struct instance to write to, which must be a FMoverDataStructBase sub-type
+	 */
+	UFUNCTION(BlueprintCallable, CustomThunk, Category = "Mover|Data", meta = (CustomStructureParam = "TargetAsRawBytes", AllowAbstract = "false", DisplayName = "Get Data From Collection"))
+	static void K2_GetDataFromCollection(bool& DidSucceed, UPARAM(Ref) const FMoverDataCollection& Collection, UPARAM(DisplayName = "Out Struct") int32& TargetAsRawBytes);
+
+	/**
+	* Clears all data from a collection
+	*/
+	UFUNCTION(BlueprintCallable, Category = "Mover|Data", meta=(DisplayName = "Clear Data From Collection"))
+	static void ClearDataFromCollection(UPARAM(Ref) FMoverDataCollection& Collection);
+
+private:
+	DECLARE_FUNCTION(execK2_AddDataToCollection);
+	DECLARE_FUNCTION(execK2_GetDataFromCollection);
+};
+
diff --git a/Engine/Plugins/Experimental/Mover/Source/Mover/Public/PhysicsMover/Modes/PhysicsDrivenFallingMode.h b/Engine/Plugins/Experimental/Mover/Source/Mover/Public/PhysicsMover/Modes/PhysicsDrivenFallingMode.h
new file mode 100644
--- /dev/null
+++ b/Engine/Plugins/Experimental/Mover/Source/Mover/Public/PhysicsMover/Modes/PhysicsDrivenFallingMode.h
@@ -0,0 +1,39 @@
+// Copyright Epic Games, Inc. All Rights Reserved.
+
+#pragma once
+
+#include "CoreMinimal.h"
+#include "Kinematic/Modes/FallingMode.h"
+#include "PhysicsMover/PhysicsMoverSimulationTypes.h"
+
+#include "PhysicsDrivenFallingMode.generated.h"
+
+
+namespace Chaos { class FCharacterGroundConstraint; }
+
+/**
+ * PhysicsDrivenFallingMode: Override base kinematic falling mode for physics based motion.
+ */
+UCLASS(Blueprintable, BlueprintType)
+class MOVER_API UPhysicsDrivenFallingMode : public UFallingMode, public IPhysicsCharacterMovementModeInterface
+{
+	GENERATED_UCLASS_BODY()
+
+public:
+
+	virtual void OnSimulationTick(const FSimulationTickParams& Params, FMoverTickEndData& OutputState) override;
+
+	virtual void UpdateConstraintSettings(Chaos::FCharacterGroundConstraint& Constraint) const override;
+
+	// Maximum torque the character can apply to rotate in air about the vertical axis
+	UPROPERTY(EditAnywhere, Category = "Physics Mover", meta = (ClampMin = "0", UIMin = "0", ForceUnits = "NewtonMeters"))
+	float TwistTorqueLimit = 0.0f;
+
+	// Maximum torque the character can apply to remain upright
+	UPROPERTY(EditAnywhere, Category = "Physics Mover", meta = (ClampMin = "0", UIMin = "0", ForceUnits = "NewtonMeters"))
+	float SwingTorqueLimit = 3000.0f;
+
+	// Target height for the character. This is the desired distance from the center of the capsule to the floor
+	UPROPERTY(EditAnywhere, Category = "Physics Mover", meta = (ClampMin = "0", UIMin = "0", ForceUnits = "cm"))
+	float TargetHeight = 54.0f;
+};
diff --git a/Engine/Plugins/Experimental/Mover/Source/Mover/Public/PhysicsMover/Modes/PhysicsDrivenFlyingMode.h b/Engine/Plugins/Experimental/Mover/Source/Mover/Public/PhysicsMover/Modes/PhysicsDrivenFlyingMode.h
new file mode 100644
--- /dev/null
+++ b/Engine/Plugins/Experimental/Mover/Source/Mover/Public/PhysicsMover/Modes/PhysicsDrivenFlyingMode.h
@@ -0,0 +1,35 @@
+// Copyright Epic Games, Inc. All Rights Reserved.
+
+#pragma once
+
+#include "CoreMinimal.h"
+#include "Kinematic/Modes/FlyingMode.h"
+#include "PhysicsMover/PhysicsMoverSimulationTypes.h"
+
+#include "PhysicsDrivenFlyingMode.generated.h"
+
+
+namespace Chaos { class FCharacterGroundConstraint; }
+
+/**
+ * PhysicsDrivenFlyingMode: Override base kinematic flying mode for physics based motion.
+ */
+UCLASS(Blueprintable, BlueprintType)
+class MOVER_API UPhysicsDrivenFlyingMode : public UFlyingMode, public IPhysicsCharacterMovementModeInterface
+{
+	GENERATED_UCLASS_BODY()
+
+public:
+
+	virtual void OnSimulationTick(const FSimulationTickParams& Params, FMoverTickEndData& OutputState) override;
+
+	virtual void UpdateConstraintSettings(Chaos::FCharacterGroundConstraint& Constraint) const override;
+
+	// Maximum torque the character can apply to rotate in air about the vertical axis
+	UPROPERTY(EditAnywhere, Category = "Physics Mover", meta = (ClampMin = "0", UIMin = "0", ForceUnits = "NewtonMeters"))
+	float TwistTorqueLimit = 0.0f;
+
+	// Maximum torque the character can apply to remain upright
+	UPROPERTY(EditAnywhere, Category = "Physics Mover", meta = (ClampMin = "0", UIMin = "0", ForceUnits = "NewtonMeters"))
+	float SwingTorqueLimit = 3000.0f;
+};
diff --git a/Engine/Plugins/Experimental/Mover/Source/Mover/Public/PhysicsMover/Modes/PhysicsDrivenSwimmingMode.h b/Engine/Plugins/Experimental/Mover/Source/Mover/Public/PhysicsMover/Modes/PhysicsDrivenSwimmingMode.h
new file mode 100644
--- /dev/null
+++ b/Engine/Plugins/Experimental/Mover/Source/Mover/Public/PhysicsMover/Modes/PhysicsDrivenSwimmingMode.h
@@ -0,0 +1,33 @@
+// Copyright Epic Games, Inc. All Rights Reserved.
+
+#pragma once
+
+#include "CoreMinimal.h"
+#include "Kinematic/Modes/SwimmingMode.h"
+#include "PhysicsMover/PhysicsMoverSimulationTypes.h"
+
+#include "PhysicsDrivenSwimmingMode.generated.h"
+
+
+namespace Chaos { class FCharacterGroundConstraint; }
+
+/**
+ * PhysicsDrivenSwimmingMode: Override base kinematic Swimming mode for physics based motion.
+ */
+UCLASS(Blueprintable, BlueprintType)
+class MOVER_API UPhysicsDrivenSwimmingMode : public USwimmingMode, public IPhysicsCharacterMovementModeInterface
+{
+	GENERATED_UCLASS_BODY()
+
+public:
+
+	virtual void OnSimulationTick(const FSimulationTickParams& Params, FMoverTickEndData& OutputState) override;
+	virtual bool AttemptJump(float UpwardsSpeed, FMoverTickEndData& Output);
+	virtual bool AttemptTeleport(USceneComponent* UpdatedComponent, const FVector& TeleportPos, const FRotator& TeleportRot, const FVector& PriorVelocity, FMoverTickEndData& Output);
+
+	virtual void UpdateConstraintSettings(Chaos::FCharacterGroundConstraint& Constraint) const override;
+
+	// Target height for the character. This is the desired distance from the center of the capsule to the floor
+	UPROPERTY(EditAnywhere, Category = "Physics Mover", meta = (ClampMin = "0", UIMin = "0", ForceUnits = "cm"))
+	float TargetHeight = 54.0f;
+};
diff --git a/Engine/Plugins/Experimental/Mover/Source/Mover/Public/PhysicsMover/Modes/PhysicsDrivenWalkingMode.h b/Engine/Plugins/Experimental/Mover/Source/Mover/Public/PhysicsMover/Modes/PhysicsDrivenWalkingMode.h
new file mode 100644
--- /dev/null
+++ b/Engine/Plugins/Experimental/Mover/Source/Mover/Public/PhysicsMover/Modes/PhysicsDrivenWalkingMode.h
@@ -0,0 +1,75 @@
+// Copyright Epic Games, Inc. All Rights Reserved.
+
+#pragma once
+
+#include "CoreMinimal.h"
+#include "Kinematic/Modes/WalkingMode.h"
+#include "PhysicsMover/PhysicsMoverSimulationTypes.h"
+
+#include "PhysicsDrivenWalkingMode.generated.h"
+
+
+namespace Chaos { class FCharacterGroundConstraint; }
+
+/**
+ * PhysicsDrivenWalkingMode: Override base kinematic walking mode for physics based motion.
+ */
+UCLASS(Blueprintable, BlueprintType)
+class MOVER_API UPhysicsDrivenWalkingMode : public UWalkingMode, public IPhysicsCharacterMovementModeInterface
+{
+	GENERATED_UCLASS_BODY()
+
+public:
+
+	virtual void OnSimulationTick(const FSimulationTickParams& Params, FMoverTickEndData& OutputState) override;
+	virtual bool AttemptTeleport(USceneComponent* UpdatedComponent, const FVector& TeleportPos, const FRotator& TeleportRot, const FVector& PriorVelocity, FMoverTickEndData& Output);
+
+	virtual void UpdateConstraintSettings(Chaos::FCharacterGroundConstraint& Constraint) const override;
+
+	// Maximum force the character can apply to reach the motion target
+	UPROPERTY(EditAnywhere, Category = "Physics Mover", meta = (ClampMin = "0", UIMin = "0", ForceUnits = "Newtons"))
+	float RadialForceLimit = 1500.0f;
+
+	// Maximum force the character can apply to hold in place while standing on an unwalkable incline
+	UPROPERTY(EditAnywhere, Category = "Physics Mover", meta = (ClampMin = "0", UIMin = "0", ForceUnits = "Newtons"))
+	float FrictionForceLimit = 100.0f;
+
+	// Maximum torque the character can apply to rotate in air about the vertical axis
+	UPROPERTY(EditAnywhere, Category = "Physics Mover", meta = (ClampMin = "0", UIMin = "0", ForceUnits = "NewtonMeters"))
+	float TwistTorqueLimit = 1000.0f;
+
+	// Maximum torque the character can apply to remain upright
+	UPROPERTY(EditAnywhere, Category = "Physics Mover", meta = (ClampMin = "0", UIMin = "0", ForceUnits = "NewtonMeters"))
+	float SwingTorqueLimit = 3000.0f;
+
+	// Target height for the character. This is the desired distance from the center of the capsule to the floor
+	UPROPERTY(EditAnywhere, Category = "Physics Mover", meta = (ClampMin = "0", UIMin = "0", ForceUnits = "cm"))
+	float TargetHeight = 54.0f;
+
+	// Damping factor to control the softness of the interaction between the character and the ground
+	// Set to 0 for no damping and 1 for maximum damping
+	UPROPERTY(EditAnywhere, Category = "Physics Mover", meta = (ClampMin = "0", UIMin = "0", ClampMax = "1", UIMax = "1"))
+	float GroundDamping = 0.0f;
+
+	// With this enabled, in addition to the constraint radial force, the mover applies an initial velocity to reach the target
+	// movement. This allows for a separation between pushing force and force to reach the target movement.
+	// A value of 1 means the character will always be able to reach the target motion (if there are no obstacles) regardless of
+	// the radial force limit.
+	// A value of 0 means the only thing moving the character is the constraint
+	UPROPERTY(EditAnywhere, Category = "Physics Mover", meta = (ClampMin = "0", UIMin = "0", ClampMax = "1", UIMax = "1"))
+	float FractionalVelocityToTarget = 1.0f;
+
+	// Controls how much downward velocity is applied to keep the character rooted to the ground when the character
+	// is within MaxStepHeight of the ground surface.
+	UPROPERTY(EditAnywhere, Category = "Physics Mover", meta = (ClampMin = "0", UIMin = "0", ClampMax = "1", UIMax = "1"))
+	float FractionalDownwardVelocityToTarget = 1.0f;
+
+	// Time limit for being unsupported before moving from a walking to a falling state.
+	// This provides some grace period when walking off of an edge during which locomotion
+	// and jumping are still possible even though the character has started falling under gravity
+	UPROPERTY(EditAnywhere, Category = "Physics Mover", meta = (ClampMin = "0", UIMin = "0", ForceUnits = "s"))
+	float MaxUnsupportedTimeBeforeFalling = 0.05f;
+
+protected:
+	void SwitchToState(const FName& StateName, const FSimulationTickParams& Params, FMoverTickEndData& OutputState);
+};
diff --git a/Engine/Plugins/Experimental/Mover/Source/Mover/Public/PhysicsMover/PhysicsMovementUtils.h b/Engine/Plugins/Experimental/Mover/Source/Mover/Public/PhysicsMover/PhysicsMovementUtils.h
new file mode 100644
--- /dev/null
+++ b/Engine/Plugins/Experimental/Mover/Source/Mover/Public/PhysicsMover/PhysicsMovementUtils.h
@@ -0,0 +1,30 @@
+// Copyright Epic Games, Inc. All Rights Reserved.
+
+#pragma once
+
+#include "Chaos/ParticleHandleFwd.h"
+#include "Containers/Array.h"
+#include "Math/MathFwd.h"
+
+struct FFloorCheckResult;
+struct FWaterCheckResult;
+struct FHitResult;
+class UPrimitiveComponent;
+
+/**
+ * PhysicsMovementUtils: a collection of stateless static functions for a variety of physics movement-related operations
+ */
+class MOVER_API UPhysicsMovementUtils
+{
+public:
+	static void FindFloor(const FVector& Location, const FVector& DeltaPos, const UPrimitiveComponent* UpdatedPrimitive, const FVector& UpDir, float QueryRadius, float TargetHeight, float MaxStepHeight, float MaxWalkSlopeCosine, FFloorCheckResult& OutFloorResult, FWaterCheckResult& OutWaterResult);
+	static const Chaos::FPBDRigidParticleHandle* GetRigidParticelHandleFromHitResult(const FHitResult& HitResult);
+
+	// Checks if the hit surface is walkable and, if stepping up, whether the surface can be stepped up on
+	static bool IsHitSurfaceWalkableWithStepUpCheck(const FHitResult& Hit, float StepHeight, float MaxStepHeight, float MinStepUpHeight, float MaxWalkSlopeCosine);
+
+	// Checks if any hit is with water and, if so, fills in the OutWaterResult
+	static bool GetWaterResultFromHitResults(const TArray<FHitResult>& Hits, const FVector& Location, FWaterCheckResult& OutWaterResult);
+
+	static FVector ComputeGroundVelocityFromHitResult(const FVector& CharacterPosition, const FHitResult& FloorHit, const float DeltaSeconds);
+};
\ No newline at end of file
diff --git a/Engine/Plugins/Experimental/Mover/Source/Mover/Public/PhysicsMover/PhysicsMoverManager.h b/Engine/Plugins/Experimental/Mover/Source/Mover/Public/PhysicsMover/PhysicsMoverManager.h
new file mode 100644
--- /dev/null
+++ b/Engine/Plugins/Experimental/Mover/Source/Mover/Public/PhysicsMover/PhysicsMoverManager.h
@@ -0,0 +1,32 @@
+// Copyright Epic Games, Inc. All Rights Reserved.
+
+#pragma once
+
+#include "PhysicsInterfaceDeclaresCore.h"
+#include "Subsystems/WorldSubsystem.h"
+#include "PhysicsMoverManager.generated.h"
+
+//////////////////////////////////////////////////////////////////////////
+
+class UMoverNetworkPhysicsLiaisonComponent;
+
+UCLASS()
+class MOVER_API UPhysicsMoverManager : public UWorldSubsystem
+{
+	GENERATED_BODY()
+
+public:
+	virtual void OnWorldBeginPlay(UWorld& InWorld) override;
+	virtual void Deinitialize() override;
+	virtual bool DoesSupportWorldType(const EWorldType::Type WorldType) const override;
+
+	void PrePhysicsUpdate(FPhysScene* PhysScene, float DeltaTime);
+
+	void RegisterPhysicsMoverComponent(TWeakObjectPtr<UMoverNetworkPhysicsLiaisonComponent> InPhysicsMoverComp);
+	void UnregisterPhysicsMoverComponent(TWeakObjectPtr<UMoverNetworkPhysicsLiaisonComponent> InPhysicsMoverComp);
+
+private:
+	TArray<TWeakObjectPtr<UMoverNetworkPhysicsLiaisonComponent>> PhysicsMoverComponents;
+	FDelegateHandle PhysScenePreTickCallbackHandle;
+	class FPhysicsMoverManagerAsyncCallback* AsyncCallback = nullptr;
+};
\ No newline at end of file
diff --git a/Engine/Plugins/Experimental/Mover/Source/Mover/Public/PhysicsMover/PhysicsMoverManagerAsyncCallback.h b/Engine/Plugins/Experimental/Mover/Source/Mover/Public/PhysicsMover/PhysicsMoverManagerAsyncCallback.h
new file mode 100644
--- /dev/null
+++ b/Engine/Plugins/Experimental/Mover/Source/Mover/Public/PhysicsMover/PhysicsMoverManagerAsyncCallback.h
@@ -0,0 +1,43 @@
+// Copyright Epic Games, Inc. All Rights Reserved.
+
+#pragma once
+
+#include "Chaos/GeometryParticlesfwd.h"
+#include "Chaos/SimCallbackObject.h"
+#include "PhysicsMover/PhysicsMoverSimulationTypes.h"
+
+//////////////////////////////////////////////////////////////////////////
+
+struct FPhysicsMoverManagerAsyncInput : public Chaos::FSimCallbackInput
+{
+	TArray<TUniquePtr<FPhysicsMoverAsyncInput>> AsyncInput;
+
+	void Reset()
+	{
+		AsyncInput.Reset();
+	}
+};
+
+struct FPhysicsMoverManagerAsyncOutput : public Chaos::FSimCallbackOutput
+{
+	TMap<Chaos::FUniqueIdx, TUniquePtr<FPhysicsMoverAsyncOutput>> PhysicsMoverToAsyncOutput;
+
+	void Reset()
+	{
+		PhysicsMoverToAsyncOutput.Reset();
+	}
+};
+
+class FPhysicsMoverManagerAsyncCallback : public Chaos::TSimCallbackObject<
+	FPhysicsMoverManagerAsyncInput,
+	FPhysicsMoverManagerAsyncOutput,
+	Chaos::ESimCallbackOptions::Presimulate |
+	Chaos::ESimCallbackOptions::ContactModification |
+	Chaos::ESimCallbackOptions::Rewind>
+{
+
+private:
+	virtual void ProcessInputs_Internal(int32 PhysicsStep) override;
+	virtual void OnPreSimulate_Internal() override;
+	virtual void OnContactModification_Internal(Chaos::FCollisionContactModifier& Modifier) override;
+};
\ No newline at end of file
diff --git a/Engine/Plugins/Experimental/Mover/Source/Mover/Public/PhysicsMover/PhysicsMoverSimulationTypes.h b/Engine/Plugins/Experimental/Mover/Source/Mover/Public/PhysicsMover/PhysicsMoverSimulationTypes.h
new file mode 100644
--- /dev/null
+++ b/Engine/Plugins/Experimental/Mover/Source/Mover/Public/PhysicsMover/PhysicsMoverSimulationTypes.h
@@ -0,0 +1,101 @@
+// Copyright Epic Games, Inc. All Rights Reserved.
+
+#pragma once
+
+#include "Chaos/GeometryParticlesfwd.h"
+#include "MoverSimulationTypes.h"
+#include "MoveLibrary/FloorQueryUtils.h"
+#include "UObject/Interface.h"
+#include "UObject/WeakInterfacePtr.h"
+#include "PhysicsMoverSimulationTypes.generated.h"
+
+//////////////////////////////////////////////////////////////////////////
+// Debug
+
+struct FPhysicsDrivenMotionDebugParams
+{
+	float TeleportThreshold = 100.0f;
+	float MinStepUpDistance = 5.0f;
+	bool EnableMultithreading = true;
+	bool DebugDrawGroundQueries = false;
+};
+
+#ifndef PHYSICSDRIVENMOTION_DEBUG_DRAW
+#define PHYSICSDRIVENMOTION_DEBUG_DRAW (!(UE_BUILD_SHIPPING || UE_BUILD_TEST) || WITH_EDITOR)
+#endif
+
+//////////////////////////////////////////////////////////////////////////
+// Async update
+
+struct FPhysicsMoverSimulationTickParams
+{
+	float SimTimeSeconds = 0.0f;
+	float DeltaTimeSeconds = 0.0f;
+};
+
+struct FPhysicsMoverSimulationInput
+{
+	FMoverInputCmdContext InputCmd;
+	FMoverSyncState SyncState;
+	FMoverAuxStateContext AuxState;
+};
+
+struct FPhysicsMoverSimulationOutput
+{
+	bool IsValid() const { return bIsValid; }
+
+	FMoverSyncState SyncState;
+	FMoverAuxStateContext AuxState;
+	FFloorCheckResult FloorResult;
+	bool bIsValid = false;
+};
+
+namespace Chaos { class FCharacterGroundConstraintHandle; }
+
+struct FPhysicsMoverAsyncInput
+{
+	bool IsValid() const
+	{
+		return MoverSimulation.IsValid() && MoverIdx.IsValid();
+	}
+
+	// Input is modified during ProcessInputs_Internal
+	mutable FMoverInputCmdContext InputCmd;
+	mutable FMoverSyncState SyncState;
+	mutable FMoverAuxStateContext AuxState;
+
+	TWeakObjectPtr<class UMoverNetworkPhysicsLiaisonComponent> MoverSimulation;
+	Chaos::FUniqueIdx MoverIdx;
+};
+
+struct FPhysicsMoverAsyncOutput
+{
+	FMoverSyncState SyncState;
+	FMoverAuxStateContext AuxState;
+	FFloorCheckResult FloorResult;
+	bool bIsValid = false;
+};
+
+//////////////////////////////////////////////////////////////////////////
+// Movement modes
+
+/**
+ * PhysicsDrivenMotionMode: Interface for movement modes that are for physics driven motion
+ * A physics driven motion mode needs to update the character ground constraint with the
+ * parameters associated with that mode
+ */
+UINTERFACE(MinimalAPI)
+class UPhysicsCharacterMovementModeInterface : public UInterface
+{
+	GENERATED_BODY()
+};
+
+namespace Chaos { class FCharacterGroundConstraint; }
+
+class IPhysicsCharacterMovementModeInterface
+{
+	GENERATED_BODY()
+
+public:
+	virtual void UpdateConstraintSettings(Chaos::FCharacterGroundConstraint& Constraint) const = 0;
+};
\ No newline at end of file
diff --git a/Engine/Plugins/Experimental/Mover/Source/Mover/Public/PhysicsMover/Transitions/PhysicsJumpCheck.h b/Engine/Plugins/Experimental/Mover/Source/Mover/Public/PhysicsMover/Transitions/PhysicsJumpCheck.h
new file mode 100644
--- /dev/null
+++ b/Engine/Plugins/Experimental/Mover/Source/Mover/Public/PhysicsMover/Transitions/PhysicsJumpCheck.h
@@ -0,0 +1,34 @@
+// Copyright Epic Games, Inc. All Rights Reserved.
+
+#pragma once
+
+#include "MovementModeTransition.h"
+#include "PhysicsJumpCheck.generated.h"
+
+/**
+ * Transition that handles jumping based on input for a physics-based character
+ */
+UCLASS(Blueprintable, BlueprintType)
+class MOVER_API UPhysicsJumpCheck : public UBaseMovementModeTransition
+{
+	GENERATED_UCLASS_BODY()
+
+
+public:
+	virtual FTransitionEvalResult OnEvaluate(const FSimulationTickParams& Params) const override;
+	virtual void OnTrigger(const FSimulationTickParams& Params) override;
+
+	/** Instantaneous speed induced in an actor upon jumping */
+	UPROPERTY(BlueprintReadWrite, EditAnywhere, Category = Parameters, meta = (ClampMin = "0", UIMin = "0", ForceUnits = "cm/s"))
+	float JumpUpwardsSpeed = 500.0f;
+
+	/** Name of movement mode to transition to when jumping is activated */
+	UPROPERTY(BlueprintReadWrite, EditAnywhere, Category = Parameters)
+	FName TransitionToMode;
+
+#if WITH_EDITOR
+	virtual EDataValidationResult IsDataValid(class FDataValidationContext& Context) const override;
+#endif // WITH_EDITOR
+
+
+};
\ No newline at end of file
diff --git a/Engine/Plugins/Experimental/MoverExamples/Source/MoverExamples/Private/MovementModes/AdvancedWalkingMode.cpp b/Engine/Plugins/Experimental/MoverExamples/Source/MoverExamples/Private/MovementModes/AdvancedWalkingMode.cpp
new file mode 100644
--- /dev/null
+++ b/Engine/Plugins/Experimental/MoverExamples/Source/MoverExamples/Private/MovementModes/AdvancedWalkingMode.cpp
@@ -0,0 +1,40 @@
+// Copyright Epic Games, Inc. All Rights Reserved.
+
+#include "MovementModes/AdvancedWalkingMode.h"
+#include "MoveLibrary/MovementUtils.h"
+#include "MoverComponent.h"
+#include "Kinematic/Settings/CommonLegacyMovementSettings.h"
+#include "MoverLog.h"
+
+#include UE_INLINE_GENERATED_CPP_BY_NAME(AdvancedWalkingMode)
+
+UAdvancedWalkingMode::UAdvancedWalkingMode(const FObjectInitializer& ObjectInitializer)
+	: Super(ObjectInitializer)
+{
+	// Default turning rate: turn slower at higher speeds
+	MaxTurningRateBySpeedPct.GetRichCurve()->Reset();
+	MaxTurningRateBySpeedPct.GetRichCurve()->AddKey(0.f, 500.f);
+	MaxTurningRateBySpeedPct.GetRichCurve()->AddKey(1.f, 100.f);
+	MaxTurningRateBySpeedPct.GetRichCurve()->SetDefaultValue(500.f);
+}
+
+
+void UAdvancedWalkingMode::OnGenerateMove(const FMoverTickStartData& StartState, const FMoverTimeStep& TimeStep, FProposedMove& OutProposedMove) const
+{
+	UMoverComponent* MoverComp = GetMoverComponent();
+
+	// TODO: this was an experiment in dynamically modifying the turning rate. It should be removed once we support a modular turning mechanism.
+	FMoverTickStartData StartStateCopy = StartState;
+	const FMoverDefaultSyncState* StartingSyncState = StartStateCopy.SyncState.SyncStateCollection.FindDataByType<FMoverDefaultSyncState>();
+	check(StartingSyncState);
+
+	// Map current speed to turning rate
+	if (UCommonLegacyMovementSettings* LegacySettings = MoverComp->FindSharedSettings_Mutable<UCommonLegacyMovementSettings>())
+	{
+		const float PctToMaxSpeed = FMath::Clamp(StartingSyncState->GetVelocity_WorldSpace().Length() / LegacySettings->MaxSpeed, 0.f, 1.f);
+		const float TurningRate = MaxTurningRateBySpeedPct.GetRichCurveConst()->Eval(PctToMaxSpeed);
+		LegacySettings->TurningRate = TurningRate;
+	}
+
+	Super::OnGenerateMove(StartStateCopy, TimeStep, OutProposedMove);
+}
diff --git a/Engine/Plugins/Experimental/MoverExamples/Source/MoverExamples/Public/MovementModes/AdvancedWalkingMode.h b/Engine/Plugins/Experimental/MoverExamples/Source/MoverExamples/Public/MovementModes/AdvancedWalkingMode.h
new file mode 100644
--- /dev/null
+++ b/Engine/Plugins/Experimental/MoverExamples/Source/MoverExamples/Public/MovementModes/AdvancedWalkingMode.h
@@ -0,0 +1,24 @@
+// Copyright Epic Games, Inc. All Rights Reserved.
+
+#pragma once
+
+#include "CoreMinimal.h"
+#include "Kinematic/Modes/WalkingMode.h"
+#include "Curves/CurveFloat.h"
+#include "AdvancedWalkingMode.generated.h"
+
+/**
+ * AdvancedWalkingMode: advanced example mode that expands tunability of movement input, speed, turning, etc.
+ */
+UCLASS(Blueprintable, BlueprintType)
+class UAdvancedWalkingMode : public UWalkingMode
+{
+	GENERATED_UCLASS_BODY()
+
+public:
+	virtual void OnGenerateMove(const FMoverTickStartData& StartState, const FMoverTimeStep& TimeStep, FProposedMove& OutProposedMove) const override;
+
+	// Maps ratio of Current Speed / Max Speed (range [0, 1]) to a turning rate (degrees per second)
+	UPROPERTY(EditAnywhere, Category = "Movement Tuning")
+	FRuntimeFloatCurve MaxTurningRateBySpeedPct;
+};
