Mover 2.0 fixing a crash that occurs when attempting to fire a Queue Layered Move that's connected to a non struct node

diff --git a/Engine/Plugins/Experimental/Mover/Source/Mover/Private/MoverComponent.cpp b/Engine/Plugins/Experimental/Mover/Source/Mover/Private/MoverComponent.cpp
--- a/Engine/Plugins/Experimental/Mover/Source/Mover/Private/MoverComponent.cpp
+++ b/Engine/Plugins/Experimental/Mover/Source/Mover/Private/MoverComponent.cpp
@@ -733,7 +733,8 @@ DEFINE_FUNCTION(UMoverComponent::execK2_QueueLayeredMove)
 
 	const bool bHasValidStructProp = StructProp && StructProp->Struct && StructProp->Struct->IsChildOf(FLayeredMoveBase::StaticStruct());
 
-	if (ensureMsgf((bHasValidStructProp && MovePtr), TEXT("An invalid struct type (%s) was sent to a QueueLayeredMove node. A struct derived from FLayeredMoveBase is required. No layered move will be queued."), *GetNameSafe(StructProp->Struct)))
+	if (ensureMsgf((bHasValidStructProp && MovePtr), TEXT("An invalid type (%s) was sent to a QueueLayeredMove node. A struct derived from FLayeredMoveBase is required. No layered move will be queued."),
+		StructProp ? *GetNameSafe(StructProp->Struct) : *Stack.MostRecentProperty->GetClass()->GetName()))
 	{
 		// Could we steal this instead of cloning? (move semantics)
 		FLayeredMoveBase* MoveAsBasePtr = reinterpret_cast<FLayeredMoveBase*>(MovePtr);
