Fixes check crash when compiling a BP that has a MoverComponent with a NetworkPhysicsLiaison backend specified

diff --git a/Engine/Plugins/Experimental/Mover/Source/Mover/Private/Backends/MoverNetworkPhysicsLiaison.cpp b/Engine/Plugins/Experimental/Mover/Source/Mover/Private/Backends/MoverNetworkPhysicsLiaison.cpp
--- a/Engine/Plugins/Experimental/Mover/Source/Mover/Private/Backends/MoverNetworkPhysicsLiaison.cpp
+++ b/Engine/Plugins/Experimental/Mover/Source/Mover/Private/Backends/MoverNetworkPhysicsLiaison.cpp
@@ -380,7 +380,11 @@ void UMoverNetworkPhysicsLiaisonComponent::InitializeComponent()
 
 void UMoverNetworkPhysicsLiaisonComponent::UninitializeComponent()
 {
-	NetworkPhysicsComponent->UninitializeComponent();
+	if (NetworkPhysicsComponent && NetworkPhysicsComponent->HasBeenInitialized())
+	{
+		NetworkPhysicsComponent->UninitializeComponent();
+	}
+
 	Super::UninitializeComponent();
 }
 
