Mover 2.0 removing OnImpact delegate since as it was unneeded and could potentially affect performance

diff --git a/Engine/Plugins/Experimental/Mover/Source/Mover/Private/MoverComponent.cpp b/Engine/Plugins/Experimental/Mover/Source/Mover/Private/MoverComponent.cpp
--- a/Engine/Plugins/Experimental/Mover/Source/Mover/Private/MoverComponent.cpp
+++ b/Engine/Plugins/Experimental/Mover/Source/Mover/Private/MoverComponent.cpp
@@ -452,7 +452,6 @@ void UMoverComponent::HandleImpact(FMoverOnImpactParams& ImpactParams)
 	}
 	
 	OnHandleImpact(ImpactParams);
-	OnImpact.Broadcast(ImpactParams);
 }
 
 void UMoverComponent::OnHandleImpact(const FMoverOnImpactParams& ImpactParams)
diff --git a/Engine/Plugins/Experimental/Mover/Source/Mover/Public/MoverComponent.h b/Engine/Plugins/Experimental/Mover/Source/Mover/Public/MoverComponent.h
--- a/Engine/Plugins/Experimental/Mover/Source/Mover/Public/MoverComponent.h
+++ b/Engine/Plugins/Experimental/Mover/Source/Mover/Public/MoverComponent.h
@@ -37,9 +37,6 @@ DECLARE_DYNAMIC_MULTICAST_DELEGATE_TwoParams(FMover_OnPostSimRollback, const FMo
 // Fired after changing movement modes. First param is the name of the previous movement mode. Second is the name of the new movement mode. 
 DECLARE_DYNAMIC_MULTICAST_DELEGATE_TwoParams(FMover_OnMovementModeChanged, const FName&, PreviousMovementModeName, const FName&, NewMovementModeName);
 
-// Fired after the actor hits a surface. Param is FMoverOnImpactParams struct that holds all of the data/params for an impact see @FMoverOnImpactParams for more details of it's properties.
-DECLARE_DYNAMIC_MULTICAST_DELEGATE_OneParam(FMover_OnImpact, const FMoverOnImpactParams&, ImpactParams);
-
 /**
  * 
  */
@@ -76,10 +73,6 @@ public:
 	// Broadcast when a rollback has occurred, just before the next simulation tick occurs
 	UPROPERTY(BlueprintAssignable, Category = Mover)
 	FMover_OnPostSimRollback OnPostSimulationRollback;
-	
-	// Broadcast when this actor impacts a surface.
-	UPROPERTY(BlueprintAssignable, Category = Mover)
-	FMover_OnImpact OnImpact;
 
 	// Broadcast when a MovementMode has changed. Happens during a simulation tick if the mode changed that tick or when SetModeImmediately is used to change modes.
 	UPROPERTY(BlueprintAssignable, Category = Mover)
