Mover: expanding transitions to be supported on the MoverComponent, not just on the modes

diff --git a/Engine/Plugins/Experimental/Mover/Source/Mover/Private/MovementModeStateMachine.cpp b/Engine/Plugins/Experimental/Mover/Source/Mover/Private/MovementModeStateMachine.cpp
--- a/Engine/Plugins/Experimental/Mover/Source/Mover/Private/MovementModeStateMachine.cpp
+++ b/Engine/Plugins/Experimental/Mover/Source/Mover/Private/MovementModeStateMachine.cpp
@@ -225,23 +225,42 @@ void UMovementModeStateMachine::OnSimulationTick(USceneComponent* UpdatedCompone
 				SimTickParams.TimeStep = SubTimeStep;
 				SimTickParams.ProposedMove = CombinedMove;
 
-				// Check any transitions that have been registered with the current movement mode
-				bool bTransitionTriggered = false;
+				// Check for any transitions, first those registered with the current movement mode, then global ones that could occur from any mode
+				FTransitionEvalResult EvalResult = FTransitionEvalResult::NoTransition;
+				TObjectPtr<UBaseMovementModeTransition> TransitionToTrigger;
+
 				for (UBaseMovementModeTransition* Transition : CurrentMode->Transitions)
 				{
-					FTransitionEvalResult EvalResult = Transition->DoEvaluate(SimTickParams);
+					EvalResult = Transition->DoEvaluate(SimTickParams);
 
 					if (!EvalResult.NextMode.IsNone())
 					{
-						OutputState.MovementEndState.NextModeName = EvalResult.NextMode;
-						OutputState.MovementEndState.RemainingMs = SimTickParams.TimeStep.StepMs; 	// Pass all remaining time to next mode
-						Transition->DoTrigger(SimTickParams);
-						bTransitionTriggered = true;
+						TransitionToTrigger = Transition;
 						break;
 					}
 				}
 
-				if (!bTransitionTriggered)
+				if (EvalResult.NextMode.IsNone())
+				{
+					for (UBaseMovementModeTransition* Transition : MoverComp->Transitions)
+					{
+						EvalResult = Transition->DoEvaluate(SimTickParams);
+
+						if (!EvalResult.NextMode.IsNone())
+						{
+							TransitionToTrigger = Transition;
+							break;
+						}
+					}
+				}
+
+				if (!EvalResult.NextMode.IsNone())
+				{
+					OutputState.MovementEndState.NextModeName = EvalResult.NextMode;
+					OutputState.MovementEndState.RemainingMs = SimTickParams.TimeStep.StepMs; 	// Pass all remaining time to next mode
+					TransitionToTrigger->DoTrigger(SimTickParams);
+				}
+				else
 				{
 					CurrentMode->DoSimulationTick(SimTickParams, OutputState);
 				}
diff --git a/Engine/Plugins/Experimental/Mover/Source/Mover/Public/MoverComponent.h b/Engine/Plugins/Experimental/Mover/Source/Mover/Public/MoverComponent.h
--- a/Engine/Plugins/Experimental/Mover/Source/Mover/Public/MoverComponent.h
+++ b/Engine/Plugins/Experimental/Mover/Source/Mover/Public/MoverComponent.h
@@ -112,6 +112,10 @@ public:
 	UPROPERTY(EditAnywhere, BlueprintReadOnly, Category=Mover)
 	FName StartingMovementMode;
 
+	// Transition checks that are always evaluated regardless of mode. Evaluated in order, stopping at the first successful transition check. Mode-owned transitions take precedence. */
+	UPROPERTY(EditAnywhere, BlueprintReadWrite, Instanced, Category=Mover)
+	TArray<TObjectPtr<UBaseMovementModeTransition>> Transitions;
+
 	/** List of types that should always be present in this actor's sync state */
 	UPROPERTY(BlueprintReadWrite, EditAnywhere, Category=Mover)
 	TArray<FMoverDataPersistence> PersistentSyncStateDataTypes;
