Physics Mover
- Interpolate output sync state to match the interpolated physics results time
- Fix initialization of NetworkPhysicsComponent
- Remove use of Aux state
- Update API from GameThread/AsyncPhysics to External/Internal to be consistent with physics

diff --git a/Engine/Plugins/Experimental/Mover/Source/Mover/Private/Backends/MoverNetworkPhysicsLiaison.cpp b/Engine/Plugins/Experimental/Mover/Source/Mover/Private/Backends/MoverNetworkPhysicsLiaison.cpp
--- a/Engine/Plugins/Experimental/Mover/Source/Mover/Private/Backends/MoverNetworkPhysicsLiaison.cpp
+++ b/Engine/Plugins/Experimental/Mover/Source/Mover/Private/Backends/MoverNetworkPhysicsLiaison.cpp
@@ -5,11 +5,13 @@
 #include "Chaos/Character/CharacterGroundConstraint.h"
 #include "Chaos/PhysicsObject.h"
 #include "Chaos/PhysicsObjectInternalInterface.h"
+#include "Components/ShapeComponent.h"
+#include "Framework/Threading.h"
 #include "GameFramework/Actor.h"
+#include "GameFramework/Pawn.h"
 #include "MovementModeStateMachine.h"
 #include "Kinematic/Settings/CommonLegacyMovementSettings.h"
 #include "PBDRigidsSolver.h"
-#include "Components/ShapeComponent.h"
 #include "Physics/Experimental/PhysScene_Chaos.h"
 #include "PhysicsMover/Modes/PhysicsDrivenFallingMode.h"
 #include "PhysicsMover/Modes/PhysicsDrivenWalkingMode.h"
@@ -120,7 +122,7 @@ void FNetworkPhysicsMoverState::ApplyData(UActorComponent* NetworkComponent) con
 	{
 		if (UMoverNetworkPhysicsLiaisonComponent* LiaisonComp = Cast<UMoverNetworkPhysicsLiaisonComponent>(NetworkComponent))
 		{
-			LiaisonComp->SetCurrentStateData(SyncStateContext, AuxStateContext);
+			LiaisonComp->SetCurrentStateData(SyncStateContext);
 		}
 	}
 }
@@ -131,7 +133,7 @@ void FNetworkPhysicsMoverState::BuildData(const UActorComponent* NetworkComponen
 	{
 		if (const UMoverNetworkPhysicsLiaisonComponent* LiaisonComp = Cast<const UMoverNetworkPhysicsLiaisonComponent>(NetworkComponent))
 		{
-			LiaisonComp->GetCurrentStateData(SyncStateContext, AuxStateContext);
+			LiaisonComp->GetCurrentStateData(SyncStateContext);
 		}
 	}
 }
@@ -144,7 +146,6 @@ bool FNetworkPhysicsMoverState::NetSerialize(FArchive& Ar, class UPackageMap* Pa
 	{
 		FNetSerializeParams Params(Ar);
 		SyncStateContext.NetSerialize(Params);
-		AuxStateContext.NetSerialize(Params);
 		bOutSuccess = true;
 	}
 	else
@@ -162,7 +163,6 @@ void FNetworkPhysicsMoverState::InterpolateData(const FNetworkPhysicsData& MinDa
 
 	const float LerpFactor = (LocalFrame - MinState.LocalFrame) / (MaxState.LocalFrame - MinState.LocalFrame);
 	SyncStateContext.Interpolate(&MinState.SyncStateContext, &MaxState.SyncStateContext, LerpFactor);
-	AuxStateContext.Interpolate(&MinState.AuxStateContext, &MaxState.AuxStateContext, LerpFactor);
 }
 
 //////////////////////////////////////////////////////////////////////////
@@ -172,10 +172,9 @@ void UMoverNetworkPhysicsLiaisonComponent::GetCurrentInputData(OUT FMoverInputCm
 	InputCmd = NetInputCmd;
 }
 
-void UMoverNetworkPhysicsLiaisonComponent::GetCurrentStateData(OUT FMoverSyncState& SyncState, OUT FMoverAuxStateContext& AuxState) const
+void UMoverNetworkPhysicsLiaisonComponent::GetCurrentStateData(OUT FMoverSyncState& SyncState) const
 {
 	SyncState = NetSyncState;
-	AuxState = NetAuxState;
 }
 
 void UMoverNetworkPhysicsLiaisonComponent::SetCurrentInputData(const FMoverInputCmdContext& InputCmd)
@@ -183,10 +182,9 @@ void UMoverNetworkPhysicsLiaisonComponent::SetCurrentInputData(const FMoverInput
 	NetInputCmd = InputCmd;
 }
 
-void UMoverNetworkPhysicsLiaisonComponent::SetCurrentStateData(const FMoverSyncState& SyncState, const FMoverAuxStateContext& AuxState)
+void UMoverNetworkPhysicsLiaisonComponent::SetCurrentStateData(const FMoverSyncState& SyncState)
 {
 	NetSyncState = SyncState;
-	NetAuxState = AuxState;
 }
 
 //////////////////////////////////////////////////////////////////////////
@@ -199,20 +197,24 @@ UMoverNetworkPhysicsLiaisonComponent::UMoverNetworkPhysicsLiaisonComponent()
 
 	bWantsInitializeComponent = true;
 	bAutoActivate = true;
-	SetIsReplicatedByDefault(true);
 
-	// Network physics relies on movement being replicated
-	if (AActor* MyActor = GetOwner())
+	if (Chaos::FPhysicsSolverBase::IsNetworkPhysicsPredictionEnabled())
 	{
-		MyActor->SetReplicatingMovement(true);
-		MyActor->SetReplicateMovement(true);
-	}
+		// Network physics relies on movement being replicated
+		SetIsReplicatedByDefault(true);
 
-	static const FName NetworkPhysicsComponentName(TEXT("PC_NetworkPhysicsComponent"));
+		if (AActor* MyActor = GetOwner())
+		{
+			MyActor->SetReplicatingMovement(true);
+			MyActor->SetReplicateMovement(true);
+		}
 
-	NetworkPhysicsComponent = CreateDefaultSubobject<UNetworkPhysicsComponent, UNetworkPhysicsComponent>(NetworkPhysicsComponentName);
-	NetworkPhysicsComponent->SetNetAddressable(); // Make DSO components net addressable
-	NetworkPhysicsComponent->SetIsReplicated(true);
+		NetworkPhysicsComponent = NewObject<UNetworkPhysicsComponent>(GetOwner(), TEXT("PhysMover_NetworkPhysicsComponent"));
+		NetworkPhysicsComponent->SetNetAddressable(); // Make DSO components net addressable
+		NetworkPhysicsComponent->SetIsReplicated(true);
+		NetworkPhysicsComponent->RegisterComponent();
+		NetworkPhysicsComponent->InitializeComponent();
+	}
 }
 
 //////////////////////////////////////////////////////////////////////////
@@ -237,16 +239,7 @@ int32 UMoverNetworkPhysicsLiaisonComponent::GetCurrentSimFrame()
 	{
 		if (Chaos::FPhysicsSolver* Solver = Scene->GetSolver())
 		{
-			int32 Offset = 0;
-			if (NetworkPhysicsComponent && !NetworkPhysicsComponent->HasServerWorld())
-			{
-				if (APlayerController* PC = NetworkPhysicsComponent->GetPlayerController())
-				{
-					Offset = PC->GetNetworkPhysicsTickOffset();
-				}
-			}
-
-			return Solver->GetCurrentFrame() + Offset;
+			return Solver->GetCurrentFrame() + GetNetworkPhysicsTickOffset();
 		}
 	}
 
@@ -260,14 +253,9 @@ void UMoverNetworkPhysicsLiaisonComponent::OnRegister()
 {
 	Super::OnRegister();
 
-	if (!NetworkPhysicsComponent->IsRegistered())
-	{
-		NetworkPhysicsComponent->RegisterComponent();
-	}
-
 	// Need to set this here as physics creation requires access to MoverComp
 	
-	if ( (MoverComp = GetOwner()->FindComponentByClass<UMoverComponent>()) != nullptr )
+	if ((MoverComp = GetOwner()->FindComponentByClass<UMoverComponent>()) != nullptr)
 	{
 	
 		CommonMovementSettings = MoverComp->FindSharedSettings<UCommonLegacyMovementSettings>();
@@ -282,12 +270,9 @@ void UMoverNetworkPhysicsLiaisonComponent::OnRegister()
 
 void UMoverNetworkPhysicsLiaisonComponent::OnUnregister()
 {
-	if (MoverComp)
+	if (MoverComp && MoverComp->UpdatedCompAsPrimitive)
 	{
-		if (MoverComp->UpdatedCompAsPrimitive)
-		{
-			MoverComp->UpdatedCompAsPrimitive->OnComponentPhysicsStateChanged.RemoveDynamic(this, &ThisClass::OnComponentPhysicsStateChanged);
-		}
+		MoverComp->UpdatedCompAsPrimitive->OnComponentPhysicsStateChanged.RemoveDynamic(this, &ThisClass::OnComponentPhysicsStateChanged);
 	}
 
 	Super::OnUnregister();
@@ -313,7 +298,7 @@ void UMoverNetworkPhysicsLiaisonComponent::SetupConstraint()
 							Constraint->Init(CharacterProxy);
 							Solver->RegisterObject(Constraint.Get());
 
-							// Set the common settings from the initial aux data
+							// Set the common settings
 							// The rest get set every frame depending on the current movement mode
 							Constraint->SetCosMaxWalkableSlopeAngle(CommonMovementSettings->MaxWalkSlopeCosine);
 							Constraint->SetVerticalAxis(MoverComp->GetUpDirection());
@@ -369,8 +354,6 @@ void UMoverNetworkPhysicsLiaisonComponent::InitializeComponent()
 {
 	Super::InitializeComponent();
 
-	NetworkPhysicsComponent->InitializeComponent();
-
 	if (ensureAlwaysMsgf(MoverComp, TEXT("UMoverNetworkPhysicsLiaisonComponent on actor %s failed to find associated Mover component. This actor's movement will not be simulated. Verify its setup."), *GetNameSafe(GetOwner())))
 	{
 		MoverComp->InitMoverSimulation();
@@ -380,11 +363,6 @@ void UMoverNetworkPhysicsLiaisonComponent::InitializeComponent()
 
 void UMoverNetworkPhysicsLiaisonComponent::UninitializeComponent()
 {
-	if (NetworkPhysicsComponent && NetworkPhysicsComponent->HasBeenInitialized())
-	{
-		NetworkPhysicsComponent->UninitializeComponent();
-	}
-
 	Super::UninitializeComponent();
 }
 
@@ -517,7 +495,7 @@ Chaos::FUniqueIdx UMoverNetworkPhysicsLiaisonComponent::GetUniqueIdx() const
 	return Chaos::FUniqueIdx();
 }
 
-void UMoverNetworkPhysicsLiaisonComponent::UpdateConstraintSettings(const FMoverAuxStateContext& AuxState)
+void UMoverNetworkPhysicsLiaisonComponent::UpdateConstraintSettings()
 {
 	if (HasValidState())
 	{
@@ -533,26 +511,59 @@ void UMoverNetworkPhysicsLiaisonComponent::UpdateConstraintSettings(const FMover
 	}
 }
 
-FMoverTimeStep UMoverNetworkPhysicsLiaisonComponent::GetCurrentMoverTimeStep() const
+int32 UMoverNetworkPhysicsLiaisonComponent::GetNetworkPhysicsTickOffset() const
+{
+	if (NetworkPhysicsComponent && !NetworkPhysicsComponent->HasServerWorld())
+	{
+		const APlayerController* PlayerController = NetworkPhysicsComponent->GetPlayerController();
+		if (!PlayerController)
+		{
+			PlayerController = GetWorld()->GetFirstPlayerController();
+		}
+
+		if (PlayerController)
+		{
+			return PlayerController->GetNetworkPhysicsTickOffset();
+		}
+	}
+
+	return 0;
+}
+
+FMoverTimeStep UMoverNetworkPhysicsLiaisonComponent::GetCurrentMoverTimeStep_Internal() const
 {
+	Chaos::EnsureIsInPhysicsThreadContext();
+
 	FMoverTimeStep TimeStep;
 
 	if (FPhysScene* Scene = GetWorld()->GetPhysicsScene())
 	{
 		if (Chaos::FPhysicsSolver* Solver = Scene->GetSolver())
 		{
-			int32 Offset = 0;
-			if (NetworkPhysicsComponent && !NetworkPhysicsComponent->HasServerWorld())
-			{
-				if (APlayerController* PC = NetworkPhysicsComponent->GetPlayerController())
-				{
-					Offset = PC->GetNetworkPhysicsTickOffset();
-				}
-			}
+			TimeStep.ServerFrame = Solver->GetCurrentFrame() + GetNetworkPhysicsTickOffset();
+			TimeStep.StepMs = Solver->GetAsyncDeltaTime() * 1000.0f;
+			TimeStep.BaseSimTimeMs = TimeStep.ServerFrame * TimeStep.StepMs;
+			TimeStep.bIsResimulating = Solver->GetEvolution()->IsResimming();
+		}
+	}
+
+	return TimeStep;
+}
+
+FMoverTimeStep UMoverNetworkPhysicsLiaisonComponent::GetCurrentMoverTimeStep_External() const
+{
+	Chaos::EnsureIsInGameThreadContext();
+
+	FMoverTimeStep TimeStep;
 
+	if (FPhysScene* Scene = GetWorld()->GetPhysicsScene())
+	{
+		if (Chaos::FPhysicsSolver* Solver = Scene->GetSolver())
+		{
+			const int32 Offset = GetNetworkPhysicsTickOffset();
 			TimeStep.ServerFrame = Solver->GetCurrentFrame() + Offset;
 			TimeStep.StepMs = Solver->GetAsyncDeltaTime() * 1000.0f;
-			TimeStep.BaseSimTimeMs = TimeStep.ServerFrame * TimeStep.StepMs;
+			TimeStep.BaseSimTimeMs = Solver->GetPhysicsResultsTime_External() * 1000.0f + Offset * TimeStep.StepMs;
 			TimeStep.bIsResimulating = Solver->GetEvolution()->IsResimming();
 		}
 	}
@@ -560,39 +571,40 @@ FMoverTimeStep UMoverNetworkPhysicsLiaisonComponent::GetCurrentMoverTimeStep() c
 	return TimeStep;
 }
 
-void UMoverNetworkPhysicsLiaisonComponent::GameThread_ProduceInput(float DeltaSeconds, OUT FPhysicsMoverAsyncInput& Input)
+void UMoverNetworkPhysicsLiaisonComponent::ProduceInput_External(float DeltaSeconds, OUT FPhysicsMoverAsyncInput& Input)
 {
+	Chaos::EnsureIsInGameThreadContext();
+
 	if (HasValidState())
 	{
 		Input.MoverIdx = GetUniqueIdx();
 		Input.MoverSimulation = this;
-		
+
 		// Produce input
-		if (!MoverComp || !NetworkPhysicsComponent)
+		if (!MoverComp)
 		{
 			return;
 		}
 
-		if (!NetworkPhysicsComponent->HasServerWorld())
+		// If in a networked game only produce input for the locally controlled character
+		APawn* PawnOwner = Cast<APawn>(GetOwner());
+		bool bProduceInput = PawnOwner ? PawnOwner->IsLocallyControlled() : false;
+
+		if (bProduceInput)
 		{
-			APlayerController* PlayerController = NetworkPhysicsComponent->GetPlayerController();
-			if (PlayerController && PlayerController->IsLocalController())
+			if (!bCachedInputIsValid)
 			{
-				if (!bCachedInputIsValid)
-				{
-					const int DeltaTimeMS = FMath::TruncToInt(DeltaSeconds * 1000.0f);
+				const int DeltaTimeMS = FMath::RoundToInt(DeltaSeconds * 1000.0f);
+				MoverComp->ProduceInput(DeltaTimeMS, &Input.InputCmd);
 
-					MoverComp->ProduceInput(DeltaTimeMS, &Input.InputCmd);
-
-					// We only want to consume one input per physics frame
-					// so if there is already a valid cached input we use that.
-					// Input is set invalid when the async output is consumed
-					bCachedInputIsValid = true;
-				}
-				else
-				{
-					Input.InputCmd = MoverComp->CachedLastProducedInputCmd;
-				}
+				// We only want to consume one input per physics frame
+				// so if there is already a valid cached input we use that.
+				// Input is set invalid when the async output is consumed
+				bCachedInputIsValid = true;
+			}
+			else
+			{
+				Input.InputCmd = MoverComp->CachedLastProducedInputCmd;
 			}
 		}
 
@@ -600,81 +612,107 @@ void UMoverNetworkPhysicsLiaisonComponent::GameThread_ProduceInput(float DeltaSe
 		if (MoverComp->bHasValidCachedState)
 		{
 			Input.SyncState = MoverComp->CachedLastSyncState;
-			Input.AuxState = MoverComp->CachedLastAuxState;
 		}
 		else
 		{
 			FMoverDefaultSyncState& Default = Input.SyncState.SyncStateCollection.FindOrAddMutableDataByType<FMoverDefaultSyncState>();
-			Default.MovementMode = KinematicModeNames::Falling;
-			Input.AuxState = FMoverAuxStateContext();
+			Default.MovementMode = MoverComp->StartingMovementMode;
 		}
 
-		// This is required so that the physics thread can have a copy of the data to access
-		NetInputCmd = Input.InputCmd;
-		NetSyncState = Input.SyncState;
-		NetAuxState = Input.AuxState;
-
-		if (MoverComp->bHasValidLastProducedInput)
+		if (bCachedInputIsValid)
 		{
-			MoverComp->OnPreSimulationTick.Broadcast(GetCurrentMoverTimeStep(), MoverComp->CachedLastProducedInputCmd);
+			MoverComp->OnPreSimulationTick.Broadcast(GetCurrentMoverTimeStep_External(), Input.InputCmd);
 		}
 		else
 		{
-			MoverComp->OnPreSimulationTick.Broadcast(GetCurrentMoverTimeStep(), FMoverInputCmdContext());
+			MoverComp->OnPreSimulationTick.Broadcast(GetCurrentMoverTimeStep_External(), NetInputCmd);
 		}
+
+		// This is required so that the physics thread can have a copy of the data to access
+		NetInputCmd = Input.InputCmd;
+		NetSyncState = Input.SyncState;
 	}
 }
 
-void UMoverNetworkPhysicsLiaisonComponent::GameThread_ConsumeOutput(const FPhysicsMoverAsyncOutput& Output)
+void UMoverNetworkPhysicsLiaisonComponent::ConsumeOutput_External(const FPhysicsMoverAsyncOutput& Output, const double OutputTimeInSeconds)
 {
+	Chaos::EnsureIsInGameThreadContext();
+
 	if (Output.bIsValid && MoverComp)
 	{
-		if (const FMoverDefaultSyncState* OutputSyncState = Output.SyncState.SyncStateCollection.FindDataByType<FMoverDefaultSyncState>())
+		if (bCachedLastPhysicsSyncStateIsValid)
 		{
-			// If landed, broadcast OnLanded
-			if (MoverComp->bHasValidCachedState)
+			const double PhysicsResultsTime = GetWorld()->GetPhysicsScene()->GetSolver()->GetPhysicsResultsTime_External();
+
+			if (PhysicsResultsTime < OutputTimeInSeconds)
 			{
-				if (const FMoverDefaultSyncState* CachedSyncState = MoverComp->CachedLastSyncState.SyncStateCollection.FindDataByType<FMoverDefaultSyncState>())
+				const FMoverTimeStep TimeStep = GetCurrentMoverTimeStep_External();
+
+				// Output is in the future. Interpolate from the cached previous state to the next state
+				const float Denom = (OutputTimeInSeconds - CachedLastPhysicsSyncStateOutputTime);
+				if (Denom > 0)
 				{
-					if ((CachedSyncState->MovementMode == KinematicModeNames::Falling) && (OutputSyncState->MovementMode == KinematicModeNames::Walking))
+					const float Alpha = FMath::Clamp((PhysicsResultsTime - CachedLastPhysicsSyncStateOutputTime) / Denom, 0.0f, 1.0f);
+					FMoverSyncState InterpolatedSyncState;
+					InterpolatedSyncState.Interpolate(&CachedLastPhysicsSyncState, &Output.SyncState, Alpha);
+
+					// The physics mover does not set the sync state transform, it uses the physics transform,
+					// so get the world transform from the physics particle.
+					if (const Chaos::FSingleParticlePhysicsProxy* CharacterProxy = Constraint->GetCharacterParticleProxy())
 					{
-						if (UFallingMode* FallingMode = MoverComp->FindMode_Mutable<UFallingMode>())
+						const Chaos::FRigidBodyHandle_External& Body = CharacterProxy->GetGameThreadAPI();
+						if (FMoverDefaultSyncState* SyncState = InterpolatedSyncState.SyncStateCollection.FindMutableDataByType<FMoverDefaultSyncState>())
 						{
-							FHitResult HitResult;
-							MoverComp->TryGetFloorCheckHitResult(HitResult);
-							FallingMode->OnLanded.Broadcast(OutputSyncState->MovementMode, HitResult);
+							SyncState->SetTransforms_WorldSpace(Body.X(), FRotator(Body.R()), Body.V());
 						}
 					}
+
+					// If landed, broadcast OnLanded
+					// TODO: Generalize handling of events generated on physics thread
+					if (MoverComp->HasValidCachedState())
+					{
+						if (const FMoverDefaultSyncState* CachedSyncState = MoverComp->GetSyncState().SyncStateCollection.FindDataByType<FMoverDefaultSyncState>())
+						{
+							if (const FMoverDefaultSyncState* InterpolatedState = InterpolatedSyncState.SyncStateCollection.FindDataByType<FMoverDefaultSyncState>())
+							{
+								if ((CachedSyncState->MovementMode == KinematicModeNames::Falling) && (InterpolatedState->MovementMode == KinematicModeNames::Walking))
+								{
+									if (UFallingMode* FallingMode = MoverComp->FindMode_Mutable<UFallingMode>())
+									{
+										FHitResult HitResult;
+										MoverComp->TryGetFloorCheckHitResult(HitResult);
+										FallingMode->OnLanded.Broadcast(InterpolatedState->MovementMode, HitResult);
+									}
+								}
+							}
+						}
+					}
+
+					// TODO: Consider moving to a util function
+					MoverComp->CachedLastSyncState = InterpolatedSyncState;
+					MoverComp->CachedLastSimTickTimeStep = TimeStep;
+					MoverComp->bHasValidCachedState = true;
 				}
-			}
 
-			// TODO: Consider moving to a util function
-			MoverComp->CachedLastSyncState = Output.SyncState;
-			MoverComp->CachedLastAuxState = Output.AuxState;
-			MoverComp->CachedLastSimTickTimeStep.BaseSimTimeMs = GetCurrentSimTimeMs();
-			MoverComp->CachedLastSimTickTimeStep.ServerFrame = GetCurrentSimFrame();
-			MoverComp->bHasValidCachedState = true;
-		}
+				// Call on the last output for the frame
+				UpdateConstraintSettings();
 
-		FMoverTimeStep TimeStep;
-		if (FPhysScene* PhysScene = GetWorld()->GetPhysicsScene())
-		{
-			if (Chaos::FPhysicsSolver* Solver = PhysScene->GetSolver())
-			{
-				TimeStep.StepMs = Solver->GetAsyncDeltaTime() * 1000.0f;
-				TimeStep.BaseSimTimeMs = Solver->GetSolverTime();
+				MoverComp->OnPostSimulationTick.Broadcast(TimeStep);
+
+				bCachedInputIsValid = false;
 			}
 		}
-		MoverComp->OnPostSimulationTick.Broadcast(TimeStep);
-
-		bCachedInputIsValid = false;
 
-		UpdateConstraintSettings(Output.AuxState);
+		CachedLastPhysicsSyncState = Output.SyncState;
+		CachedLastPhysicsSyncStateOutputTime = OutputTimeInSeconds;
+		bCachedLastPhysicsSyncStateIsValid = true;
 	}
 }
 
-void UMoverNetworkPhysicsLiaisonComponent::AsyncPhysics_ProcessInputs(int32 PhysicsStep, float DeltaTime, const FPhysicsMoverAsyncInput& Input) const
+void UMoverNetworkPhysicsLiaisonComponent::ProcessInputs_Internal(int32 PhysicsStep, float DeltaTime, const FPhysicsMoverAsyncInput& Input) const
 {
+	Chaos::EnsureIsInPhysicsThreadContext();
+
 	// Override input data unless player is local client
 	if (FPhysScene* PhysScene = GetWorld()->GetPhysicsScene())
 	{
@@ -687,7 +725,7 @@ void UMoverNetworkPhysicsLiaisonComponent::AsyncPhysics_ProcessInputs(int32 Phys
 			}
 			else
 			{
-				GetCurrentStateData(Input.SyncState, Input.AuxState);
+				GetCurrentStateData(Input.SyncState);
 
 				bool bIsSolverResim = false;
 				if (Chaos::FPhysicsSolver* Solver = PhysScene->GetSolver())
@@ -712,11 +750,12 @@ void UMoverNetworkPhysicsLiaisonComponent::AsyncPhysics_ProcessInputs(int32 Phys
 	}
 }
 
-void UMoverNetworkPhysicsLiaisonComponent::AsyncPhysics_OnPreSimulate(const FPhysicsMoverSimulationTickParams& TickParams, const FPhysicsMoverAsyncInput& Input, OUT FPhysicsMoverAsyncOutput& Output) const
+void UMoverNetworkPhysicsLiaisonComponent::OnPreSimulate_Internal(const FPhysicsMoverSimulationTickParams& TickParams, const FPhysicsMoverAsyncInput& Input, OUT FPhysicsMoverAsyncOutput& Output) const
 {
-	// Sync/aux state should carry over to the next sim frame by default unless something modifies it
+	Chaos::EnsureIsInPhysicsThreadContext();
+
+	// Sync state should carry over to the next sim frame by default unless something modifies it
 	Output.SyncState = Input.SyncState;
-	Output.AuxState = Input.AuxState;
 
 	if (!HasValidState() || !Input.InputCmd.InputCollection.FindDataByType<FKinematicDefaultInputs>())
 	{
@@ -772,11 +811,11 @@ void UMoverNetworkPhysicsLiaisonComponent::AsyncPhysics_OnPreSimulate(const FPhy
 	// Update the simulation
 	//Input.MoverSimulation->SimulationTick(TickParams, Input.SimInput, SimOutput);
 
-	FMoverTickStartData TickStartData(Input.InputCmd, Input.SyncState, Input.AuxState);
+	FMoverTickStartData TickStartData(Input.InputCmd, Input.SyncState, FMoverAuxStateContext());
 	FMoverTickEndData TickEndData;
 	FFloorCheckResult FloorResult;
 
-	FMoverTimeStep TimeStep = GetCurrentMoverTimeStep();
+	FMoverTimeStep TimeStep = GetCurrentMoverTimeStep_Internal();
 
 	bool bHasRolledBack = false; // TODO
 
@@ -788,9 +827,8 @@ void UMoverNetworkPhysicsLiaisonComponent::AsyncPhysics_OnPreSimulate(const FPhy
 		MoverComp->ProcessFirstSimTickAfterRollback(TimeStep);
 	}
 
-	// Sync/aux state should carry over to the next sim frame by default unless something modifies it
+	// Sync state should carry over to the next sim frame by default unless something modifies it
 	TickEndData.SyncState = TickStartData.SyncState;
-	TickEndData.AuxState = TickStartData.AuxState;
 
 	if (MoverComp->ModeFSM->IsValidLowLevel() && MoverComp->SimBlackboard->IsValidLowLevel())
 	{
@@ -817,11 +855,10 @@ void UMoverNetworkPhysicsLiaisonComponent::AsyncPhysics_OnPreSimulate(const FPhy
 	MoverComp->SimBlackboard->TryGet(KinematicBlackboard::LastFloorResult, Output.FloorResult);
 
 	Output.SyncState = TickEndData.SyncState;
-	Output.AuxState = TickEndData.AuxState;
 
 	/////////////////////////////////////////////////////////////////////////////////////////////////////////////////////////////
 
-	// Update physics constraint from output sync state and aux state
+	// Update physics constraint from output sync state
 	const FMoverDefaultSyncState* OutputSyncState = Output.SyncState.SyncStateCollection.FindDataByType<FMoverDefaultSyncState>();
 	if (!ensure(OutputSyncState))
 	{
@@ -962,8 +999,10 @@ void UMoverNetworkPhysicsLiaisonComponent::WakeParticleIfSleeping(Chaos::FGeomet
 	}
 }
 
-void UMoverNetworkPhysicsLiaisonComponent::ASyncPhysics_OnContactModification(const FPhysicsMoverAsyncInput& Input, Chaos::FCollisionContactModifier& Modifier) const
+void UMoverNetworkPhysicsLiaisonComponent::OnContactModification_Internal(const FPhysicsMoverAsyncInput& Input, Chaos::FCollisionContactModifier& Modifier) const
 {
+	Chaos::EnsureIsInPhysicsThreadContext();
+
 	if (!HasValidState())
 	{
 		return;
diff --git a/Engine/Plugins/Experimental/Mover/Source/Mover/Private/PhysicsMover/PhysicsMoverManager.cpp b/Engine/Plugins/Experimental/Mover/Source/Mover/Private/PhysicsMover/PhysicsMoverManager.cpp
--- a/Engine/Plugins/Experimental/Mover/Source/Mover/Private/PhysicsMover/PhysicsMoverManager.cpp
+++ b/Engine/Plugins/Experimental/Mover/Source/Mover/Private/PhysicsMover/PhysicsMoverManager.cpp
@@ -71,6 +71,7 @@ void UPhysicsMoverManager::PrePhysicsUpdate(FPhysScene* PhysScene, float DeltaTi
 
 	while (Chaos::TSimCallbackOutputHandle<FPhysicsMoverManagerAsyncOutput> ManagerAsyncOutput = AsyncCallback->PopFutureOutputData_External())
 	{
+		const double OutputTime = ManagerAsyncOutput->InternalTime;
 		for (TWeakObjectPtr<UMoverNetworkPhysicsLiaisonComponent> PhysicsMoverComp : PhysicsMoverComponents)
 		{
 			Chaos::FUniqueIdx Idx = PhysicsMoverComp->GetUniqueIdx();
@@ -78,7 +79,7 @@ void UPhysicsMoverManager::PrePhysicsUpdate(FPhysScene* PhysScene, float DeltaTi
 			{
 				if (TUniquePtr<FPhysicsMoverAsyncOutput>* OutputData = ManagerAsyncOutput->PhysicsMoverToAsyncOutput.Find(Idx))
 				{
-					PhysicsMoverComp->GameThread_ConsumeOutput(**OutputData);
+					PhysicsMoverComp->ConsumeOutput_External(**OutputData, OutputTime);
 				}
 			}
 		}
@@ -93,7 +94,7 @@ void UPhysicsMoverManager::PrePhysicsUpdate(FPhysScene* PhysScene, float DeltaTi
 		for (TWeakObjectPtr<UMoverNetworkPhysicsLiaisonComponent> PhysicsMoverComp : PhysicsMoverComponents)
 		{
 			TUniquePtr<FPhysicsMoverAsyncInput> InputData = MakeUnique<FPhysicsMoverAsyncInput>();
-			PhysicsMoverComp->GameThread_ProduceInput(DeltaTime, *InputData);
+			PhysicsMoverComp->ProduceInput_External(DeltaTime, *InputData);
 			ManagerAsyncInput->AsyncInput.Add(MoveTemp(InputData));
 		}
 	}
diff --git a/Engine/Plugins/Experimental/Mover/Source/Mover/Private/PhysicsMover/PhysicsMoverManagerAsyncCallback.cpp b/Engine/Plugins/Experimental/Mover/Source/Mover/Private/PhysicsMover/PhysicsMoverManagerAsyncCallback.cpp
--- a/Engine/Plugins/Experimental/Mover/Source/Mover/Private/PhysicsMover/PhysicsMoverManagerAsyncCallback.cpp
+++ b/Engine/Plugins/Experimental/Mover/Source/Mover/Private/PhysicsMover/PhysicsMoverManagerAsyncCallback.cpp
@@ -33,7 +33,7 @@ void FPhysicsMoverManagerAsyncCallback::ProcessInputs_Internal(int32 PhysicsStep
 		const FPhysicsMoverAsyncInput& AsyncInput = *InputDataArray[Idx];
 		if (AsyncInput.IsValid())
 		{
-			AsyncInput.MoverSimulation->AsyncPhysics_ProcessInputs(PhysicsStep, DeltaTime, AsyncInput);
+			AsyncInput.MoverSimulation->ProcessInputs_Internal(PhysicsStep, DeltaTime, AsyncInput);
 		}
 	};
 
@@ -65,7 +65,7 @@ void FPhysicsMoverManagerAsyncCallback::OnPreSimulate_Internal()
 		if (AsyncInput.IsValid())
 		{
 			TUniquePtr<FPhysicsMoverAsyncOutput> AsyncOutput = MakeUnique<FPhysicsMoverAsyncOutput>();
-			AsyncInput.MoverSimulation->AsyncPhysics_OnPreSimulate(TickParams, AsyncInput, *AsyncOutput);
+			AsyncInput.MoverSimulation->OnPreSimulate_Internal(TickParams, AsyncInput, *AsyncOutput);
 			ManagerAsyncOutput.PhysicsMoverToAsyncOutput.Add(AsyncInput.MoverIdx, MoveTemp(AsyncOutput));
 		}
 	};
@@ -95,7 +95,7 @@ void FPhysicsMoverManagerAsyncCallback::OnContactModification_Internal(Chaos::FC
 		const FPhysicsMoverAsyncInput& AsyncInput = *InputDataArray[Idx];
 		if (AsyncInput.IsValid())
 		{
-			AsyncInput.MoverSimulation->ASyncPhysics_OnContactModification(AsyncInput, Modifier);
+			AsyncInput.MoverSimulation->OnContactModification_Internal(AsyncInput, Modifier);
 		}
 	};
 
diff --git a/Engine/Plugins/Experimental/Mover/Source/Mover/Public/Backends/MoverNetworkPhysicsLiaison.h b/Engine/Plugins/Experimental/Mover/Source/Mover/Public/Backends/MoverNetworkPhysicsLiaison.h
--- a/Engine/Plugins/Experimental/Mover/Source/Mover/Public/Backends/MoverNetworkPhysicsLiaison.h
+++ b/Engine/Plugins/Experimental/Mover/Source/Mover/Public/Backends/MoverNetworkPhysicsLiaison.h
@@ -57,9 +57,6 @@ struct MOVER_API FNetworkPhysicsMoverState : public FNetworkPhysicsData
 	UPROPERTY()
 	FMoverSyncState SyncStateContext;
 
-	UPROPERTY()
-	FMoverAuxStateContext AuxStateContext;
-
 	/**  Apply the data onto the network physics component */
 	virtual void ApplyData(UActorComponent* NetworkComponent) const override;
 
@@ -121,18 +118,18 @@ public:
 	// Required for Network Physics Rewind/Resim data
 	void GetCurrentInputData(OUT FMoverInputCmdContext& InputCmd) const;
 	void SetCurrentInputData(const FMoverInputCmdContext& InputCmd);
-	void GetCurrentStateData(OUT FMoverSyncState& SyncState, OUT FMoverAuxStateContext& AuxState) const;
-	void SetCurrentStateData(const FMoverSyncState& SyncState, const FMoverAuxStateContext& AuxState);
+	void GetCurrentStateData(OUT FMoverSyncState& SyncState) const;
+	void SetCurrentStateData(const FMoverSyncState& SyncState);
 
 	// Used by the manager to uniquely identify the component
 	Chaos::FUniqueIdx GetUniqueIdx() const;
 
-	void GameThread_ProduceInput(float DeltaSeconds, OUT FPhysicsMoverAsyncInput& Input);
-	void GameThread_ConsumeOutput(const FPhysicsMoverAsyncOutput& Output);
+	void ProduceInput_External(float DeltaSeconds, OUT FPhysicsMoverAsyncInput& Input);
+	void ConsumeOutput_External(const FPhysicsMoverAsyncOutput& Output, const double OutputTimeInSeconds);
 
-	void AsyncPhysics_ProcessInputs(int32 PhysicsStep, float DeltaTime, const FPhysicsMoverAsyncInput& Input) const;
-	void AsyncPhysics_OnPreSimulate(const FPhysicsMoverSimulationTickParams& TickParams, const FPhysicsMoverAsyncInput& Input, OUT FPhysicsMoverAsyncOutput& SimOutput) const;
-	void ASyncPhysics_OnContactModification(const FPhysicsMoverAsyncInput& Input, Chaos::FCollisionContactModifier& Modifier) const;
+	void ProcessInputs_Internal(int32 PhysicsStep, float DeltaTime, const FPhysicsMoverAsyncInput& Input) const;
+	void OnPreSimulate_Internal(const FPhysicsMoverSimulationTickParams& TickParams, const FPhysicsMoverAsyncInput& Input, OUT FPhysicsMoverAsyncOutput& SimOutput) const;
+	void OnContactModification_Internal(const FPhysicsMoverAsyncInput& Input, Chaos::FCollisionContactModifier& Modifier) const;
 
 protected:
 	UFUNCTION()
@@ -143,20 +140,35 @@ protected:
 	void DestroyConstraint();
 	void SetupConstraint();
 
-	void UpdateConstraintSettings(const FMoverAuxStateContext& AuxState);
+	void UpdateConstraintSettings();
 	void TeleportParticle(Chaos::FGeometryParticleHandle* Particle, const FVector& Position, const FQuat& Rotation) const;
 	void WakeParticleIfSleeping(Chaos::FGeometryParticleHandle* Particle) const;
 
-	FMoverTimeStep GetCurrentMoverTimeStep() const;
+	int32 GetNetworkPhysicsTickOffset() const;
+
+	// Time step on the physics thread
+	FMoverTimeStep GetCurrentMoverTimeStep_Internal() const;
+
+	// Time step on the game thread. Uses physics results time
+	FMoverTimeStep GetCurrentMoverTimeStep_External() const;
 
 	// These are written to by the network input and state data
 	FMoverInputCmdContext NetInputCmd;
 	FMoverSyncState NetSyncState;
-	FMoverAuxStateContext NetAuxState;
 
 	TUniquePtr<Chaos::FCharacterGroundConstraint> Constraint;
-	TObjectPtr<UNetworkPhysicsComponent> NetworkPhysicsComponent;
 	TObjectPtr<UMoverComponent> MoverComp;
 	TObjectPtr<const UCommonLegacyMovementSettings> CommonMovementSettings;
+
+	TObjectPtr<UNetworkPhysicsComponent> NetworkPhysicsComponent;
+
+	// The cached physics state is the latest output from the physics thread
+	// This can be different from the cached sync state on the mover component
+	// which is interpolated to match the interpolated physics particle
+	// Note: Time is physics solver time
+	FMoverSyncState CachedLastPhysicsSyncState;
+	double CachedLastPhysicsSyncStateOutputTime;
+	bool bCachedLastPhysicsSyncStateIsValid = false;
+
 	bool bCachedInputIsValid = false;
 };
\ No newline at end of file
diff --git a/Engine/Plugins/Experimental/Mover/Source/Mover/Public/PhysicsMover/PhysicsMoverSimulationTypes.h b/Engine/Plugins/Experimental/Mover/Source/Mover/Public/PhysicsMover/PhysicsMoverSimulationTypes.h
--- a/Engine/Plugins/Experimental/Mover/Source/Mover/Public/PhysicsMover/PhysicsMoverSimulationTypes.h
+++ b/Engine/Plugins/Experimental/Mover/Source/Mover/Public/PhysicsMover/PhysicsMoverSimulationTypes.h
@@ -33,23 +33,6 @@ struct FPhysicsMoverSimulationTickParams
 	float DeltaTimeSeconds = 0.0f;
 };
 
-struct FPhysicsMoverSimulationInput
-{
-	FMoverInputCmdContext InputCmd;
-	FMoverSyncState SyncState;
-	FMoverAuxStateContext AuxState;
-};
-
-struct FPhysicsMoverSimulationOutput
-{
-	bool IsValid() const { return bIsValid; }
-
-	FMoverSyncState SyncState;
-	FMoverAuxStateContext AuxState;
-	FFloorCheckResult FloorResult;
-	bool bIsValid = false;
-};
-
 namespace Chaos { class FCharacterGroundConstraintHandle; }
 
 struct FPhysicsMoverAsyncInput
@@ -62,7 +45,6 @@ struct FPhysicsMoverAsyncInput
 	// Input is modified during ProcessInputs_Internal
 	mutable FMoverInputCmdContext InputCmd;
 	mutable FMoverSyncState SyncState;
-	mutable FMoverAuxStateContext AuxState;
 
 	TWeakObjectPtr<class UMoverNetworkPhysicsLiaisonComponent> MoverSimulation;
 	Chaos::FUniqueIdx MoverIdx;
@@ -71,7 +53,6 @@ struct FPhysicsMoverAsyncInput
 struct FPhysicsMoverAsyncOutput
 {
 	FMoverSyncState SyncState;
-	FMoverAuxStateContext AuxState;
 	FFloorCheckResult FloorResult;
 	bool bIsValid = false;
 };
