Physics Mover - Only set input as valid if OnProduceInput actually produced an input

diff --git a/Engine/Plugins/Experimental/Mover/Source/Mover/Private/Backends/MoverNetworkPhysicsLiaison.cpp b/Engine/Plugins/Experimental/Mover/Source/Mover/Private/Backends/MoverNetworkPhysicsLiaison.cpp
--- a/Engine/Plugins/Experimental/Mover/Source/Mover/Private/Backends/MoverNetworkPhysicsLiaison.cpp
+++ b/Engine/Plugins/Experimental/Mover/Source/Mover/Private/Backends/MoverNetworkPhysicsLiaison.cpp
@@ -600,7 +600,7 @@ void UMoverNetworkPhysicsLiaisonComponent::ProduceInput_External(float DeltaSeco
 				// We only want to consume one input per physics frame
 				// so if there is already a valid cached input we use that.
 				// Input is set invalid when the async output is consumed
-				bCachedInputIsValid = true;
+				bCachedInputIsValid = MoverComp->CachedLastProducedInputCmd.InputCollection.FindDataByType<FKinematicDefaultInputs>() != nullptr;
 			}
 			else
 			{
