Mover 2.0 adding some additional info to a couple ensure messages so they are more helpful when hit

diff --git a/Engine/Plugins/Experimental/Mover/Source/Mover/Private/Kinematic/LayeredMoves/BasicLayeredMoves.cpp b/Engine/Plugins/Experimental/Mover/Source/Mover/Private/Kinematic/LayeredMoves/BasicLayeredMoves.cpp
--- a/Engine/Plugins/Experimental/Mover/Source/Mover/Private/Kinematic/LayeredMoves/BasicLayeredMoves.cpp
+++ b/Engine/Plugins/Experimental/Mover/Source/Mover/Private/Kinematic/LayeredMoves/BasicLayeredMoves.cpp
@@ -120,7 +120,7 @@ bool FLayeredMove_JumpImpulse::GenerateMove(const FMoverTickStartData& SimState,
 	}
 	else
 	{
-		ensure(0); // other mix modes not supported yet
+		ensureMsgf(false, TEXT("JumpImpulse layered move only supports Override Velocity mix mode and was queued with a different mix mode. Layered move will do nothing."));
 		return false;
 	}
 
diff --git a/Engine/Plugins/Experimental/Mover/Source/Mover/Private/Kinematic/LayeredMoves/MultiJumpLayeredMove.cpp b/Engine/Plugins/Experimental/Mover/Source/Mover/Private/Kinematic/LayeredMoves/MultiJumpLayeredMove.cpp
--- a/Engine/Plugins/Experimental/Mover/Source/Mover/Private/Kinematic/LayeredMoves/MultiJumpLayeredMove.cpp
+++ b/Engine/Plugins/Experimental/Mover/Source/Mover/Private/Kinematic/LayeredMoves/MultiJumpLayeredMove.cpp
@@ -132,7 +132,7 @@ bool FLayeredMove_MultiJump::PerformJump(const FMoverDefaultSyncState* SyncState
 	}
 	else
 	{
-		ensure(0); // other mix modes not supported yet
+		ensureMsgf(false, TEXT("Multi-Jump layered move only supports Override Velocity mix mode and was queued with a different mix mode. Layered move will do nothing."));
 		return false;
 	}
 
