Mover 2.0 first pass at a layered move priority system to allow specific layered moves to override others

diff --git a/Engine/Plugins/Experimental/Mover/Source/Mover/Private/LayeredMove.cpp b/Engine/Plugins/Experimental/Mover/Source/Mover/Private/LayeredMove.cpp
--- a/Engine/Plugins/Experimental/Mover/Source/Mover/Private/LayeredMove.cpp
+++ b/Engine/Plugins/Experimental/Mover/Source/Mover/Private/LayeredMove.cpp
@@ -31,6 +31,7 @@ void FLayeredMoveFinishVelocitySettings::NetSerialize(FArchive& Ar)
 
 FLayeredMoveBase::FLayeredMoveBase() :
 	MixMode(EMoveMixMode::AdditiveVelocity),
+	Priority(0),
 	DurationMs(-1.f),
 	StartSimTimeMs(LayeredMove_InvalidTime)
 {
@@ -71,6 +72,13 @@ void FLayeredMoveBase::NetSerialize(FArchive& Ar)
 	Ar << MixModeAsU8;
 	MixMode = (EMoveMixMode)MixModeAsU8;
 
+	uint8 bHasDefaultPriority = Priority == 0;
+	Ar.SerializeBits(&bHasDefaultPriority, 1);
+	if (!bHasDefaultPriority)
+	{
+		Ar << Priority;
+	}
+	
 	Ar << DurationMs;
 	Ar << StartSimTimeMs;
 
@@ -118,6 +126,8 @@ bool FLayeredMoveGroup::DoGenerateMove(const FMoverTickStartData& StartState, co
 	CumulativeMoveStep.MixMode = EMoveMixMode::AdditiveVelocity;
 
 	bool bDidAccumulateAnyMoves = false;
+	uint8 HighestPriority = 0;
+	float CurrentLayeredMoveStartTime = TNumericLimits<float>::Max();
 
 	// Tick and accumulate all active moves
 	// Gather all proposed moves and distill this into a final movement report. May include separate additive vs override moves.
@@ -126,23 +136,22 @@ bool FLayeredMoveGroup::DoGenerateMove(const FMoverTickStartData& StartState, co
 	{
 		FProposedMove MoveStep;
 		
-		if (ActiveMove->GenerateMove(StartState, TimeStep, MoverComp, SimBlackboard, MoveStep) &&
-			CumulativeMoveStep.MixMode < EMoveMixMode::OverrideAll)
+		if (ActiveMove->GenerateMove(StartState, TimeStep, MoverComp, SimBlackboard, MoveStep))
 		{
 			bDidAccumulateAnyMoves = true;
 
 			if (CumulativeMoveStep.PreferredMode != MoveStep.PreferredMode &&
 				!CumulativeMoveStep.PreferredMode.IsNone() && !MoveStep.PreferredMode.IsNone())
 			{
-				UE_LOG(LogMover, Warning, TEXT("Multiple LayeredMoves are conflicting with preferred moves. %s will override %s"),
+				UE_LOG(LogMover, Log, TEXT("Multiple LayeredMoves are conflicting with preferred moves. %s will override %s"),
 					*MoveStep.PreferredMode.ToString(), *CumulativeMoveStep.PreferredMode.ToString());
 			}
 
-			if (MoveStep.bHasDirIntent && CumulativeMoveStep.MixMode != EMoveMixMode::OverrideAll)
+			if (MoveStep.bHasDirIntent && CumulativeMoveStep.MixMode != EMoveMixMode::OverrideAll && ActiveMove->Priority >= HighestPriority)
 			{
 				if (CumulativeMoveStep.bHasDirIntent)
 				{
-					UE_LOG(LogMover, Warning, TEXT("Multiple LayeredMoves are setting direction intent and only one direction intent will be used."));
+					UE_LOG(LogMover, Log, TEXT("Multiple LayeredMoves are setting direction intent and the layered move with highest priority will be used."));
 				}
 				
 				CumulativeMoveStep.bHasDirIntent = MoveStep.bHasDirIntent;
@@ -151,19 +160,22 @@ bool FLayeredMoveGroup::DoGenerateMove(const FMoverTickStartData& StartState, co
 
 			if (ActiveMove->MixMode == EMoveMixMode::OverrideVelocity)
 			{
-				if (CumulativeMoveStep.MixMode == EMoveMixMode::OverrideVelocity || CumulativeMoveStep.MixMode == EMoveMixMode::OverrideAll)
+				if (CheckPriority(ActiveMove.Get(), HighestPriority, CurrentLayeredMoveStartTime))
 				{
-					UE_LOG(LogMover, Warning, TEXT("Multiple LayeredMoves with Override mix mode are active simultaneously. Only one will take effect."));
-				}
+					if (CumulativeMoveStep.MixMode == EMoveMixMode::OverrideVelocity || CumulativeMoveStep.MixMode == EMoveMixMode::OverrideAll)
+					{
+						UE_LOG(LogMover, Log, TEXT("Multiple LayeredMoves with Override mix mode are active simultaneously. Layered move with the highest priority will take effect."));
+					}
 
-				if (!MoveStep.PreferredMode.IsNone())
-				{
-					CumulativeMoveStep.PreferredMode = MoveStep.PreferredMode;
-				}
+					if (!MoveStep.PreferredMode.IsNone())
+					{
+						CumulativeMoveStep.PreferredMode = MoveStep.PreferredMode;
+					}
 				
-				CumulativeMoveStep.MixMode = EMoveMixMode::OverrideVelocity;
-				CumulativeMoveStep.LinearVelocity  = MoveStep.LinearVelocity;
-				CumulativeMoveStep.AngularVelocity = MoveStep.AngularVelocity;
+					CumulativeMoveStep.MixMode = EMoveMixMode::OverrideVelocity;
+					CumulativeMoveStep.LinearVelocity  = MoveStep.LinearVelocity;
+					CumulativeMoveStep.AngularVelocity = MoveStep.AngularVelocity;
+				}
 			}
 			else if (ActiveMove->MixMode == EMoveMixMode::AdditiveVelocity)
 			{
@@ -181,12 +193,15 @@ bool FLayeredMoveGroup::DoGenerateMove(const FMoverTickStartData& StartState, co
 			}
 			else if (ActiveMove->MixMode == EMoveMixMode::OverrideAll)
 			{
-				if (CumulativeMoveStep.MixMode == EMoveMixMode::OverrideVelocity || CumulativeMoveStep.MixMode == EMoveMixMode::OverrideAll)
+				if (CheckPriority(ActiveMove.Get(), HighestPriority, CurrentLayeredMoveStartTime))
 				{
-					UE_LOG(LogMover, Warning, TEXT("Multiple LayeredMoves with Override mix mode are active simultaneously. Only one will take effect."));
-				}
+					if (CumulativeMoveStep.MixMode == EMoveMixMode::OverrideVelocity || CumulativeMoveStep.MixMode == EMoveMixMode::OverrideAll)
+					{
+						UE_LOG(LogMover, Log, TEXT("Multiple LayeredMoves with Override mix mode are active simultaneously. Layered move with the highest priority will take effect."));
+					}
 				
-				CumulativeMoveStep = MoveStep;
+					CumulativeMoveStep = MoveStep;
+				}
 			}
 			else
 			{
@@ -454,6 +469,23 @@ void FLayeredMoveGroup::GatherResidualVelocitySettings(const TSharedPtr<FLayered
 	}
 }
 
+bool FLayeredMoveGroup::CheckPriority(const FLayeredMoveBase* LayeredMove, uint8& InOutHighestPriority, float& InOutCurrentLayeredMoveStartTimeMs)
+{
+	if (LayeredMove->Priority > InOutHighestPriority)
+	{
+		InOutHighestPriority = LayeredMove->Priority;
+		InOutCurrentLayeredMoveStartTimeMs = LayeredMove->StartSimTimeMs;
+		return true;
+	}
+	if (LayeredMove->Priority == InOutHighestPriority && LayeredMove->StartSimTimeMs < InOutCurrentLayeredMoveStartTimeMs)
+	{
+		InOutCurrentLayeredMoveStartTimeMs = LayeredMove->StartSimTimeMs;
+		return true;
+	}
+
+	return false;
+}
+
 struct FLayeredMoveDeleter
 {
 	FORCEINLINE void operator()(FLayeredMoveBase* Object) const
diff --git a/Engine/Plugins/Experimental/Mover/Source/Mover/Public/LayeredMove.h b/Engine/Plugins/Experimental/Mover/Source/Mover/Public/LayeredMove.h
--- a/Engine/Plugins/Experimental/Mover/Source/Mover/Public/LayeredMove.h
+++ b/Engine/Plugins/Experimental/Mover/Source/Mover/Public/LayeredMove.h
@@ -75,6 +75,10 @@ struct MOVER_API FLayeredMoveBase
 	UPROPERTY(EditAnywhere, BlueprintReadWrite, Category = Mover)
 	EMoveMixMode MixMode;
 
+	// Determines if this layered move should take priority over other layered moves when different moves have conflicting overrides - higher numbers taking precedent.
+	UPROPERTY(EditAnywhere, BlueprintReadWrite, Category = Mover)
+	uint8 Priority;
+	
 	// This move will expire after a set amount of time if > 0. If 0, it will be ticked only once, regardless of time step. It will need to be manually ended if < 0. 
 	UPROPERTY(EditAnywhere, BlueprintReadWrite, Category = Mover)
 	float DurationMs;
@@ -197,6 +201,12 @@ protected:
 
 	// Helper function for gathering any residual velocity settings from layered moves that just ended
 	void GatherResidualVelocitySettings(const TSharedPtr<FLayeredMoveBase>& Move, bool& bResidualVelocityOverriden, bool& bClampVelocityOverriden);
+
+	/**
+	 * Helper function for layered move mixing to check priority and start time if priority is the same.
+	 * Returns true if this layered move should take priority given current HighestPriority and CurrentLayeredMoveStartTimeMs
+	 */
+	static bool CheckPriority(const FLayeredMoveBase* LayeredMove, uint8& InOutHighestPriority, float& InOutCurrentLayeredMoveStartTimeMs);
 	
 	/** Helper function for serializing array of root motion sources */
 	static void NetSerializeLayeredMovesArray(FArchive& Ar, TArray< TSharedPtr<FLayeredMoveBase> >& LayeredMovesArray, uint8 MaxNumLayeredMovesToSerialize = MAX_uint8);
