Physics Mover
- Ensure transition to landing only happens if character is moving towards the ground

diff --git a/Engine/Plugins/Experimental/Mover/Source/Mover/Private/Backends/MoverNetworkPhysicsLiaison.cpp b/Engine/Plugins/Experimental/Mover/Source/Mover/Private/Backends/MoverNetworkPhysicsLiaison.cpp
--- a/Engine/Plugins/Experimental/Mover/Source/Mover/Private/Backends/MoverNetworkPhysicsLiaison.cpp
+++ b/Engine/Plugins/Experimental/Mover/Source/Mover/Private/Backends/MoverNetworkPhysicsLiaison.cpp
@@ -1035,9 +1035,9 @@ void UMoverNetworkPhysicsLiaisonComponent::OnContactModification_Internal(const
 		const float CharacterHeight = CharacterParticle->GetX().Z;
 		const float EndCapHeight = CharacterHeight - PawnHalfHeight + PawnRadius;
 
-		const float CosThetaMax = 0.97f;
+		const float CosThetaMax = 0.707f;
 
-		float MinContactHeightStepUps = CharacterHeight + 1.0e10f;
+		float MinContactHeightStepUps = CharacterHeight - 1.0e10f;
 		if (SyncState->MovementMode == KinematicModeNames::Walking)
 		{
 			if (const UPhysicsDrivenWalkingMode* WalkingMode = Cast<UPhysicsDrivenWalkingMode>(MoverComp->FindMode_Mutable<UPhysicsDrivenWalkingMode>()))
@@ -1061,26 +1061,23 @@ void UMoverNetworkPhysicsLiaisonComponent::OnContactModification_Internal(const
 			const int32 CharacterIdx = CharacterParticle == PairModifier.GetParticlePair()[0] ? 0 : 1;
 			const int32 OtherIdx = CharacterIdx == 0 ? 1 : 0;
 
-			if (GroundParticle == PairModifier.GetParticlePair()[OtherIdx])
+			for (int32 Idx = 0; Idx < PairModifier.GetNumContacts(); ++Idx)
 			{
-				for (int32 Idx = 0; Idx < PairModifier.GetNumContacts(); ++Idx)
-				{
-					Chaos::FVec3 Point0, Point1;
-					PairModifier.GetWorldContactLocations(Idx, Point0, Point1);
-					Chaos::FVec3 CharacterPoint = CharacterIdx == 0 ? Point0 : Point1;
+				Chaos::FVec3 Point0, Point1;
+				PairModifier.GetWorldContactLocations(Idx, Point0, Point1);
+				Chaos::FVec3 CharacterPoint = CharacterIdx == 0 ? Point0 : Point1;
 
-					Chaos::FVec3 ContactNormal = PairModifier.GetWorldNormal(Idx);
-					if ((ContactNormal.Z > CosThetaMax) && CharacterPoint.Z < EndCapHeight)
-					{
-						// Disable any nearly vertical contact with the end cap of the capsule
-						// This will be handled by the character ground constraint
-						PairModifier.SetContactPointDisabled(Idx);
-					}
-					else if (CharacterPoint.Z < MinContactHeightStepUps)
-					{
-						// In the case of steps ups disable all contacts below the max step height
-						PairModifier.SetContactPointDisabled(Idx);
-					}
+				Chaos::FVec3 ContactNormal = PairModifier.GetWorldNormal(Idx);
+				if ((ContactNormal.Z > CosThetaMax) && CharacterPoint.Z < EndCapHeight)
+				{
+					// Disable any nearly vertical contact with the end cap of the capsule
+					// This will be handled by the character ground constraint
+					PairModifier.SetContactPointDisabled(Idx);
+				}
+				else if ((CharacterPoint.Z < MinContactHeightStepUps) && (GroundParticle == PairModifier.GetParticlePair()[OtherIdx]))
+				{
+					// In the case of steps ups disable all contacts below the max step height
+					PairModifier.SetContactPointDisabled(Idx);
 				}
 			}
 		}
diff --git a/Engine/Plugins/Experimental/Mover/Source/Mover/Private/PhysicsMover/Modes/PhysicsDrivenFallingMode.cpp b/Engine/Plugins/Experimental/Mover/Source/Mover/Private/PhysicsMover/Modes/PhysicsDrivenFallingMode.cpp
--- a/Engine/Plugins/Experimental/Mover/Source/Mover/Private/PhysicsMover/Modes/PhysicsDrivenFallingMode.cpp
+++ b/Engine/Plugins/Experimental/Mover/Source/Mover/Private/PhysicsMover/Modes/PhysicsDrivenFallingMode.cpp
@@ -104,10 +104,13 @@ void UPhysicsDrivenFallingMode::OnSimulationTick(const FSimulationTickParams& Pa
 	FVector TargetPos = StartingSyncState->GetLocation_WorldSpace() + TargetVel * DeltaSeconds;
 
 	constexpr float FloorDistanceTolerance = 2.0f;
-	const float ProjectedFloorDistance = FloorResult.FloorDist + UpDir.Dot(ProposedMove.LinearVelocity) * DeltaSeconds;
+	const FVector ProjectedGroundVelocity = UPhysicsMovementUtils::ComputeIntegratedGroundVelocityFromHitResult(StartingSyncState->GetLocation_WorldSpace(), FloorResult.HitResult, DeltaSeconds);
+	const float ProjectedRelativeVerticalVelocity = FloorResult.HitResult.ImpactNormal.Dot(ProposedMove.LinearVelocity - ProjectedGroundVelocity);
+	const float ProjectedFloorDistance = FloorResult.FloorDist + ProjectedRelativeVerticalVelocity * DeltaSeconds;
 	const bool bIsFloorWithinReach = ProjectedFloorDistance < TargetHeight + FloorDistanceTolerance;
+	const bool bIsMovingUpRelativeToFloor = ProjectedRelativeVerticalVelocity > UE_KINDA_SMALL_NUMBER;
 
-	if (FloorResult.IsWalkableFloor() && bIsFloorWithinReach)
+	if (FloorResult.IsWalkableFloor() && bIsFloorWithinReach && !bIsMovingUpRelativeToFloor)
 	{
 		OutputState.MovementEndState.NextModeName = KinematicModeNames::Walking;
 		TargetPos -= UpDir * (FloorResult.FloorDist - TargetHeight);
diff --git a/Engine/Plugins/Experimental/Mover/Source/Mover/Private/PhysicsMover/Modes/PhysicsDrivenWalkingMode.cpp b/Engine/Plugins/Experimental/Mover/Source/Mover/Private/PhysicsMover/Modes/PhysicsDrivenWalkingMode.cpp
--- a/Engine/Plugins/Experimental/Mover/Source/Mover/Private/PhysicsMover/Modes/PhysicsDrivenWalkingMode.cpp
+++ b/Engine/Plugins/Experimental/Mover/Source/Mover/Private/PhysicsMover/Modes/PhysicsDrivenWalkingMode.cpp
@@ -123,23 +123,8 @@ void UPhysicsDrivenWalkingMode::OnSimulationTick(const FSimulationTickParams& Pa
 		{
 			CharacterGravity = PhysVolume->GetGravityZ();
 		}
-		const FVector ProjectedVelocity = TargetVelocity + CharacterGravity * UpDir * DeltaSeconds;
-
-		float GroundGravity = 0.0f;
-		if (const Chaos::FPBDRigidParticleHandle* Rigid = UPhysicsMovementUtils::GetRigidParticelHandleFromHitResult(FloorResult.HitResult))
-		{
-			if (Rigid->IsDynamic() && Rigid->GravityEnabled())
-			{
-				if (const UPrimitiveComponent* GroundComp = FloorResult.HitResult.GetComponent())
-				{
-					if (const APhysicsVolume* PhysVolume = GroundComp->GetPhysicsVolume())
-					{
-						GroundGravity = PhysVolume->GetGravityZ();
-					}
-				}
-			}
-		}
-		const FVector ProjectedGroundVelocity = StartGroundVelocity + GroundGravity * UpDir * DeltaSeconds;
+		const FVector ProjectedVelocity = TargetVelocity + CharacterGravity * FVector::UpVector * DeltaSeconds;
+		const FVector ProjectedGroundVelocity = UPhysicsMovementUtils::ComputeIntegratedGroundVelocityFromHitResult(StartingSyncState->GetLocation_WorldSpace(), FloorResult.HitResult, DeltaSeconds);
 
 		const float ProjectedRelativeVerticalVelocity = FloorResult.HitResult.ImpactNormal.Dot(ProjectedVelocity - ProjectedGroundVelocity);
 		const float VerticalVelocityLimit = 2.0f / DeltaSeconds;
diff --git a/Engine/Plugins/Experimental/Mover/Source/Mover/Private/PhysicsMover/PhysicsMovementUtils.cpp b/Engine/Plugins/Experimental/Mover/Source/Mover/Private/PhysicsMover/PhysicsMovementUtils.cpp
--- a/Engine/Plugins/Experimental/Mover/Source/Mover/Private/PhysicsMover/PhysicsMovementUtils.cpp
+++ b/Engine/Plugins/Experimental/Mover/Source/Mover/Private/PhysicsMover/PhysicsMovementUtils.cpp
@@ -5,6 +5,7 @@
 #include "Chaos/PhysicsObjectInternalInterface.h"
 #include "Components/PrimitiveComponent.h"
 #include "Engine/World.h"
+#include "GameFramework/PhysicsVolume.h"
 #include "MoveLibrary/FloorQueryUtils.h"
 #include "MoveLibrary/GroundMovementUtils.h"
 #include "MoveLibrary/WaterMovementUtils.h"
@@ -224,7 +225,7 @@ bool UPhysicsMovementUtils::IsHitSurfaceWalkableWithStepUpCheck(const FHitResult
 	return bWalkable;
 }
 
-const Chaos::FPBDRigidParticleHandle* UPhysicsMovementUtils::GetRigidParticelHandleFromHitResult(const FHitResult& HitResult)
+const Chaos::FPBDRigidParticleHandle* UPhysicsMovementUtils::GetRigidParticleHandleFromHitResult(const FHitResult& HitResult)
 {
 	if (IPhysicsComponent* PhysicsComp = Cast<IPhysicsComponent>(HitResult.Component))
 	{
@@ -244,7 +245,7 @@ const Chaos::FPBDRigidParticleHandle* UPhysicsMovementUtils::GetRigidParticelHan
 FVector UPhysicsMovementUtils::ComputeGroundVelocityFromHitResult(const FVector& CharacterPosition, const FHitResult& FloorHit, const float DeltaSeconds)
 {
 	FVector GroundVelocity = FVector::ZeroVector;
-	if (const Chaos::FPBDRigidParticleHandle* Rigid = GetRigidParticelHandleFromHitResult(FloorHit))
+	if (const Chaos::FPBDRigidParticleHandle* Rigid = GetRigidParticleHandleFromHitResult(FloorHit))
 	{
 		FVector Offset = CharacterPosition - Rigid->GetX();
 		Offset -= Offset.ProjectOnToNormal(FloorHit.ImpactNormal);
@@ -264,6 +265,40 @@ FVector UPhysicsMovementUtils::ComputeGroundVelocityFromHitResult(const FVector&
 	return GroundVelocity;
 }
 
+FVector UPhysicsMovementUtils::ComputeIntegratedGroundVelocityFromHitResult(const FVector& CharacterPosition, const FHitResult& FloorHit, const float DeltaSeconds)
+{
+	FVector GroundVelocity = FVector::ZeroVector;
+	if (const Chaos::FPBDRigidParticleHandle* Rigid = GetRigidParticleHandleFromHitResult(FloorHit))
+	{
+		FVector Offset = CharacterPosition - Rigid->GetX();
+		Offset -= Offset.ProjectOnToNormal(FloorHit.ImpactNormal);
+
+		if (Rigid->KinematicTarget().IsSet())
+		{
+			const FVector LinearDisplacement = Rigid->KinematicTarget().GetTargetPosition() - Rigid->GetX();
+			const FQuat RelativeQuat = Rigid->GetR().Inverse() * Rigid->KinematicTarget().GetTargetRotation();
+			const FVector AngularDisplacement = RelativeQuat.ToRotationVector();
+			GroundVelocity = (LinearDisplacement + AngularDisplacement.Cross(Offset)) / DeltaSeconds;
+		}
+		else
+		{
+			GroundVelocity = Rigid->GetV() + Rigid->GetW().Cross(Offset);
+		}
+
+		if (Rigid->IsDynamic() && Rigid->GravityEnabled())
+		{
+			if (const UPrimitiveComponent* GroundComp = FloorHit.GetComponent())
+			{
+				if (const APhysicsVolume* PhysVolume = GroundComp->GetPhysicsVolume())
+				{
+					GroundVelocity += PhysVolume->GetGravityZ() * FVector::UpVector * DeltaSeconds;
+				}
+			}
+		}
+	}
+	return GroundVelocity;
+}
+
 bool UPhysicsMovementUtils::GetWaterResultFromHitResults(const TArray<FHitResult>& Hits, const FVector& Location, FWaterCheckResult& OutWaterResult)
 {
 	// Find the closet hit that is a water body
diff --git a/Engine/Plugins/Experimental/Mover/Source/Mover/Public/PhysicsMover/PhysicsMovementUtils.h b/Engine/Plugins/Experimental/Mover/Source/Mover/Public/PhysicsMover/PhysicsMovementUtils.h
--- a/Engine/Plugins/Experimental/Mover/Source/Mover/Public/PhysicsMover/PhysicsMovementUtils.h
+++ b/Engine/Plugins/Experimental/Mover/Source/Mover/Public/PhysicsMover/PhysicsMovementUtils.h
@@ -18,7 +18,9 @@ class MOVER_API UPhysicsMovementUtils
 {
 public:
 	static void FindFloor(const FVector& Location, const FVector& DeltaPos, const UPrimitiveComponent* UpdatedPrimitive, const FVector& UpDir, float QueryRadius, float TargetHeight, float MaxStepHeight, float MaxWalkSlopeCosine, FFloorCheckResult& OutFloorResult, FWaterCheckResult& OutWaterResult);
-	static const Chaos::FPBDRigidParticleHandle* GetRigidParticelHandleFromHitResult(const FHitResult& HitResult);
+
+	// If the hit result hit something, return the particle handle
+	static const Chaos::FPBDRigidParticleHandle* GetRigidParticleHandleFromHitResult(const FHitResult& HitResult);
 
 	// Checks if the hit surface is walkable and, if stepping up, whether the surface can be stepped up on
 	static bool IsHitSurfaceWalkableWithStepUpCheck(const FHitResult& Hit, float StepHeight, float MaxStepHeight, float MinStepUpHeight, float MaxWalkSlopeCosine);
@@ -26,5 +28,9 @@ public:
 	// Checks if any hit is with water and, if so, fills in the OutWaterResult
 	static bool GetWaterResultFromHitResults(const TArray<FHitResult>& Hits, const FVector& Location, FWaterCheckResult& OutWaterResult);
 
+	// Returns the current ground velocity at the character position
 	static FVector ComputeGroundVelocityFromHitResult(const FVector& CharacterPosition, const FHitResult& FloorHit, const float DeltaSeconds);
+
+	// Returns the integrated with gravity velocity of the ground at the character position
+	static FVector ComputeIntegratedGroundVelocityFromHitResult(const FVector& CharacterPosition, const FHitResult& FloorHit, const float DeltaSeconds);
 };
\ No newline at end of file
