Mover 2.0 making Layered Move mixing more extensible.

diff --git a/Engine/Plugins/Experimental/Mover/Source/Mover/Private/LayeredMove.cpp b/Engine/Plugins/Experimental/Mover/Source/Mover/Private/LayeredMove.cpp
--- a/Engine/Plugins/Experimental/Mover/Source/Mover/Private/LayeredMove.cpp
+++ b/Engine/Plugins/Experimental/Mover/Source/Mover/Private/LayeredMove.cpp
@@ -115,110 +115,16 @@ void FLayeredMoveGroup::QueueLayeredMove(TSharedPtr<FLayeredMoveBase> Move)
 	}
 }
 
-bool FLayeredMoveGroup::DoGenerateMove(const FMoverTickStartData& StartState, const FMoverTimeStep& TimeStep, const UMoverComponent* MoverComp, UMoverBlackboard* SimBlackboard, FProposedMove& OutMove)
+TArray<TSharedPtr<FLayeredMoveBase>> FLayeredMoveGroup::GenerateActiveMoves(const FMoverTimeStep& TimeStep, const UMoverComponent* MoverComp, UMoverBlackboard* SimBlackboard)
 {
 	const float SimStartTimeMs		= TimeStep.BaseSimTimeMs;
 	const float SimTimeAfterTickMs	= SimStartTimeMs + TimeStep.StepMs;
 
 	FlushMoveArrays(MoverComp, SimBlackboard, SimStartTimeMs);
 
-	FProposedMove CumulativeMoveStep;
-	CumulativeMoveStep.MixMode = EMoveMixMode::AdditiveVelocity;
-
-	bool bDidAccumulateAnyMoves = false;
-	uint8 HighestPriority = 0;
-	float CurrentLayeredMoveStartTime = TNumericLimits<float>::Max();
-
-	// Tick and accumulate all active moves
-	// Gather all proposed moves and distill this into a final movement report. May include separate additive vs override moves.
-	// TODO: may want to sort by priority or other factors
-	for (TSharedPtr<FLayeredMoveBase>& ActiveMove : ActiveLayeredMoves)
-	{
-		FProposedMove MoveStep;
-		
-		if (ActiveMove->GenerateMove(StartState, TimeStep, MoverComp, SimBlackboard, MoveStep))
-		{
-			bDidAccumulateAnyMoves = true;
-
-			if (CumulativeMoveStep.PreferredMode != MoveStep.PreferredMode &&
-				!CumulativeMoveStep.PreferredMode.IsNone() && !MoveStep.PreferredMode.IsNone())
-			{
-				UE_LOG(LogMover, Log, TEXT("Multiple LayeredMoves are conflicting with preferred moves. %s will override %s"),
-					*MoveStep.PreferredMode.ToString(), *CumulativeMoveStep.PreferredMode.ToString());
-			}
-
-			if (MoveStep.bHasDirIntent && CumulativeMoveStep.MixMode != EMoveMixMode::OverrideAll && ActiveMove->Priority >= HighestPriority)
-			{
-				if (CumulativeMoveStep.bHasDirIntent)
-				{
-					UE_LOG(LogMover, Log, TEXT("Multiple LayeredMoves are setting direction intent and the layered move with highest priority will be used."));
-				}
-				
-				CumulativeMoveStep.bHasDirIntent = MoveStep.bHasDirIntent;
-				CumulativeMoveStep.DirectionIntent = MoveStep.DirectionIntent;
-			}
-
-			if (ActiveMove->MixMode == EMoveMixMode::OverrideVelocity)
-			{
-				if (CheckPriority(ActiveMove.Get(), HighestPriority, CurrentLayeredMoveStartTime))
-				{
-					if (CumulativeMoveStep.MixMode == EMoveMixMode::OverrideVelocity || CumulativeMoveStep.MixMode == EMoveMixMode::OverrideAll)
-					{
-						UE_LOG(LogMover, Log, TEXT("Multiple LayeredMoves with Override mix mode are active simultaneously. Layered move with the highest priority will take effect."));
-					}
-
-					if (!MoveStep.PreferredMode.IsNone())
-					{
-						CumulativeMoveStep.PreferredMode = MoveStep.PreferredMode;
-					}
-				
-					CumulativeMoveStep.MixMode = EMoveMixMode::OverrideVelocity;
-					CumulativeMoveStep.LinearVelocity  = MoveStep.LinearVelocity;
-					CumulativeMoveStep.AngularVelocity = MoveStep.AngularVelocity;
-				}
-			}
-			else if (ActiveMove->MixMode == EMoveMixMode::AdditiveVelocity)
-			{
-				if (CumulativeMoveStep.MixMode != EMoveMixMode::OverrideVelocity && CumulativeMoveStep.MixMode != EMoveMixMode::OverrideAll)
-				{
-					if (!MoveStep.PreferredMode.IsNone())
-					{
-						CumulativeMoveStep.PreferredMode = MoveStep.PreferredMode;
-					}
-
-					CumulativeMoveStep.PreferredMode = MoveStep.PreferredMode;
-					CumulativeMoveStep.LinearVelocity += MoveStep.LinearVelocity;
-					CumulativeMoveStep.AngularVelocity += MoveStep.AngularVelocity;
-				}
-			}
-			else if (ActiveMove->MixMode == EMoveMixMode::OverrideAll)
-			{
-				if (CheckPriority(ActiveMove.Get(), HighestPriority, CurrentLayeredMoveStartTime))
-				{
-					if (CumulativeMoveStep.MixMode == EMoveMixMode::OverrideVelocity || CumulativeMoveStep.MixMode == EMoveMixMode::OverrideAll)
-					{
-						UE_LOG(LogMover, Log, TEXT("Multiple LayeredMoves with Override mix mode are active simultaneously. Layered move with the highest priority will take effect."));
-					}
-				
-					CumulativeMoveStep = MoveStep;
-				}
-			}
-			else
-			{
-				check(0);	// unhandled case
-			}
-		}
-	}
-
-	if (bDidAccumulateAnyMoves)
-	{
-		OutMove = CumulativeMoveStep;
-	}
-
-	return bDidAccumulateAnyMoves;
+	return ActiveLayeredMoves;
 }
 
-
 void FLayeredMoveGroup::NetSerialize(FArchive& Ar, uint8 MaxNumMovesToSerialize/* = MAX_uint8*/)
 {
 	// TODO: Warn if some sources will be dropped
@@ -469,23 +375,6 @@ void FLayeredMoveGroup::GatherResidualVelocitySettings(const TSharedPtr<FLayered
 	}
 }
 
-bool FLayeredMoveGroup::CheckPriority(const FLayeredMoveBase* LayeredMove, uint8& InOutHighestPriority, float& InOutCurrentLayeredMoveStartTimeMs)
-{
-	if (LayeredMove->Priority > InOutHighestPriority)
-	{
-		InOutHighestPriority = LayeredMove->Priority;
-		InOutCurrentLayeredMoveStartTimeMs = LayeredMove->StartSimTimeMs;
-		return true;
-	}
-	if (LayeredMove->Priority == InOutHighestPriority && LayeredMove->StartSimTimeMs < InOutCurrentLayeredMoveStartTimeMs)
-	{
-		InOutCurrentLayeredMoveStartTimeMs = LayeredMove->StartSimTimeMs;
-		return true;
-	}
-
-	return false;
-}
-
 struct FLayeredMoveDeleter
 {
 	FORCEINLINE void operator()(FLayeredMoveBase* Object) const
diff --git a/Engine/Plugins/Experimental/Mover/Source/Mover/Private/MoveLibrary/MovementMixer.cpp b/Engine/Plugins/Experimental/Mover/Source/Mover/Private/MoveLibrary/MovementMixer.cpp
new file mode 100644
--- /dev/null
+++ b/Engine/Plugins/Experimental/Mover/Source/Mover/Private/MoveLibrary/MovementMixer.cpp
@@ -0,0 +1,133 @@
+﻿// Copyright Epic Games, Inc. All Rights Reserved.
+
+#include "MoveLibrary/MovementMixer.h"
+#include "LayeredMove.h"
+#include "MoverLog.h"
+
+UMovementMixer::UMovementMixer()
+	: CurrentHighestPriority(0)
+	, CurrentLayeredMoveStartTime(TNumericLimits<float>::Max())
+{
+}
+
+void UMovementMixer::MixLayeredMove(const FLayeredMoveBase& ActiveMove, const FProposedMove& MoveStep, FProposedMove& OutCumulativeMove)
+{
+	if (OutCumulativeMove.PreferredMode != MoveStep.PreferredMode && !OutCumulativeMove.PreferredMode.IsNone() && !MoveStep.PreferredMode.IsNone())
+	{
+		UE_LOG(LogMover, Log, TEXT("Multiple LayeredMoves are conflicting with preferred moves. %s will override %s"),
+			*MoveStep.PreferredMode.ToString(), *OutCumulativeMove.PreferredMode.ToString());
+	}
+
+	if (MoveStep.bHasDirIntent && OutCumulativeMove.MixMode != EMoveMixMode::OverrideAll && ActiveMove.Priority >= CurrentHighestPriority)
+	{
+		if (OutCumulativeMove.bHasDirIntent)
+		{
+			UE_LOG(LogMover, Log, TEXT("Multiple LayeredMoves are setting direction intent and the layered move with highest priority will be used."));
+		}
+				
+		OutCumulativeMove.bHasDirIntent = MoveStep.bHasDirIntent;
+		OutCumulativeMove.DirectionIntent = MoveStep.DirectionIntent;
+	}
+
+	if (ActiveMove.MixMode == EMoveMixMode::OverrideVelocity)
+	{
+		if (CheckPriority(&ActiveMove, CurrentHighestPriority, CurrentLayeredMoveStartTime))
+		{
+			if (OutCumulativeMove.MixMode == EMoveMixMode::OverrideVelocity || OutCumulativeMove.MixMode == EMoveMixMode::OverrideAll)
+			{
+				UE_LOG(LogMover, Log, TEXT("Multiple LayeredMoves with Override mix mode are active simultaneously. Layered move with the highest priority will take effect."));
+			}
+
+			if (!MoveStep.PreferredMode.IsNone())
+			{
+				OutCumulativeMove.PreferredMode = MoveStep.PreferredMode;
+			}
+				
+			OutCumulativeMove.MixMode = EMoveMixMode::OverrideVelocity;
+			OutCumulativeMove.LinearVelocity  = MoveStep.LinearVelocity;
+			OutCumulativeMove.AngularVelocity = MoveStep.AngularVelocity;
+		}
+	}
+	else if (ActiveMove.MixMode == EMoveMixMode::AdditiveVelocity)
+	{
+		if (OutCumulativeMove.MixMode != EMoveMixMode::OverrideVelocity && OutCumulativeMove.MixMode != EMoveMixMode::OverrideAll)
+		{
+			if (!MoveStep.PreferredMode.IsNone())
+			{
+				OutCumulativeMove.PreferredMode = MoveStep.PreferredMode;
+			}
+
+			OutCumulativeMove.PreferredMode = MoveStep.PreferredMode;
+			OutCumulativeMove.LinearVelocity += MoveStep.LinearVelocity;
+			OutCumulativeMove.AngularVelocity += MoveStep.AngularVelocity;
+		}
+	}
+	else if (ActiveMove.MixMode == EMoveMixMode::OverrideAll)
+	{
+		if (CheckPriority(&ActiveMove, CurrentHighestPriority, CurrentLayeredMoveStartTime))
+		{
+			if (OutCumulativeMove.MixMode == EMoveMixMode::OverrideVelocity || OutCumulativeMove.MixMode == EMoveMixMode::OverrideAll)
+			{
+				UE_LOG(LogMover, Log, TEXT("Multiple LayeredMoves with Override mix mode are active simultaneously. Layered move with the highest priority will take effect."));
+			}
+				
+			OutCumulativeMove = MoveStep;
+		}
+	}
+	else
+	{
+		ensureMsgf(false, TEXT("Unhandled move mix mode was found."));
+	}
+}
+
+void UMovementMixer::MixProposedMoves(const FProposedMove& MoveToMix, FProposedMove& OutCumulativeMove)
+{
+	if (MoveToMix.bHasDirIntent && OutCumulativeMove.MixMode != EMoveMixMode::OverrideAll)
+	{
+		OutCumulativeMove.bHasDirIntent = MoveToMix.bHasDirIntent;
+		OutCumulativeMove.DirectionIntent = MoveToMix.DirectionIntent;
+	}
+
+	// Combine movement parameters from layered moves into what the mode wants to do
+	if (MoveToMix.MixMode == EMoveMixMode::OverrideAll)
+	{
+		OutCumulativeMove = MoveToMix;
+	}
+	else if (MoveToMix.MixMode == EMoveMixMode::AdditiveVelocity)
+	{
+		OutCumulativeMove.LinearVelocity += MoveToMix.LinearVelocity;
+		OutCumulativeMove.AngularVelocity += MoveToMix.AngularVelocity;
+	}
+	else if (MoveToMix.MixMode == EMoveMixMode::OverrideVelocity)
+	{
+		OutCumulativeMove.LinearVelocity = MoveToMix.LinearVelocity;
+		OutCumulativeMove.AngularVelocity = MoveToMix.AngularVelocity;
+	}
+	else
+	{
+		ensureMsgf(false, TEXT("Unhandled move mix mode was found."));
+	}
+}
+
+void UMovementMixer::ResetMixerState()
+{
+	CurrentHighestPriority = 0;
+	CurrentLayeredMoveStartTime = TNumericLimits<float>::Max();
+}
+
+bool UMovementMixer::CheckPriority(const FLayeredMoveBase* LayeredMove, uint8& InOutHighestPriority, float& InOutCurrentLayeredMoveStartTimeMs)
+{
+	if (LayeredMove->Priority > InOutHighestPriority)
+	{
+		InOutHighestPriority = LayeredMove->Priority;
+		InOutCurrentLayeredMoveStartTimeMs = LayeredMove->StartSimTimeMs;
+		return true;
+	}
+	if (LayeredMove->Priority == InOutHighestPriority && LayeredMove->StartSimTimeMs < InOutCurrentLayeredMoveStartTimeMs)
+	{
+		InOutCurrentLayeredMoveStartTimeMs = LayeredMove->StartSimTimeMs;
+		return true;
+	}
+
+	return false;
+}
diff --git a/Engine/Plugins/Experimental/Mover/Source/Mover/Private/MovementModeStateMachine.cpp b/Engine/Plugins/Experimental/Mover/Source/Mover/Private/MovementModeStateMachine.cpp
--- a/Engine/Plugins/Experimental/Mover/Source/Mover/Private/MovementModeStateMachine.cpp
+++ b/Engine/Plugins/Experimental/Mover/Source/Mover/Private/MovementModeStateMachine.cpp
@@ -6,6 +6,7 @@
 #include "MoverLog.h"
 #include "MoveLibrary/MovementUtils.h"
 #include "MoverComponent.h"
+#include "MoveLibrary/MovementMixer.h"
 #include "Templates/SubclassOf.h"
 
 
@@ -136,6 +137,7 @@ void UMovementModeStateMachine::OnSimulationTick(USceneComponent* UpdatedCompone
 	FMoverDefaultSyncState* SubstepStartSyncState = SubstepStartData.SyncState.SyncStateCollection.FindMutableDataByType<FMoverDefaultSyncState>();
 
 	UMoverComponent* MoverComp = CastChecked<UMoverComponent>(GetOuter());
+	check(MoverComp->MovementMixer);
 
 	if (!QueuedModeTransition->IsSet())
 	{
@@ -163,7 +165,24 @@ void UMovementModeStateMachine::OnSimulationTick(USceneComponent* UpdatedCompone
 
 		// Gather any layered move contributions
 		FProposedMove CombinedLayeredMove;
-		const bool bHasLayeredMoveContributions = CurrentLayeredMoves.DoGenerateMove(SubstepStartData, SubTimeStep, MoverComp, SimBlackboard, OUT CombinedLayeredMove);
+		CombinedLayeredMove.MixMode = EMoveMixMode::AdditiveVelocity;
+		bool bHasLayeredMoveContributions = false;
+		MoverComp->MovementMixer->ResetMixerState();
+		
+		TArray<TSharedPtr<FLayeredMoveBase>> ActiveMoves = CurrentLayeredMoves.GenerateActiveMoves(SubTimeStep, MoverComp, SimBlackboard);
+
+		// Tick and accumulate all active moves
+		// Gather all proposed moves and distill this into a cumulative movement report. May include separate additive vs override moves.
+		// TODO: may want to sort by priority or other factors
+		for (TSharedPtr<FLayeredMoveBase>& ActiveMove : ActiveMoves)
+		{
+			FProposedMove MoveStep;
+			if (ActiveMove->GenerateMove(StartState, TimeStep, MoverComp, SimBlackboard, MoveStep))
+			{
+				bHasLayeredMoveContributions = true;
+				MoverComp->MovementMixer->MixLayeredMove(*ActiveMove, MoveStep, CombinedLayeredMove);
+			}
+		}
 
 		if (bHasLayeredMoveContributions && !CombinedLayeredMove.PreferredMode.IsNone())
 		{
@@ -180,27 +199,7 @@ void UMovementModeStateMachine::OnSimulationTick(USceneComponent* UpdatedCompone
 
 			if (bHasLayeredMoveContributions)
 			{
-				if (CombinedLayeredMove.bHasDirIntent && CombinedMove.MixMode != EMoveMixMode::OverrideAll)
-				{
-					CombinedMove.bHasDirIntent = CombinedLayeredMove.bHasDirIntent;
-					CombinedMove.DirectionIntent = CombinedLayeredMove.DirectionIntent;
-				}
-				
-				// Combine movement parameters from layered moves into what the mode wants to do
-				if (CombinedLayeredMove.MixMode == EMoveMixMode::OverrideAll)
-				{
-					CombinedMove = CombinedLayeredMove;
-				}
-				else if (CombinedLayeredMove.MixMode == EMoveMixMode::AdditiveVelocity)
-				{
-					CombinedMove.LinearVelocity += CombinedLayeredMove.LinearVelocity;
-					CombinedMove.AngularVelocity += CombinedLayeredMove.AngularVelocity;
-				}
-				else if (CombinedLayeredMove.MixMode == EMoveMixMode::OverrideVelocity)
-				{
-					CombinedMove.LinearVelocity = CombinedLayeredMove.LinearVelocity;
-					CombinedMove.AngularVelocity = CombinedLayeredMove.AngularVelocity;
-				}
+				MoverComp->MovementMixer->MixProposedMoves(CombinedLayeredMove, CombinedMove);
 			}
 
 			// Apply any layered move finish velocity settings
diff --git a/Engine/Plugins/Experimental/Mover/Source/Mover/Private/MoverComponent.cpp b/Engine/Plugins/Experimental/Mover/Source/Mover/Private/MoverComponent.cpp
--- a/Engine/Plugins/Experimental/Mover/Source/Mover/Private/MoverComponent.cpp
+++ b/Engine/Plugins/Experimental/Mover/Source/Mover/Private/MoverComponent.cpp
@@ -7,6 +7,7 @@
 #include "Kinematic/Modes/WalkingMode.h"
 #include "Kinematic/Modes/FallingMode.h"
 #include "Kinematic/Modes/FlyingMode.h"
+#include "MoveLibrary/MovementMixer.h"
 #include "MoveLibrary/MovementUtils.h"
 #include "MoveLibrary/FloorQueryUtils.h"
 #include "MoverLog.h"
@@ -189,6 +190,11 @@ void UMoverComponent::BeginPlay()
 			}
 		}
 	}
+	
+	if (!MovementMixer)
+	{
+		MovementMixer = NewObject<UMovementMixer>(this, TEXT("Default Movement Mixer"));
+	}
 }
 
 void UMoverComponent::TickComponent(float DeltaTime, enum ELevelTick TickType, FActorComponentTickFunction* ThisTickFunction)
diff --git a/Engine/Plugins/Experimental/Mover/Source/Mover/Public/LayeredMove.h b/Engine/Plugins/Experimental/Mover/Source/Mover/Public/LayeredMove.h
--- a/Engine/Plugins/Experimental/Mover/Source/Mover/Public/LayeredMove.h
+++ b/Engine/Plugins/Experimental/Mover/Source/Mover/Public/LayeredMove.h
@@ -5,6 +5,7 @@
 #include "MoveLibrary/MovementUtilsTypes.h"
 #include "LayeredMove.generated.h"
 
+class UMovementMixer;
 struct FMoverTickStartData;
 struct FMoverTimeStep;
 class UMoverBlackboard;
@@ -167,7 +168,8 @@ struct MOVER_API FLayeredMoveGroup
 
 	bool HasAnyMoves() const { return (!ActiveLayeredMoves.IsEmpty() || !QueuedLayeredMoves.IsEmpty()); }
 
-	bool DoGenerateMove(const FMoverTickStartData& StartState, const FMoverTimeStep& TimeStep, const UMoverComponent* MoverComp, UMoverBlackboard* SimBlackboard, FProposedMove& OutMove);
+	// Generates active layered move list (by calling FlushMoveArrays) and returns the an array of all currently active layered moves
+	TArray<TSharedPtr<FLayeredMoveBase>> GenerateActiveMoves(const FMoverTimeStep& TimeStep, const UMoverComponent* MoverComp, UMoverBlackboard* SimBlackboard);
 
 	/** Serialize all moves and their states for this group */
 	void NetSerialize(FArchive& Ar, uint8 MaxNumMovesToSerialize = MAX_uint8);
@@ -197,16 +199,11 @@ struct MOVER_API FLayeredMoveGroup
 	void ResetResidualVelocity();
 
 protected:
+	// Clears out any finished or invalid active moves and adds any queued moves to the active moves
 	void FlushMoveArrays(const UMoverComponent* MoverComp, UMoverBlackboard* SimBlackboard, float CurrentSimTimeMs);
 
 	// Helper function for gathering any residual velocity settings from layered moves that just ended
 	void GatherResidualVelocitySettings(const TSharedPtr<FLayeredMoveBase>& Move, bool& bResidualVelocityOverriden, bool& bClampVelocityOverriden);
-
-	/**
-	 * Helper function for layered move mixing to check priority and start time if priority is the same.
-	 * Returns true if this layered move should take priority given current HighestPriority and CurrentLayeredMoveStartTimeMs
-	 */
-	static bool CheckPriority(const FLayeredMoveBase* LayeredMove, uint8& InOutHighestPriority, float& InOutCurrentLayeredMoveStartTimeMs);
 	
 	/** Helper function for serializing array of root motion sources */
 	static void NetSerializeLayeredMovesArray(FArchive& Ar, TArray< TSharedPtr<FLayeredMoveBase> >& LayeredMovesArray, uint8 MaxNumLayeredMovesToSerialize = MAX_uint8);
diff --git a/Engine/Plugins/Experimental/Mover/Source/Mover/Public/MoveLibrary/MovementMixer.h b/Engine/Plugins/Experimental/Mover/Source/Mover/Public/MoveLibrary/MovementMixer.h
new file mode 100644
--- /dev/null
+++ b/Engine/Plugins/Experimental/Mover/Source/Mover/Public/MoveLibrary/MovementMixer.h
@@ -0,0 +1,43 @@
+﻿// Copyright Epic Games, Inc. All Rights Reserved.
+
+#pragma once
+
+#include "CoreMinimal.h"
+#include "MovementMixer.generated.h"
+
+struct FLayeredMoveBase;
+struct FProposedMove;
+
+/**
+ * Class in charge of Mixing various moves when evaluating/combining moves. The mixer used can be set on the MoverComponent itself.
+ */
+UCLASS(BlueprintType)
+class MOVER_API UMovementMixer : public UObject
+{
+	GENERATED_BODY()
+
+public:
+	UMovementMixer();
+	
+	/** In charge of mixing Layered Move proposed moves into a cumulative proposed move based on mix mode and priority.*/
+	virtual void MixLayeredMove(const FLayeredMoveBase& ActiveMove, const FProposedMove& MoveStep, FProposedMove& OutCumulativeMove);
+
+	/** In charge of mixing proposed moves together. Is similar to MixLayeredMove but is only responsible for mixing proposed moves instead of layered moves. */
+	virtual void MixProposedMoves(const FProposedMove& MoveToMix, FProposedMove& OutCumulativeMove);
+
+	/** Resets all state used for mixing. Should be called before or after finished mixing moves. */
+	virtual void ResetMixerState();
+	
+protected:
+	// Stores the current highest priority we've hit during this round of mixing. Will get reset. Note: Currently only used for mixing layered moves
+	uint8 CurrentHighestPriority;
+
+	// Earliest start time of the layered move with highest priority. Used to help break ties of moves with same priority. Note: Currently only used for mixing layered moves
+	float CurrentLayeredMoveStartTime;
+	
+	/**
+	 * Helper function for layered move mixing to check priority and start time if priority is the same.
+	 * Returns true if this layered move should take priority given current HighestPriority and CurrentLayeredMoveStartTimeMs
+	 */
+	static bool CheckPriority(const FLayeredMoveBase* LayeredMove, uint8& InOutHighestPriority, float& InOutCurrentLayeredMoveStartTimeMs);
+};
diff --git a/Engine/Plugins/Experimental/Mover/Source/Mover/Public/MoverComponent.h b/Engine/Plugins/Experimental/Mover/Source/Mover/Public/MoverComponent.h
--- a/Engine/Plugins/Experimental/Mover/Source/Mover/Public/MoverComponent.h
+++ b/Engine/Plugins/Experimental/Mover/Source/Mover/Public/MoverComponent.h
@@ -18,6 +18,7 @@
 
 struct FMoverTimeStep;
 class UMovementModeStateMachine;
+class UMovementMixer;
 
 namespace MoverComponentConstants
 {
@@ -128,6 +129,10 @@ public:
 	UPROPERTY(BlueprintReadWrite, EditAnywhere, Category = Mover, meta = (MustImplement = "/Script/Mover.MoverInputProducerInterface"))
 	TObjectPtr<UObject> InputProducer;
 
+	/** Optional object for mixing proposed moves.Typically set at BeginPlay time. If not specified, UDefaultMovementMixer will be used. */
+	UPROPERTY(BlueprintReadWrite, EditAnywhere, Category = Mover)
+	TObjectPtr<UMovementMixer> MovementMixer;
+	
 	/**
 	 * Queue a layered move to start during the next simulation frame. This will clone whatever move you pass in, so you'll need to fully set it up before queuing.
 	 * @param LayeredMove			The move to queue, which must be a LayeredMoveBase sub-type. 
