Mover - Add networked input struct for max speed and acceleration

diff --git a/Engine/Plugins/Experimental/Mover/Source/Mover/Private/Backends/MoverNetworkPhysicsLiaison.cpp b/Engine/Plugins/Experimental/Mover/Source/Mover/Private/Backends/MoverNetworkPhysicsLiaison.cpp
--- a/Engine/Plugins/Experimental/Mover/Source/Mover/Private/Backends/MoverNetworkPhysicsLiaison.cpp
+++ b/Engine/Plugins/Experimental/Mover/Source/Mover/Private/Backends/MoverNetworkPhysicsLiaison.cpp
@@ -113,6 +113,17 @@ void FNetworkPhysicsMoverInputs::MergeData(const FNetworkPhysicsData& FromData)
 	}
 }
 
+void FNetworkPhysicsMoverInputs::ValidateData(const UActorComponent* NetworkComponent)
+{
+	if (NetworkComponent)
+	{
+		if (const UMoverNetworkPhysicsLiaisonComponent* LiaisonComp = Cast<UMoverNetworkPhysicsLiaisonComponent>(NetworkComponent))
+		{
+			LiaisonComp->ValidateInputData(InputCmdContext);
+		}
+	}
+}
+
 //////////////////////////////////////////////////////////////////////////
 // FNetworkPhysicsMoverState
 
@@ -166,6 +177,7 @@ void FNetworkPhysicsMoverState::InterpolateData(const FNetworkPhysicsData& MinDa
 }
 
 //////////////////////////////////////////////////////////////////////////
+// UMoverNetworkPhysicsLiaisonComponent
 
 void UMoverNetworkPhysicsLiaisonComponent::GetCurrentInputData(OUT FMoverInputCmdContext& InputCmd) const
 {
@@ -187,8 +199,29 @@ void UMoverNetworkPhysicsLiaisonComponent::SetCurrentStateData(const FMoverSyncS
 	NetSyncState = SyncState;
 }
 
-//////////////////////////////////////////////////////////////////////////
-// UMoverNetworkPhysicsLiaisonComponent
+bool UMoverNetworkPhysicsLiaisonComponent::ValidateInputData(FMoverInputCmdContext& InputCmd) const
+{
+	bool bValidInputs = true;
+	if (FMovementSettingsInputs* MovementSettings = InputCmd.InputCollection.FindMutableDataByType<FMovementSettingsInputs>())
+	{
+		// TODO - proper data validation
+		const float MaxMaxSpeed = 1000.0f;
+		if (MovementSettings->MaxSpeed > MaxMaxSpeed)
+		{
+			MovementSettings->MaxSpeed = MaxMaxSpeed;
+			bValidInputs = false;
+		}
+
+		const float MaxAcceleration = 10000.0f;
+		if (MovementSettings->Acceleration > MaxAcceleration)
+		{
+			MovementSettings->Acceleration = MaxAcceleration;
+			bValidInputs = false;
+		}
+	}
+
+	return bValidInputs;
+}
 
 UMoverNetworkPhysicsLiaisonComponent::UMoverNetworkPhysicsLiaisonComponent()
 {
@@ -258,7 +291,7 @@ void UMoverNetworkPhysicsLiaisonComponent::OnRegister()
 	if ((MoverComp = GetOwner()->FindComponentByClass<UMoverComponent>()) != nullptr)
 	{
 	
-		CommonMovementSettings = MoverComp->FindSharedSettings<UCommonLegacyMovementSettings>();
+		CommonMovementSettings = MoverComp->FindSharedSettings_Mutable<UCommonLegacyMovementSettings>();
 		check(CommonMovementSettings);
 
 		if (MoverComp->UpdatedCompAsPrimitive)
@@ -746,7 +779,13 @@ void UMoverNetworkPhysicsLiaisonComponent::ProcessInputs_Internal(int32 PhysicsS
 				}
 			}
 		}
-		
+	}
+
+	// Override common settings data with data from FMovementSettingsInputs if present in the input cmd
+	if (const FMovementSettingsInputs* MovementSettings = Input.InputCmd.InputCollection.FindDataByType<FMovementSettingsInputs>())
+	{
+		CommonMovementSettings->MaxSpeed = MovementSettings->MaxSpeed;
+		CommonMovementSettings->Acceleration = MovementSettings->Acceleration;
 	}
 }
 
diff --git a/Engine/Plugins/Experimental/Mover/Source/Mover/Private/PhysicsMover/PhysicsMoverSimulationTypes.cpp b/Engine/Plugins/Experimental/Mover/Source/Mover/Private/PhysicsMover/PhysicsMoverSimulationTypes.cpp
--- a/Engine/Plugins/Experimental/Mover/Source/Mover/Private/PhysicsMover/PhysicsMoverSimulationTypes.cpp
+++ b/Engine/Plugins/Experimental/Mover/Source/Mover/Private/PhysicsMover/PhysicsMoverSimulationTypes.cpp
@@ -18,4 +18,33 @@ FAutoConsoleVariableRef CVarPhysicsDrivenMotionTeleportThreshold(TEXT("p.mover.p
 	GPhysicsDrivenMotionDebugParams.TeleportThreshold, TEXT("Single frame movement threshold in cm that will trigger a teleport."));
 
 FAutoConsoleVariableRef CVarPhysicsDrivenMotionMinStepUpDistance(TEXT("p.mover.physics.MinStepUpDistance"),
-	GPhysicsDrivenMotionDebugParams.MinStepUpDistance, TEXT("Minimum distance that will be considered a step up."));
\ No newline at end of file
+	GPhysicsDrivenMotionDebugParams.MinStepUpDistance, TEXT("Minimum distance that will be considered a step up."));
+
+//////////////////////////////////////////////////////////////////////////
+// FMovementSettingsInputs
+
+FMoverDataStructBase* FMovementSettingsInputs::Clone() const
+{
+	FMovementSettingsInputs* CopyPtr = new FMovementSettingsInputs(*this);
+	return CopyPtr;
+}
+
+bool FMovementSettingsInputs::NetSerialize(FArchive& Ar, UPackageMap* Map, bool& bOutSuccess)
+{
+	Super::NetSerialize(Ar, Map, bOutSuccess);
+
+	Ar << MaxSpeed;
+	Ar << Acceleration;
+
+	bOutSuccess = true;
+	return true;
+}
+
+
+void FMovementSettingsInputs::ToString(FAnsiStringBuilderBase& Out) const
+{
+	Super::ToString(Out);
+
+	Out.Appendf("MaxSpeed=%.2f\n", MaxSpeed);
+	Out.Appendf("Acceleration=%.2f\n", Acceleration);
+}
\ No newline at end of file
diff --git a/Engine/Plugins/Experimental/Mover/Source/Mover/Public/Backends/MoverNetworkPhysicsLiaison.h b/Engine/Plugins/Experimental/Mover/Source/Mover/Public/Backends/MoverNetworkPhysicsLiaison.h
--- a/Engine/Plugins/Experimental/Mover/Source/Mover/Public/Backends/MoverNetworkPhysicsLiaison.h
+++ b/Engine/Plugins/Experimental/Mover/Source/Mover/Public/Backends/MoverNetworkPhysicsLiaison.h
@@ -38,6 +38,9 @@ struct MOVER_API FNetworkPhysicsMoverInputs : public FNetworkPhysicsData
 	
 	/** Merge data into this input */
 	virtual void MergeData(const FNetworkPhysicsData& FromData) override;
+
+	/** Check input data is valid - Input is send from client to server, no need to make sure it's reasonable */
+	virtual void ValidateData(const UActorComponent* NetworkComponent) override;
 };
 
 template<>
@@ -120,6 +123,7 @@ public:
 	void SetCurrentInputData(const FMoverInputCmdContext& InputCmd);
 	void GetCurrentStateData(OUT FMoverSyncState& SyncState) const;
 	void SetCurrentStateData(const FMoverSyncState& SyncState);
+	bool ValidateInputData(FMoverInputCmdContext& InputCmd) const;
 
 	// Used by the manager to uniquely identify the component
 	Chaos::FUniqueIdx GetUniqueIdx() const;
@@ -158,7 +162,7 @@ protected:
 
 	TUniquePtr<Chaos::FCharacterGroundConstraint> Constraint;
 	TObjectPtr<UMoverComponent> MoverComp;
-	TObjectPtr<const UCommonLegacyMovementSettings> CommonMovementSettings;
+	TObjectPtr<UCommonLegacyMovementSettings> CommonMovementSettings;
 
 	TObjectPtr<UNetworkPhysicsComponent> NetworkPhysicsComponent;
 
diff --git a/Engine/Plugins/Experimental/Mover/Source/Mover/Public/PhysicsMover/PhysicsMoverSimulationTypes.h b/Engine/Plugins/Experimental/Mover/Source/Mover/Public/PhysicsMover/PhysicsMoverSimulationTypes.h
--- a/Engine/Plugins/Experimental/Mover/Source/Mover/Public/PhysicsMover/PhysicsMoverSimulationTypes.h
+++ b/Engine/Plugins/Experimental/Mover/Source/Mover/Public/PhysicsMover/PhysicsMoverSimulationTypes.h
@@ -79,4 +79,54 @@ class IPhysicsCharacterMovementModeInterface
 
 public:
 	virtual void UpdateConstraintSettings(Chaos::FCharacterGroundConstraint& Constraint) const = 0;
+};
+
+//////////////////////////////////////////////////////////////////////////
+// FMovementSettingsInput
+
+// Data block containing movement settings inputs that are networked from client to server.
+// This is useful if settings changes need to be predicted on the client and synced on the server.
+// Also supports rewind/resimulation of settings changes.
+USTRUCT(BlueprintType)
+struct MOVER_API FMovementSettingsInputs : public FMoverDataStructBase
+{
+	GENERATED_USTRUCT_BODY()
+
+public:
+	// Maximum speed in cm/s
+	UPROPERTY(BlueprintReadWrite, Category = Mover)
+		float MaxSpeed;
+
+	// Maximum acceleration in cm/s^2
+	UPROPERTY(BlueprintReadWrite, Category = Mover)
+		float Acceleration;
+
+	FMovementSettingsInputs()
+		: MaxSpeed(800.0f)
+		, Acceleration(4000.0f)
+	{
+	}
+
+	virtual ~FMovementSettingsInputs() {}
+
+	// @return newly allocated copy of this FKinematicDefaultInputs. Must be overridden by child classes
+	virtual FMoverDataStructBase* Clone() const override;
+
+	virtual bool NetSerialize(FArchive& Ar, UPackageMap* Map, bool& bOutSuccess) override;
+
+	virtual UScriptStruct* GetScriptStruct() const override { return StaticStruct(); }
+
+	virtual void ToString(FAnsiStringBuilderBase& Out) const override;
+
+	virtual void AddReferencedObjects(FReferenceCollector& Collector) override { Super::AddReferencedObjects(Collector); }
+};
+
+template<>
+struct TStructOpsTypeTraits< FMovementSettingsInputs > : public TStructOpsTypeTraitsBase2< FMovementSettingsInputs >
+{
+	enum
+	{
+		WithNetSerializer = true,
+		WithCopy = true
+	};
 };
\ No newline at end of file
