Mover 2.0 giving teleport layered move a higher default priority so it isn't overridden by other moves by default

diff --git a/Engine/Plugins/Experimental/Mover/Source/Mover/Private/Kinematic/LayeredMoves/BasicLayeredMoves.cpp b/Engine/Plugins/Experimental/Mover/Source/Mover/Private/Kinematic/LayeredMoves/BasicLayeredMoves.cpp
--- a/Engine/Plugins/Experimental/Mover/Source/Mover/Private/Kinematic/LayeredMoves/BasicLayeredMoves.cpp
+++ b/Engine/Plugins/Experimental/Mover/Source/Mover/Private/Kinematic/LayeredMoves/BasicLayeredMoves.cpp
@@ -315,6 +315,7 @@ FLayeredMove_Teleport::FLayeredMove_Teleport()
 {
 	DurationMs = 0.f;
 	MixMode = EMoveMixMode::OverrideAll;
+	Priority = 10;
 }
 
 bool FLayeredMove_Teleport::GenerateMove(const FMoverTickStartData& SimState, const FMoverTimeStep& TimeStep, const UMoverComponent* MoverComp, UMoverBlackboard* SimBlackboard, FProposedMove& OutProposedMove)
