Mover: adding more thorough info about the selected Mover actor to the Gameplay Debugger Tool readout

diff --git a/Engine/Plugins/Experimental/Mover/Source/Mover/Private/Debug/GameplayDebuggerCategory_Mover.cpp b/Engine/Plugins/Experimental/Mover/Source/Mover/Private/Debug/GameplayDebuggerCategory_Mover.cpp
--- a/Engine/Plugins/Experimental/Mover/Source/Mover/Private/Debug/GameplayDebuggerCategory_Mover.cpp
+++ b/Engine/Plugins/Experimental/Mover/Source/Mover/Private/Debug/GameplayDebuggerCategory_Mover.cpp
@@ -91,9 +91,15 @@ void FGameplayDebuggerCategory_Mover::CollectData(APlayerController* OwnerPC, AA
 	DataPack.PawnName = MyPawn ? MyPawn->GetHumanReadableName() : FString(TEXT("{red}No selected pawn."));
 	DataPack.LocalRole = MyPawn ? UEnum::GetValueAsString(TEXT("Engine.ENetRole"), MyPawn->GetLocalRole()) : FString();
 
+	// Set defaults for info that may not be available
+	DataPack.MovementModeName = FString("invalid");
+	DataPack.MovementBaseInfo = FString("invalid");
+	DataPack.MoveIntent = FVector::ZeroVector;
+	DataPack.ActiveLayeredMoves.Empty();
+	DataPack.ModeMap.Empty();
 	DataPack.ActiveLayeredMoves.Empty();
 
-	if (MyMoverComponent && MyMoverComponent->HasValidCachedState())
+	if (MyMoverComponent)
 	{
 		UPrimitiveComponent* MovementBaseComp = MyMoverComponent->GetMovementBase();
 
@@ -101,22 +107,38 @@ void FGameplayDebuggerCategory_Mover::CollectData(APlayerController* OwnerPC, AA
 		DataPack.MovementBaseInfo = MovementBaseComp ? FString::Printf(TEXT("%s.%s"), *GetNameSafe(MovementBaseComp->GetOwner()), *MovementBaseComp->GetName()) : FString();
 		DataPack.MoveIntent = MyMoverComponent->GetMovementIntent();
 
-		const FMoverSyncState& SyncState = MyMoverComponent->GetSyncState();
+		for (auto ModeIter = MyMoverComponent->MovementModes.begin(); ModeIter; ++ModeIter)
+		{
+			UBaseMovementMode* MappedMode = ModeIter->Value;
+			DataPack.ModeMap.Add(FString::Printf(TEXT("%s => %s"), *ModeIter->Key.ToString(), (MappedMode ? *MappedMode->GetClass()->GetName() : TEXT("null"))));
+
+			if (ModeIter->Key == MyMoverComponent->GetMovementModeName())
+			{
+				for (UBaseMovementModeTransition* Transition : MyMoverComponent->MovementModes[ModeIter->Key]->Transitions)
+				{
+					DataPack.ActiveTransitions.Add(FString::Printf(TEXT("%s (%s)"), (Transition ? *Transition->GetClass()->GetName() : TEXT("null")), *ModeIter->Key.ToString()));
+				}
+			}
+		}
 
-		if (const FMoverDefaultSyncState* MoverState = SyncState.SyncStateCollection.FindDataByType<FMoverDefaultSyncState>())
+		for (UBaseMovementModeTransition* Transition : MyMoverComponent->Transitions)
 		{
-			for (auto it = MoverState->LayeredMoves.GetActiveMovesIterator(); it; ++it)
+			DataPack.ActiveTransitions.Add(FString::Printf(TEXT("%s (global)"), (Transition ? *Transition->GetClass()->GetName() : TEXT("null"))));
+		}
+		
+		if (MyMoverComponent->HasValidCachedState())
+		{
+			const FMoverSyncState& SyncState = MyMoverComponent->GetSyncState();
+
+			if (const FMoverDefaultSyncState* MoverState = SyncState.SyncStateCollection.FindDataByType<FMoverDefaultSyncState>())
 			{
-				DataPack.ActiveLayeredMoves.Add(*it->Get()->ToSimpleString());
+				for (auto it = MoverState->LayeredMoves.GetActiveMovesIterator(); it; ++it)
+				{
+					DataPack.ActiveLayeredMoves.Add(*it->Get()->ToSimpleString());
+				}
 			}
 		}
 	}
-	else
-	{
-		DataPack.MovementModeName = FString("invalid");
-		DataPack.MovementBaseInfo = FString("invalid");
-		DataPack.MoveIntent = FVector::ZeroVector;
-	}	
 }
 
 
@@ -132,12 +154,14 @@ void FGameplayDebuggerCategory_Mover::DrawData(APlayerController* OwnerPC, FGame
 		DrawInWorldInfo(*FocusedActor, CanvasContext);
 	}
 	
-
-	CanvasContext.Printf(TEXT("{yellow}%s\n{grey}Local Role: {white}%s\n{grey}Mode: {white}%s\n{grey}Active Moves: {green}\n%s"),
+	CanvasContext.Printf(TEXT("{yellow}%s\n{grey}Local Role: {white}%s\n{grey}Mode: {white}%s\n{yellow}Active Moves: {white}\n%s\n{yellow}Mode Map: \n{white}%s\n{yellow}Active Transitions: {white}\n%s"),
 		*DataPack.PawnName,
 		*DataPack.LocalRole,
 		*DataPack.MovementModeName,
-		*FString::JoinBy(DataPack.ActiveLayeredMoves, TEXT("\n"), [](FString MoveAsString) { return MoveAsString; }));
+		*FString::JoinBy(DataPack.ActiveLayeredMoves, TEXT("\n"), [](FString MoveAsString) { return MoveAsString; }),
+		*FString::JoinBy(DataPack.ModeMap, TEXT("\n"), [](FString ModeMappingAsString) { return ModeMappingAsString; }),
+		*FString::JoinBy(DataPack.ActiveTransitions, TEXT("\n"), [](FString TransitionAsString) { return TransitionAsString; })
+		);
 }
 
 
@@ -242,6 +266,9 @@ void FGameplayDebuggerCategory_Mover::FRepData::Serialize(FArchive& Ar)
 	Ar << MovementModeName;
 	Ar << MovementBaseInfo;
 	Ar << MoveIntent;
+	Ar << ActiveLayeredMoves;
+	Ar << ModeMap;
+	Ar << ActiveTransitions;
 }
 
 #endif // WITH_GAMEPLAY_DEBUGGER
\ No newline at end of file
diff --git a/Engine/Plugins/Experimental/Mover/Source/Mover/Private/Debug/GameplayDebuggerCategory_Mover.h b/Engine/Plugins/Experimental/Mover/Source/Mover/Private/Debug/GameplayDebuggerCategory_Mover.h
--- a/Engine/Plugins/Experimental/Mover/Source/Mover/Private/Debug/GameplayDebuggerCategory_Mover.h
+++ b/Engine/Plugins/Experimental/Mover/Source/Mover/Private/Debug/GameplayDebuggerCategory_Mover.h
@@ -31,6 +31,8 @@ protected:
 		FString MovementBaseInfo;
 		FVector MoveIntent;
 		TArray<FString> ActiveLayeredMoves;
+		TArray<FString> ModeMap;
+		TArray<FString> ActiveTransitions;
 
 		FRepData() {}
 
