[Mover]
* Changed creation of NetworkPhysicsComponent to use CreateDefaultSubobject instead of NewObject to fix live coding crash

diff --git a/Engine/Plugins/Experimental/Mover/Source/Mover/Private/Backends/MoverNetworkPhysicsLiaison.cpp b/Engine/Plugins/Experimental/Mover/Source/Mover/Private/Backends/MoverNetworkPhysicsLiaison.cpp
--- a/Engine/Plugins/Experimental/Mover/Source/Mover/Private/Backends/MoverNetworkPhysicsLiaison.cpp
+++ b/Engine/Plugins/Experimental/Mover/Source/Mover/Private/Backends/MoverNetworkPhysicsLiaison.cpp
@@ -242,7 +242,7 @@ UMoverNetworkPhysicsLiaisonComponent::UMoverNetworkPhysicsLiaisonComponent()
 			MyActor->SetReplicateMovement(true);
 		}
 
-		NetworkPhysicsComponent = NewObject<UNetworkPhysicsComponent>(GetOwner(), TEXT("PhysMover_NetworkPhysicsComponent"));
+		NetworkPhysicsComponent = CreateDefaultSubobject<UNetworkPhysicsComponent>(TEXT("PhysMover_NetworkPhysicsComponent"));
 		NetworkPhysicsComponent->SetNetAddressable(); // Make DSO components net addressable
 		NetworkPhysicsComponent->SetIsReplicated(true);
 		NetworkPhysicsComponent->RegisterComponent();
