Physics Mover
- Update sync state correctly (including movement mode) when there are multiple physics ticks, including when resimulating
- Reset movement modes and blackboard on rollback

diff --git a/Engine/Plugins/Experimental/Mover/Source/Mover/Private/Backends/MoverNetworkPhysicsLiaison.cpp b/Engine/Plugins/Experimental/Mover/Source/Mover/Private/Backends/MoverNetworkPhysicsLiaison.cpp
--- a/Engine/Plugins/Experimental/Mover/Source/Mover/Private/Backends/MoverNetworkPhysicsLiaison.cpp
+++ b/Engine/Plugins/Experimental/Mover/Source/Mover/Private/Backends/MoverNetworkPhysicsLiaison.cpp
@@ -761,9 +761,11 @@ void UMoverNetworkPhysicsLiaisonComponent::ProcessInputs_Internal(int32 PhysicsS
 				GetCurrentStateData(Input.SyncState);
 
 				bool bIsSolverResim = false;
+				bool bIsFirstResimFrame = false;
 				if (Chaos::FPhysicsSolver* Solver = PhysScene->GetSolver())
 				{
 					bIsSolverResim = Solver->GetEvolution()->IsResimming();
+					bIsFirstResimFrame = Solver->GetEvolution()->IsResetting();
 				}
 
 				bool bLocalPlayer = false;
@@ -777,6 +779,13 @@ void UMoverNetworkPhysicsLiaisonComponent::ProcessInputs_Internal(int32 PhysicsS
 				{
 					GetCurrentInputData(Input.InputCmd);
 				}
+
+				// Rollback mover state if on the first resimulation frame
+				if (bLocalPlayer && bIsSolverResim && bIsFirstResimFrame)
+				{
+					FMoverAuxStateContext UnusedAuxState;
+					MoverComp->OnSimulationRollback(&Input.SyncState, &UnusedAuxState);
+				}
 			}
 		}
 	}
@@ -848,74 +857,60 @@ void UMoverNetworkPhysicsLiaisonComponent::OnPreSimulate_Internal(const FPhysics
 
 	///////////////////////////////////////////////////////////////////////////////////////////////////////////////////////////
 	// Update the simulation
-	//Input.MoverSimulation->SimulationTick(TickParams, Input.SimInput, SimOutput);
 
 	FMoverTickStartData TickStartData(Input.InputCmd, Input.SyncState, FMoverAuxStateContext());
 	FMoverTickEndData TickEndData;
-	FFloorCheckResult FloorResult;
 
-	FMoverTimeStep TimeStep = GetCurrentMoverTimeStep_Internal();
+	// Sync state should carry over to the next sim frame by default unless something modifies it
+	TickEndData.SyncState = TickStartData.SyncState;
 
-	bool bHasRolledBack = false; // TODO
+	const FMoverDefaultSyncState* StartingSyncState = TickStartData.SyncState.SyncStateCollection.FindDataByType<FMoverDefaultSyncState>();
+	check(StartingSyncState);
 
-	//-------------------------------------------------------------------------------------------
-	// Copied from KinematicMoverComponent::SimulationTick
+	FKinematicDefaultInputs* InputCmd = TickStartData.InputCmd.InputCollection.FindMutableDataByType<FKinematicDefaultInputs>();
+	check(InputCmd);
 
-	if (bHasRolledBack)
+	FMoverTimeStep TimeStep = GetCurrentMoverTimeStep_Internal();
+
+	// Update movement state machine
+
+	if (MoverComp->bHasRolledBack)
 	{
 		MoverComp->ProcessFirstSimTickAfterRollback(TimeStep);
 	}
 
-	// Sync state should carry over to the next sim frame by default unless something modifies it
-	TickEndData.SyncState = TickStartData.SyncState;
-
-	if (MoverComp->ModeFSM->IsValidLowLevel() && MoverComp->SimBlackboard->IsValidLowLevel())
+	if (InputCmd && !InputCmd->SuggestedMovementMode.IsNone())
 	{
-		const FMoverDefaultSyncState* StartingSyncState = TickStartData.SyncState.SyncStateCollection.FindDataByType<FMoverDefaultSyncState>();
-		check(StartingSyncState);
-
-		FKinematicDefaultInputs* InputCmd = TickStartData.InputCmd.InputCollection.FindMutableDataByType<FKinematicDefaultInputs>();
-		FMoverDefaultSyncState& OutputSyncState = TickEndData.SyncState.SyncStateCollection.FindOrAddMutableDataByType<FMoverDefaultSyncState>();
-
-		if (InputCmd && !InputCmd->SuggestedMovementMode.IsNone())
-		{
-			MoverComp->QueueNextMode(InputCmd->SuggestedMovementMode);
-		}
+		MoverComp->QueueNextMode(InputCmd->SuggestedMovementMode);
+	}
 
-		// Tick the actual simulation. This is where the proposed moves are queried and executed, affecting change to the moving actor's gameplay state and captured in the output sim state
-		MoverComp->ModeFSM->OnSimulationTick(MoverComp->UpdatedComponent, MoverComp->UpdatedCompAsPrimitive, MoverComp->SimBlackboard.Get(), TickStartData, TimeStep, TickEndData);
+	// Tick the actual simulation. This is where the proposed moves are queried and executed, affecting change to the moving actor's gameplay state and captured in the output sim state
+	MoverComp->ModeFSM->OnSimulationTick(MoverComp->UpdatedComponent, MoverComp->UpdatedCompAsPrimitive, MoverComp->SimBlackboard.Get(), TickStartData, TimeStep, TickEndData);
 
-		const FName MovementModeAfterTick = MoverComp->ModeFSM->GetCurrentModeName();
-		OutputSyncState.MovementMode = MovementModeAfterTick;
-	}
+	// Set the output sync state and fill in the movement mode
+	Output.SyncState = TickEndData.SyncState;
+	FMoverDefaultSyncState& OutputSyncState = Output.SyncState.SyncStateCollection.FindOrAddMutableDataByType<FMoverDefaultSyncState>();
 
-	//-------------------------------------------------------------------------------------------
+	const FName MovementModeAfterTick = MoverComp->ModeFSM->GetCurrentModeName();
+	OutputSyncState.MovementMode = MovementModeAfterTick;
 
 	MoverComp->SimBlackboard->TryGet(KinematicBlackboard::LastFloorResult, Output.FloorResult);
 
-	Output.SyncState = TickEndData.SyncState;
-
 	/////////////////////////////////////////////////////////////////////////////////////////////////////////////////////////////
-
 	// Update physics constraint from output sync state
-	const FMoverDefaultSyncState* OutputSyncState = Output.SyncState.SyncStateCollection.FindDataByType<FMoverDefaultSyncState>();
-	if (!ensure(OutputSyncState))
-	{
-		return;
-	}
 
-	FVector TargetDeltaPos = OutputSyncState->GetLocation_WorldSpace() - CharacterParticle->GetX();
+	FVector TargetDeltaPos = OutputSyncState.GetLocation_WorldSpace() - CharacterParticle->GetX();
 
 	if (TargetDeltaPos.SizeSquared2D() > GPhysicsDrivenMotionDebugParams.TeleportThreshold * GPhysicsDrivenMotionDebugParams.TeleportThreshold)
 	{
-		TeleportParticle(CharacterParticle, OutputSyncState->GetLocation_WorldSpace(), OutputSyncState->GetOrientation_WorldSpace().Quaternion());
+		TeleportParticle(CharacterParticle, OutputSyncState.GetLocation_WorldSpace(), OutputSyncState.GetOrientation_WorldSpace().Quaternion());
 	}
 
 	// Add back the ground velocity that was subtracted to but the movement velocity in local space
-	FVector TargetVelocity = OutputSyncState->GetVelocity_WorldSpace() + LocalGroundVelocity;
+	FVector TargetVelocity = OutputSyncState.GetVelocity_WorldSpace() + LocalGroundVelocity;
 
 	// Landed so add the new ground velocity
-	if ((OutputSyncState->MovementMode == KinematicModeNames::Walking) && (SyncState.MovementMode != KinematicModeNames::Walking))
+	if ((OutputSyncState.MovementMode == KinematicModeNames::Walking) && (SyncState.MovementMode != KinematicModeNames::Walking))
 	{
 		if (const UPhysicsDrivenWalkingMode* WalkingMode = Cast<UPhysicsDrivenWalkingMode>(MoverComp->FindMovementMode(UPhysicsDrivenWalkingMode::StaticClass())))
 		{
@@ -929,7 +924,7 @@ void UMoverNetworkPhysicsLiaisonComponent::OnPreSimulate_Internal(const FPhysics
 
 	// Note: Output sync state does not have a target angular velocity so
 	// use the target orientation
-	FRotator DeltaRotation = OutputSyncState->GetOrientation_WorldSpace() - FRotator(CharacterParticle->GetR());
+	FRotator DeltaRotation = OutputSyncState.GetOrientation_WorldSpace() - FRotator(CharacterParticle->GetR());
 	FRotator Winding, Remainder;
 	DeltaRotation.GetWindingAndRemainder(Winding, Remainder);
 	float TargetDeltaFacing = FMath::DegreesToRadians(Remainder.Yaw);
@@ -989,6 +984,11 @@ void UMoverNetworkPhysicsLiaisonComponent::OnPreSimulate_Internal(const FPhysics
 			});
 	}
 
+	// Physics can tick multiple times using the same input data from the game thread
+	// so make sure to update it here using the results of this update
+	Input.SyncState = Output.SyncState;
+	InputCmd->SuggestedMovementMode = NAME_None; // Should have already been set this frame
+
 	Output.bIsValid = true;
 }
 
