Mover 2.0 adding functions to allow users to poll mover component for active moves  by type.  Also includes some small changes to Multi-jump so some state is accessible in BP for double jump animations

diff --git a/Engine/Plugins/Experimental/Mover/Source/Mover/Private/Kinematic/LayeredMoves/MultiJumpLayeredMove.cpp b/Engine/Plugins/Experimental/Mover/Source/Mover/Private/Kinematic/LayeredMoves/MultiJumpLayeredMove.cpp
--- a/Engine/Plugins/Experimental/Mover/Source/Mover/Private/Kinematic/LayeredMoves/MultiJumpLayeredMove.cpp
+++ b/Engine/Plugins/Experimental/Mover/Source/Mover/Private/Kinematic/LayeredMoves/MultiJumpLayeredMove.cpp
@@ -71,7 +71,7 @@ bool FLayeredMove_MultiJump::GenerateMove(const FMoverTickStartData& StartState,
 	}
 
 	// if we hit a valid floor and it's not the start of the move (since we could start this move on the ground) end this move
-	if (bValidBlackboard && FloorHitResult.IsWalkableFloor() && StartSimTimeMs < TimeStep.BaseSimTimeMs)
+	if ((bValidBlackboard && FloorHitResult.IsWalkableFloor() && StartSimTimeMs < TimeStep.BaseSimTimeMs) || JumpsInAirRemaining <= 0)
 	{
 		DurationMs = 0;
 	}
diff --git a/Engine/Plugins/Experimental/Mover/Source/Mover/Private/MoverComponent.cpp b/Engine/Plugins/Experimental/Mover/Source/Mover/Private/MoverComponent.cpp
--- a/Engine/Plugins/Experimental/Mover/Source/Mover/Private/MoverComponent.cpp
+++ b/Engine/Plugins/Experimental/Mover/Source/Mover/Private/MoverComponent.cpp
@@ -19,6 +19,7 @@
 #include "GameFramework/PhysicsVolume.h"
 #include "Misc/AssertionMacros.h"
 #include "Misc/TransactionObjectEvent.h"
+#include "Blueprint/BlueprintExceptionInfo.h"
 
 #if WITH_EDITOR
 #include "Misc/DataValidation.h"
@@ -757,6 +758,83 @@ void UMoverComponent::QueueLayeredMove(TSharedPtr<FLayeredMoveBase> LayeredMove)
 	ModeFSM->QueueLayeredMove(LayeredMove);
 }
 
+void UMoverComponent::K2_FindActiveLayeredMove(bool& DidSucceed, int32& TargetAsRawBytes) const
+{
+	// This will never be called, the exec version below will be hit instead
+	checkNoEntry();
+}
+
+DEFINE_FUNCTION(UMoverComponent::execK2_FindActiveLayeredMove)
+{
+	P_GET_UBOOL_REF(DidSucceed);
+
+	Stack.MostRecentPropertyAddress = nullptr;
+	Stack.MostRecentPropertyContainer = nullptr;
+	Stack.StepCompiledIn<FStructProperty>(nullptr);
+	
+	void* MovePtr = Stack.MostRecentPropertyAddress;
+	FStructProperty* StructProp = CastField<FStructProperty>(Stack.MostRecentProperty);
+
+	P_FINISH;
+
+	DidSucceed = false;
+	
+	if (!MovePtr || !StructProp)
+	{
+		FBlueprintExceptionInfo ExceptionInfo(
+			EBlueprintExceptionType::AbortExecution,
+			LOCTEXT("MoverComponent_GetActiveLayeredMove_UnresolvedTarget", "Failed to resolve the TargetAsRawBytes for GetActiveLayeredMove")
+		);
+
+		FBlueprintCoreDelegates::ThrowScriptException(P_THIS, Stack, ExceptionInfo);
+	}
+	else if (!StructProp->Struct || !StructProp->Struct->IsChildOf(FLayeredMoveBase::StaticStruct()))
+	{
+		FBlueprintExceptionInfo ExceptionInfo(
+			EBlueprintExceptionType::AbortExecution,
+			LOCTEXT("MoverComponent_GetActiveLayeredMove_BadType", "TargetAsRawBytes is not a valid type. Must be a child of FLayeredMoveBase.")
+		);
+
+		FBlueprintCoreDelegates::ThrowScriptException(P_THIS, Stack, ExceptionInfo);
+	}
+	else
+	{
+		P_NATIVE_BEGIN;
+		
+		if (const FLayeredMoveBase* FoundActiveMove = P_THIS->FindActiveLayeredMoveByType(StructProp->Struct))
+		{
+			StructProp->Struct->CopyScriptStruct(MovePtr, FoundActiveMove);
+			DidSucceed = true;
+		}
+
+		P_NATIVE_END;
+	}
+}
+
+const FLayeredMoveBase* UMoverComponent::FindActiveLayeredMoveByType(const UScriptStruct* DataStructType) const
+{
+	if (bHasValidCachedState)
+	{
+		if (const FMoverDefaultSyncState* MoverState = CachedLastSyncState.SyncStateCollection.FindDataByType<FMoverDefaultSyncState>())
+		{
+			for (auto it = MoverState->LayeredMoves.GetActiveMovesIterator(); it; ++it)
+			{
+				UStruct* CandidateStruct = it->Get()->GetScriptStruct();
+				while (CandidateStruct)
+				{
+					if (DataStructType == CandidateStruct)
+					{
+						return it->Get();
+					}
+
+					CandidateStruct = CandidateStruct->GetSuperStruct();
+				}
+			}
+		}
+	}
+
+	return nullptr;
+}
 
 void UMoverComponent::QueueNextMode(FName DesiredModeName, bool bShouldReenter)
 {
diff --git a/Engine/Plugins/Experimental/Mover/Source/Mover/Public/Kinematic/LayeredMoves/MultiJumpLayeredMove.h b/Engine/Plugins/Experimental/Mover/Source/Mover/Public/Kinematic/LayeredMoves/MultiJumpLayeredMove.h
--- a/Engine/Plugins/Experimental/Mover/Source/Mover/Public/Kinematic/LayeredMoves/MultiJumpLayeredMove.h
+++ b/Engine/Plugins/Experimental/Mover/Source/Mover/Public/Kinematic/LayeredMoves/MultiJumpLayeredMove.h
@@ -44,10 +44,10 @@ struct MOVER_API FLayeredMove_MultiJump : public FLayeredMoveBase
 	virtual void AddReferencedObjects(class FReferenceCollector& Collector) override;
 
 protected:
-	UPROPERTY()
+	UPROPERTY(BlueprintReadOnly, EditAnywhere, Category = "Jumping")
 	int32 JumpsInAirRemaining;
 
-	UPROPERTY()
+	UPROPERTY(BlueprintReadOnly, EditAnywhere, Category = "Jumping")
 	float TimeOfLastJumpMS;
 
 	bool PerformJump(const FMoverDefaultSyncState* SyncState, const FMoverTimeStep& TimeStep, const UMoverComponent* MoverComp, FProposedMove& OutProposedMove);
diff --git a/Engine/Plugins/Experimental/Mover/Source/Mover/Public/MoverComponent.h b/Engine/Plugins/Experimental/Mover/Source/Mover/Public/MoverComponent.h
--- a/Engine/Plugins/Experimental/Mover/Source/Mover/Public/MoverComponent.h
+++ b/Engine/Plugins/Experimental/Mover/Source/Mover/Public/MoverComponent.h
@@ -282,6 +282,30 @@ public:	// Queries
 	UFUNCTION(BlueprintPure, Category = Mover,  meta=(DeterminesOutputType="MovementMode"))
 	UBaseMovementMode* FindMovementMode(TSubclassOf<UBaseMovementMode> MovementMode) const;
 
+	/**
+	 * Retrieves an active layered move, by writing to a target instance if it is the matching type. Note: Writing to the struct returned will not modify the active struct.
+	 * @param DidSucceed			Flag indicating whether data was actually written to target struct instance
+	 * @param TargetAsRawBytes		The data struct instance to write to, which must be a FLayeredMoveBase sub-type
+	 */
+	UFUNCTION(BlueprintCallable, CustomThunk, Category = Mover, meta = (CustomStructureParam = "TargetAsRawBytes", AllowAbstract = "false", DisplayName = "Find Active Layered Move"))
+	void K2_FindActiveLayeredMove(bool& DidSucceed, UPARAM(DisplayName = "Out Layered Move") int32& TargetAsRawBytes) const;
+	DECLARE_FUNCTION(execK2_FindActiveLayeredMove);
+
+	// Find an active layered move by type. Returns null if one wasn't found 
+	const FLayeredMoveBase* FindActiveLayeredMoveByType(const UScriptStruct* DataStructType) const;
+
+	/** Find a layered move of a specific type in this components active layered moves. If not found, null will be returned. */
+	template <typename T>
+	const T* FindActiveLayeredMoveByType() const
+	{
+		if (const FLayeredMoveBase* FoundData = FindActiveLayeredMoveByType(T::StaticStruct()))
+		{
+			return static_cast<const T*>(FoundData);
+		}
+
+		return nullptr;
+	}
+	
 protected:
 
 	/** Makes this component and owner actor reflect the state of a particular frame snapshot. This occurs after simulation ticking, as well as during a rollback before we resimulate forward.
