Mover 2.0 making multijump properly assign InAirJumps left so walking off platforms still allows the actor an in air jump

diff --git a/Engine/Plugins/Experimental/Mover/Source/Mover/Private/Kinematic/LayeredMoves/MultiJumpLayeredMove.cpp b/Engine/Plugins/Experimental/Mover/Source/Mover/Private/Kinematic/LayeredMoves/MultiJumpLayeredMove.cpp
--- a/Engine/Plugins/Experimental/Mover/Source/Mover/Private/Kinematic/LayeredMoves/MultiJumpLayeredMove.cpp
+++ b/Engine/Plugins/Experimental/Mover/Source/Mover/Private/Kinematic/LayeredMoves/MultiJumpLayeredMove.cpp
@@ -37,13 +37,16 @@ bool FLayeredMove_MultiJump::GenerateMove(const FMoverTickStartData& StartState,
 	FFloorCheckResult FloorHitResult;
 	bool bValidBlackboard = SimBlackboard->TryGet(KinematicBlackboard::LastFloorResult, OUT FloorHitResult);
 
+	if (StartSimTimeMs == TimeStep.BaseSimTimeMs)
+	{
+		JumpsInAirRemaining = MaximumInAirJumps;
+	}
+	
 	bool bPerformedJump = false;
 	if (KinematicInputs && KinematicInputs->bIsJumpJustPressed)
 	{
-		if (JumpsInAirRemaining < 0)
+		if (StartSimTimeMs == TimeStep.BaseSimTimeMs)
 		{
-			JumpsInAirRemaining = MaximumInAirJumps;
-
 			// if this was the first jump and its a valid floor we do the initial jump from walking and back out so we don't get extra jump
 			if (bValidBlackboard && FloorHitResult.IsWalkableFloor())
 			{
