Mover: removing old 'kinematic' nomenclature.

diff --git a/Engine/Plugins/Experimental/Mover/Config/BaseMover.ini b/Engine/Plugins/Experimental/Mover/Config/BaseMover.ini
--- a/Engine/Plugins/Experimental/Mover/Config/BaseMover.ini
+++ b/Engine/Plugins/Experimental/Mover/Config/BaseMover.ini
@@ -5,4 +5,8 @@
 +StructRedirects=(OldName="/Script/Mover.KinematicMoverAuxStateContext",NewName="/Script/Mover.MoverAuxStateContext")
 +StructRedirects=(OldName="/Script/Mover.KinematicMoverTickStartData",NewName="/Script/Mover.MoverTickStartData")
 +StructRedirects=(OldName="/Script/Mover.KinematicMoverTickEndData",NewName="/Script/Mover.MoverTickEndData")
++StructRedirects=(OldName="/Script/Mover.KinematicDefaultInputs",NewName="/Script/Mover.CharacterDefaultInputs")
+
+
+
 
diff --git a/Engine/Plugins/Experimental/Mover/Source/Mover/Private/Backends/MoverNetworkPhysicsLiaison.cpp b/Engine/Plugins/Experimental/Mover/Source/Mover/Private/Backends/MoverNetworkPhysicsLiaison.cpp
--- a/Engine/Plugins/Experimental/Mover/Source/Mover/Private/Backends/MoverNetworkPhysicsLiaison.cpp
+++ b/Engine/Plugins/Experimental/Mover/Source/Mover/Private/Backends/MoverNetworkPhysicsLiaison.cpp
@@ -10,7 +10,7 @@
 #include "GameFramework/Actor.h"
 #include "GameFramework/Pawn.h"
 #include "MovementModeStateMachine.h"
-#include "Kinematic/Settings/CommonLegacyMovementSettings.h"
+#include "DefaultMovementSet/Settings/CommonLegacyMovementSettings.h"
 #include "PBDRigidsSolver.h"
 #include "Physics/Experimental/PhysScene_Chaos.h"
 #include "PhysicsMover/Modes/PhysicsDrivenFallingMode.h"
@@ -72,15 +72,15 @@ void FNetworkPhysicsMoverInputs::InterpolateData(const FNetworkPhysicsData& MinD
 
 	const float LerpFactor = (LocalFrame - MinDataInput.LocalFrame) / (MaxDataInput.LocalFrame - MinDataInput.LocalFrame);
 
-	const FKinematicDefaultInputs* MinInput = MinDataInput.InputCmdContext.InputCollection.FindDataByType<FKinematicDefaultInputs>();
-	const FKinematicDefaultInputs* MaxInput = MaxDataInput.InputCmdContext.InputCollection.FindDataByType<FKinematicDefaultInputs>();
+	const FCharacterDefaultInputs* MinInput = MinDataInput.InputCmdContext.InputCollection.FindDataByType<FCharacterDefaultInputs>();
+	const FCharacterDefaultInputs* MaxInput = MaxDataInput.InputCmdContext.InputCollection.FindDataByType<FCharacterDefaultInputs>();
 
-	FKinematicDefaultInputs& LocalInput = InputCmdContext.InputCollection.FindOrAddMutableDataByType<FKinematicDefaultInputs>();
+	FCharacterDefaultInputs& LocalInput = InputCmdContext.InputCollection.FindOrAddMutableDataByType<FCharacterDefaultInputs>();
 
 	if (MinInput && MaxInput)
 	{
 		// Note, this ignores movement base as this is not used by the physics mover
-		const FKinematicDefaultInputs* ClosestInputs = LerpFactor < 0.5f ? MinInput : MaxInput;
+		const FCharacterDefaultInputs* ClosestInputs = LerpFactor < 0.5f ? MinInput : MaxInput;
 		LocalInput.bIsJumpJustPressed = ClosestInputs->bIsJumpJustPressed;
 		LocalInput.bIsJumpPressed = ClosestInputs->bIsJumpPressed;
 		LocalInput.SuggestedMovementMode = ClosestInputs->SuggestedMovementMode;
@@ -104,9 +104,9 @@ void FNetworkPhysicsMoverInputs::MergeData(const FNetworkPhysicsData& FromData)
 {
 	const FNetworkPhysicsMoverInputs& FromDataInput = static_cast<const FNetworkPhysicsMoverInputs&>(FromData);
 
-	if (const FKinematicDefaultInputs* FromInput = FromDataInput.InputCmdContext.InputCollection.FindDataByType<FKinematicDefaultInputs>())
+	if (const FCharacterDefaultInputs* FromInput = FromDataInput.InputCmdContext.InputCollection.FindDataByType<FCharacterDefaultInputs>())
 	{
-		FKinematicDefaultInputs& LocalInputs = InputCmdContext.InputCollection.FindOrAddMutableDataByType<FKinematicDefaultInputs>();
+		FCharacterDefaultInputs& LocalInputs = InputCmdContext.InputCollection.FindOrAddMutableDataByType<FCharacterDefaultInputs>();
 
 		LocalInputs.bIsJumpJustPressed |= FromInput->bIsJumpJustPressed;
 		LocalInputs.bIsJumpPressed |= FromInput->bIsJumpPressed;
@@ -633,7 +633,7 @@ void UMoverNetworkPhysicsLiaisonComponent::ProduceInput_External(float DeltaSeco
 				// We only want to consume one input per physics frame
 				// so if there is already a valid cached input we use that.
 				// Input is set invalid when the async output is consumed
-				bCachedInputIsValid = MoverComp->CachedLastProducedInputCmd.InputCollection.FindDataByType<FKinematicDefaultInputs>() != nullptr;
+				bCachedInputIsValid = MoverComp->CachedLastProducedInputCmd.InputCollection.FindDataByType<FCharacterDefaultInputs>() != nullptr;
 			}
 			else
 			{
@@ -708,7 +708,7 @@ void UMoverNetworkPhysicsLiaisonComponent::ConsumeOutput_External(const FPhysics
 						{
 							if (const FMoverDefaultSyncState* InterpolatedState = InterpolatedSyncState.SyncStateCollection.FindDataByType<FMoverDefaultSyncState>())
 							{
-								if ((CachedSyncState->MovementMode == KinematicModeNames::Falling) && (InterpolatedState->MovementMode == KinematicModeNames::Walking))
+								if ((CachedSyncState->MovementMode == DefaultModeNames::Falling) && (InterpolatedState->MovementMode == DefaultModeNames::Walking))
 								{
 									if (UFallingMode* FallingMode = MoverComp->FindMode_Mutable<UFallingMode>())
 									{
@@ -805,7 +805,7 @@ void UMoverNetworkPhysicsLiaisonComponent::OnPreSimulate_Internal(const FPhysics
 	// Sync state should carry over to the next sim frame by default unless something modifies it
 	Output.SyncState = Input.SyncState;
 
-	if (!HasValidState() || !Input.InputCmd.InputCollection.FindDataByType<FKinematicDefaultInputs>())
+	if (!HasValidState() || !Input.InputCmd.InputCollection.FindDataByType<FCharacterDefaultInputs>())
 	{
 		return;
 	}
@@ -841,12 +841,12 @@ void UMoverNetworkPhysicsLiaisonComponent::OnPreSimulate_Internal(const FPhysics
 
 	// Make the sync state velocity relative to the ground if walking
 	FVector LocalGroundVelocity = FVector::ZeroVector;
-	if (SyncState.MovementMode == KinematicModeNames::Walking)
+	if (SyncState.MovementMode == DefaultModeNames::Walking)
 	{
 		if (const UMoverBlackboard* Blackboard = MoverComp->GetSimBlackboard())
 		{
 			FFloorCheckResult LastFloorResult;
-			if (Blackboard->TryGet(KinematicBlackboard::LastFloorResult, LastFloorResult))
+			if (Blackboard->TryGet(CommonBlackboard::LastFloorResult, LastFloorResult))
 			{
 				LocalGroundVelocity = UPhysicsMovementUtils::ComputeGroundVelocityFromHitResult(CharacterParticle->GetX(), LastFloorResult.HitResult, TickParams.DeltaTimeSeconds);
 				LocalGroundVelocity -= LocalGroundVelocity.ProjectOnToNormal(LastFloorResult.HitResult.ImpactNormal);
@@ -867,7 +867,7 @@ void UMoverNetworkPhysicsLiaisonComponent::OnPreSimulate_Internal(const FPhysics
 	const FMoverDefaultSyncState* StartingSyncState = TickStartData.SyncState.SyncStateCollection.FindDataByType<FMoverDefaultSyncState>();
 	check(StartingSyncState);
 
-	FKinematicDefaultInputs* InputCmd = TickStartData.InputCmd.InputCollection.FindMutableDataByType<FKinematicDefaultInputs>();
+	FCharacterDefaultInputs* InputCmd = TickStartData.InputCmd.InputCollection.FindMutableDataByType<FCharacterDefaultInputs>();
 	check(InputCmd);
 
 	FMoverTimeStep TimeStep = GetCurrentMoverTimeStep_Internal();
@@ -894,7 +894,7 @@ void UMoverNetworkPhysicsLiaisonComponent::OnPreSimulate_Internal(const FPhysics
 	const FName MovementModeAfterTick = MoverComp->ModeFSM->GetCurrentModeName();
 	OutputSyncState.MovementMode = MovementModeAfterTick;
 
-	MoverComp->SimBlackboard->TryGet(KinematicBlackboard::LastFloorResult, Output.FloorResult);
+	MoverComp->SimBlackboard->TryGet(CommonBlackboard::LastFloorResult, Output.FloorResult);
 
 	/////////////////////////////////////////////////////////////////////////////////////////////////////////////////////////////
 	// Update physics constraint from output sync state
@@ -910,7 +910,7 @@ void UMoverNetworkPhysicsLiaisonComponent::OnPreSimulate_Internal(const FPhysics
 	FVector TargetVelocity = OutputSyncState.GetVelocity_WorldSpace() + LocalGroundVelocity;
 
 	// Landed so add the new ground velocity
-	if ((OutputSyncState.MovementMode == KinematicModeNames::Walking) && (SyncState.MovementMode != KinematicModeNames::Walking))
+	if ((OutputSyncState.MovementMode == DefaultModeNames::Walking) && (SyncState.MovementMode != DefaultModeNames::Walking))
 	{
 		if (const UPhysicsDrivenWalkingMode* WalkingMode = Cast<UPhysicsDrivenWalkingMode>(MoverComp->FindMovementMode(UPhysicsDrivenWalkingMode::StaticClass())))
 		{
@@ -1077,7 +1077,7 @@ void UMoverNetworkPhysicsLiaisonComponent::OnContactModification_Internal(const
 		const float CosThetaMax = 0.707f;
 
 		float MinContactHeightStepUps = CharacterHeight - 1.0e10f;
-		if (SyncState->MovementMode == KinematicModeNames::Walking)
+		if (SyncState->MovementMode == DefaultModeNames::Walking)
 		{
 			if (const UPhysicsDrivenWalkingMode* WalkingMode = Cast<UPhysicsDrivenWalkingMode>(MoverComp->FindMode_Mutable<UPhysicsDrivenWalkingMode>()))
 			{
diff --git a/Engine/Plugins/Experimental/Mover/Source/Mover/Private/DefaultMovementSet/LayeredMoves/AnimRootMotionLayeredMove.cpp b/Engine/Plugins/Experimental/Mover/Source/Mover/Private/DefaultMovementSet/LayeredMoves/AnimRootMotionLayeredMove.cpp
new file mode 100644
--- /dev/null
+++ b/Engine/Plugins/Experimental/Mover/Source/Mover/Private/DefaultMovementSet/LayeredMoves/AnimRootMotionLayeredMove.cpp
@@ -0,0 +1,103 @@
+// Copyright Epic Games, Inc. All Rights Reserved.
+
+#include "DefaultMovementSet/LayeredMoves/AnimRootMotionLayeredMove.h"
+#include "MoverComponent.h"
+#include "Animation/AnimInstance.h"
+#include "Animation/AnimMontage.h"
+#include "Components/SkeletalMeshComponent.h"
+#include "MoveLibrary/MovementUtils.h"
+#include "MoverTypes.h"
+#include "MoverLog.h"
+#include "MotionWarpingComponent.h"
+
+
+FLayeredMove_AnimRootMotion::FLayeredMove_AnimRootMotion()
+{
+	DurationMs = 0.f;
+	MixMode = EMoveMixMode::OverrideAll;
+
+	Montage = nullptr;
+	StartingMontagePosition = 0.f;
+	PlayRate = 1.f;
+}
+
+bool FLayeredMove_AnimRootMotion::GenerateMove(const FMoverTickStartData& SimState, const FMoverTimeStep& TimeStep, const UMoverComponent* MoverComp, UMoverBlackboard* SimBlackboard, FProposedMove& OutProposedMove)
+{
+	const float DeltaSeconds = TimeStep.StepMs / 1000.f;
+
+	const AActor* MoverActor = MoverComp->GetOwner();
+
+	// First pass simply samples based on the duration. For long animations, this has the potential to diverge.
+	// Future improvements could include:
+	//     - speeding up or slowing down slightly to match the associated montage instance
+	//     - detecting if the montage instance is interrupted and attempting to interrupt and scheduling this move to end at the same sim time
+	
+	// Note that Montage 'position' equates to seconds when PlayRate is 1
+	const float SecondsSinceMontageStarted = (TimeStep.BaseSimTimeMs - StartSimTimeMs) / 1000.f;
+	const float ScaledSecondsSinceMontageStarted = SecondsSinceMontageStarted * PlayRate;
+
+	const float ExtractionStartPosition = StartingMontagePosition + ScaledSecondsSinceMontageStarted;
+	const float ExtractionEndPosition   = ExtractionStartPosition + (DeltaSeconds * PlayRate);
+
+	// Read the local transform directly from the montage
+	const FTransform LocalRootMotion = UMotionWarpingUtilities::ExtractRootMotionFromAnimation(Montage, ExtractionStartPosition, ExtractionEndPosition);
+
+	FTransform WorldSpaceRootMotion;
+	
+	if (USkeletalMeshComponent* SkeletalMesh = MoverActor->FindComponentByClass<USkeletalMeshComponent>())
+	{
+		WorldSpaceRootMotion = SkeletalMesh->ConvertLocalRootMotionToWorld(LocalRootMotion);
+	}
+	else
+	{
+		const FTransform ActorToWorldTransform = MoverActor->GetTransform();
+		const FVector DeltaWorldTranslation = LocalRootMotion.GetTranslation() - ActorToWorldTransform.GetTranslation();
+
+		const FQuat NewWorldRotation = ActorToWorldTransform.GetRotation() * LocalRootMotion.GetRotation();
+		const FQuat DeltaWorldRotation = NewWorldRotation * ActorToWorldTransform.GetRotation().Inverse();
+
+		WorldSpaceRootMotion.SetComponents(DeltaWorldRotation, DeltaWorldTranslation, FVector::OneVector);
+	}
+	
+	OutProposedMove = FProposedMove();
+	OutProposedMove.MixMode = MixMode;
+
+	// Convert the transform into linear and angular velocities
+	const FPlane MovementPlane(FVector::ZeroVector, MoverComp->GetUpDirection());
+
+	OutProposedMove.LinearVelocity    = WorldSpaceRootMotion.GetTranslation() / DeltaSeconds;
+	OutProposedMove.MovePlaneVelocity = UMovementUtils::ConstrainToPlane(OutProposedMove.LinearVelocity, MovementPlane);
+	OutProposedMove.AngularVelocity   = WorldSpaceRootMotion.GetRotation().Rotator() * (1.f / DeltaSeconds);
+
+	return true;
+}
+
+FLayeredMoveBase* FLayeredMove_AnimRootMotion::Clone() const
+{
+	FLayeredMove_AnimRootMotion* CopyPtr = new FLayeredMove_AnimRootMotion(*this);
+	return CopyPtr;
+}
+
+void FLayeredMove_AnimRootMotion::NetSerialize(FArchive& Ar)
+{
+	Super::NetSerialize(Ar);
+
+	Ar << Montage;
+	Ar << StartingMontagePosition;
+	Ar << PlayRate;
+}
+
+UScriptStruct* FLayeredMove_AnimRootMotion::GetScriptStruct() const
+{
+	return FLayeredMove_AnimRootMotion::StaticStruct();
+}
+
+FString FLayeredMove_AnimRootMotion::ToSimpleString() const
+{
+	return FString::Printf(TEXT("AnimRootMotion"));
+}
+
+void FLayeredMove_AnimRootMotion::AddReferencedObjects(class FReferenceCollector& Collector)
+{
+	Super::AddReferencedObjects(Collector);
+}
diff --git a/Engine/Plugins/Experimental/Mover/Source/Mover/Private/DefaultMovementSet/LayeredMoves/BasicLayeredMoves.cpp b/Engine/Plugins/Experimental/Mover/Source/Mover/Private/DefaultMovementSet/LayeredMoves/BasicLayeredMoves.cpp
new file mode 100644
--- /dev/null
+++ b/Engine/Plugins/Experimental/Mover/Source/Mover/Private/DefaultMovementSet/LayeredMoves/BasicLayeredMoves.cpp
@@ -0,0 +1,658 @@
+// Copyright Epic Games, Inc. All Rights Reserved.
+
+#include "DefaultMovementSet/LayeredMoves/BasicLayeredMoves.h"
+#include "MoverSimulationTypes.h"
+#include "MoverComponent.h"
+#include "Curves/CurveVector.h"
+#include "Curves/CurveFloat.h"
+#include "MoverLog.h"
+#include "DefaultMovementSet/Settings/CommonLegacyMovementSettings.h"
+#include "MoveLibrary/MovementUtils.h"
+
+
+// -------------------------------------------------------------------
+// FLayeredMove_LinearVelocity
+// -------------------------------------------------------------------
+
+FLayeredMove_LinearVelocity::FLayeredMove_LinearVelocity()
+	: Velocity(FVector::ZeroVector)
+	, MagnitudeOverTime(nullptr)
+	, SettingsFlags(0)
+{
+}
+
+bool FLayeredMove_LinearVelocity::GenerateMove(const FMoverTickStartData& SimState, const FMoverTimeStep& TimeStep, const UMoverComponent* MoverComp, UMoverBlackboard* SimBlackboard, FProposedMove& OutProposedMove)
+{ 
+	const FMoverDefaultSyncState* SyncState = SimState.SyncState.SyncStateCollection.FindDataByType<FMoverDefaultSyncState>();
+	check(SyncState);
+
+	const float DeltaSeconds = TimeStep.StepMs * 0.001f;
+
+	// Convert starting velocity based on starting orientation, if settings call for it
+	if (SettingsFlags & (uint8)ELayeredMove_ConstantVelocitySettingsFlags::VelocityStartRelative &&
+		StartSimTimeMs == TimeStep.BaseSimTimeMs)
+	{
+		SettingsFlags &= ~(uint8)ELayeredMove_ConstantVelocitySettingsFlags::VelocityStartRelative;
+		Velocity = SyncState->GetOrientation_WorldSpace().RotateVector(Velocity);
+	}
+
+	FVector VelocityThisFrame = Velocity;
+
+	// Put velocity into worldspace
+	if (SettingsFlags & (uint8)ELayeredMove_ConstantVelocitySettingsFlags::VelocityAlwaysRelative)
+	{
+		VelocityThisFrame = SyncState->GetOrientation_WorldSpace().RotateVector(Velocity);
+	}
+
+	if (MagnitudeOverTime && DurationMs > 0)
+	{
+		const float TimeValue = DurationMs > 0.f ? FMath::Clamp((TimeStep.BaseSimTimeMs - StartSimTimeMs) / DurationMs, 0.f, 1.f) : TimeStep.BaseSimTimeMs;
+		const float TimeFactor = MagnitudeOverTime->GetFloatValue(TimeValue);
+		VelocityThisFrame *= TimeFactor;
+	}
+	
+	OutProposedMove.LinearVelocity = VelocityThisFrame;
+
+	return true;
+}
+
+FLayeredMoveBase* FLayeredMove_LinearVelocity::Clone() const
+{
+	FLayeredMove_LinearVelocity* CopyPtr = new FLayeredMove_LinearVelocity(*this);
+	return CopyPtr;
+}
+
+void FLayeredMove_LinearVelocity::NetSerialize(FArchive& Ar)
+{
+	Super::NetSerialize(Ar);
+
+	SerializePackedVector<10, 16>(Velocity, Ar);
+	Ar << SettingsFlags;
+	Ar << MagnitudeOverTime;
+}
+
+UScriptStruct* FLayeredMove_LinearVelocity::GetScriptStruct() const
+{
+	return FLayeredMove_LinearVelocity::StaticStruct();
+}
+
+FString FLayeredMove_LinearVelocity::ToSimpleString() const
+{
+	return FString::Printf(TEXT("LinearVelocity"));
+}
+
+void FLayeredMove_LinearVelocity::AddReferencedObjects(class FReferenceCollector& Collector)
+{
+	Super::AddReferencedObjects(Collector);
+}
+
+
+
+
+// -------------------------------------------------------------------
+// FLayeredMove_JumpImpulse
+// -------------------------------------------------------------------
+
+FLayeredMove_JumpImpulse::FLayeredMove_JumpImpulse() 
+	: UpwardsSpeed(0.f)
+{
+	DurationMs = 0.f;
+	MixMode = EMoveMixMode::OverrideVelocity;
+}
+
+
+bool FLayeredMove_JumpImpulse::GenerateMove(const FMoverTickStartData& SimState, const FMoverTimeStep& TimeStep, const UMoverComponent* MoverComp, UMoverBlackboard* SimBlackboard, FProposedMove& OutProposedMove)
+{	
+	const FMoverDefaultSyncState* SyncState = SimState.SyncState.SyncStateCollection.FindDataByType<FMoverDefaultSyncState>();
+	check(SyncState);
+
+	const FVector UpDir = MoverComp->GetUpDirection();
+
+	const FVector ImpulseVelocity = UpDir * UpwardsSpeed;
+
+	// Jump impulse overrides vertical velocity while maintaining the rest
+	if (MixMode == EMoveMixMode::OverrideVelocity)
+	{
+		const FVector PriorVelocityWS = SyncState->GetVelocity_WorldSpace();
+		const FVector StartingNonUpwardsVelocity = PriorVelocityWS - PriorVelocityWS.ProjectOnToNormal(UpDir);
+
+		OutProposedMove.LinearVelocity = StartingNonUpwardsVelocity + ImpulseVelocity;
+	}
+	else
+	{
+		ensureMsgf(false, TEXT("JumpImpulse layered move only supports Override Velocity mix mode and was queued with a different mix mode. Layered move will do nothing."));
+		return false;
+	}
+
+	return true;
+}
+
+FLayeredMoveBase* FLayeredMove_JumpImpulse::Clone() const
+{
+	FLayeredMove_JumpImpulse* CopyPtr = new FLayeredMove_JumpImpulse(*this);
+	return CopyPtr;
+}
+
+void FLayeredMove_JumpImpulse::NetSerialize(FArchive& Ar)
+{
+	Super::NetSerialize(Ar);
+
+	Ar << UpwardsSpeed;
+}
+
+UScriptStruct* FLayeredMove_JumpImpulse::GetScriptStruct() const
+{
+	return FLayeredMove_JumpImpulse::StaticStruct();
+}
+
+FString FLayeredMove_JumpImpulse::ToSimpleString() const
+{
+	return FString::Printf(TEXT("JumpImpulse"));
+}
+
+void FLayeredMove_JumpImpulse::AddReferencedObjects(class FReferenceCollector& Collector)
+{
+	Super::AddReferencedObjects(Collector);
+}
+
+// -------------------------------------------------------------------
+// FLayeredMove_JumpTo
+// -------------------------------------------------------------------
+
+FLayeredMove_JumpTo::FLayeredMove_JumpTo()
+	: JumpDistance(-1.0f)
+	, JumpHeight(-1.0f)
+	, bUseActorRotation(true)
+	, JumpRotation(ForceInitToZero)
+	, PathOffsetCurve(nullptr)
+	, TimeMappingCurve(nullptr)
+{
+	DurationMs = 1.f;
+	MixMode = EMoveMixMode::OverrideVelocity;
+}
+
+bool FLayeredMove_JumpTo::GenerateMove(const FMoverTickStartData& StartState, const FMoverTimeStep& TimeStep, const UMoverComponent* MoverComp, UMoverBlackboard* SimBlackboard, FProposedMove& OutProposedMove)
+{
+	const FMoverDefaultSyncState* SyncState = StartState.SyncState.SyncStateCollection.FindDataByType<FMoverDefaultSyncState>();
+	check(SyncState);
+
+	if (DurationMs == 0)
+	{
+		ensureMsgf(false, TEXT("JumpTo expected a non-zero duration."));
+		return false;
+	}
+	
+	const float DeltaSeconds = TimeStep.StepMs / 1000.f;
+	float CurrentTimeFraction = (TimeStep.BaseSimTimeMs - StartSimTimeMs) / DurationMs;
+	float TargetTimeFraction = CurrentTimeFraction + DeltaSeconds;
+	
+	// If we're beyond specified duration, we need to re-map times so that
+	// we continue our desired ending velocity
+	if (TargetTimeFraction > 1.f)
+	{
+		float TimeFractionPastAllowable = TargetTimeFraction - 1.0f;
+		TargetTimeFraction -= TimeFractionPastAllowable;
+		CurrentTimeFraction -= TimeFractionPastAllowable;
+	}
+
+	float CurrentMoveFraction = CurrentTimeFraction;
+	float TargetMoveFraction = TargetTimeFraction;
+
+	if (TimeMappingCurve)
+	{
+		CurrentMoveFraction = EvaluateFloatCurveAtFraction(*TimeMappingCurve, CurrentMoveFraction);
+		TargetMoveFraction  = EvaluateFloatCurveAtFraction(*TimeMappingCurve, TargetMoveFraction);
+	}
+
+	const FRotator Rotation = bUseActorRotation ? SyncState->GetOrientation_WorldSpace() : JumpRotation;
+	const FVector CurrentRelativeLocation = GetRelativeLocation(CurrentMoveFraction, Rotation);
+	const FVector TargetRelativeLocation = GetRelativeLocation(TargetMoveFraction, Rotation);
+
+	OutProposedMove.LinearVelocity = (TargetRelativeLocation - CurrentRelativeLocation) / DeltaSeconds;;
+	if (const TObjectPtr<const UCommonLegacyMovementSettings> CommonLegacySettings = MoverComp->FindSharedSettings<UCommonLegacyMovementSettings>())
+	{
+		OutProposedMove.PreferredMode = CommonLegacySettings->AirMovementModeName;
+	}
+	
+	return true;
+}
+
+FLayeredMoveBase* FLayeredMove_JumpTo::Clone() const
+{
+	FLayeredMove_JumpTo* CopyPtr = new FLayeredMove_JumpTo(*this);
+	return CopyPtr;
+}
+
+void FLayeredMove_JumpTo::NetSerialize(FArchive& Ar)
+{
+	Super::NetSerialize(Ar);
+
+	Ar << JumpDistance;
+	Ar << JumpHeight;
+	Ar << bUseActorRotation;
+
+	if (bUseActorRotation)
+	{
+		Ar << JumpRotation;
+	}
+
+	Ar << PathOffsetCurve;
+	Ar << TimeMappingCurve;
+}
+
+UScriptStruct* FLayeredMove_JumpTo::GetScriptStruct() const
+{
+	return FLayeredMove_JumpTo::StaticStruct();
+}
+
+FString FLayeredMove_JumpTo::ToSimpleString() const
+{
+	return FString::Printf(TEXT("JumpTo"));
+}
+
+void FLayeredMove_JumpTo::AddReferencedObjects(FReferenceCollector& Collector)
+{
+	Super::AddReferencedObjects(Collector);
+}
+
+FVector FLayeredMove_JumpTo::GetPathOffset(const float MoveFraction) const
+{
+	FVector PathOffset(FVector::ZeroVector);
+	if (PathOffsetCurve)
+	{
+		// Calculate path offset
+		float MinCurveTime(0.f);
+		float MaxCurveTime(1.f);
+
+		PathOffsetCurve->GetTimeRange(MinCurveTime, MaxCurveTime);
+		PathOffset =  PathOffsetCurve->GetVectorValue(FMath::GetRangeValue(FVector2f(MinCurveTime, MaxCurveTime), MoveFraction));
+	}
+	else
+	{
+		// Default to "jump parabola", a simple x^2 shifted to be upside-down and shifted
+		// to get [0,1] X (MoveFraction/Distance) mapping to [0,1] Y (height)
+		// Height = -(2x-1)^2 + 1
+		const float Phi = 2.f*MoveFraction - 1;
+		const float Z = -(Phi*Phi) + 1;
+		PathOffset.Z = Z;
+	}
+
+	// Scale Z offset to height. If Height < 0, we use direct path offset values
+	if (JumpHeight >= 0.f)
+	{
+		PathOffset.Z *= JumpHeight;
+	}
+
+	return PathOffset;
+}
+
+float FLayeredMove_JumpTo::EvaluateFloatCurveAtFraction(const UCurveFloat& Curve, const float Fraction) const
+{
+	float MinCurveTime(0.f);
+	float MaxCurveTime(1.f);
+
+	Curve.GetTimeRange(MinCurveTime, MaxCurveTime);
+	return Curve.GetFloatValue(FMath::GetRangeValue(FVector2f(MinCurveTime, MaxCurveTime), Fraction));
+}
+
+FVector FLayeredMove_JumpTo::GetRelativeLocation(float MoveFraction, const FRotator& Rotator) const
+{
+	// Given MoveFraction, what relative location should a character be at?
+	FRotator FacingRotation(Rotator);
+	FacingRotation.Pitch = 0.f; // By default we don't include pitch, but an option could be added if necessary
+
+	const FVector RelativeLocationFacingSpace = FVector(MoveFraction * JumpDistance, 0.f, 0.f) + GetPathOffset(MoveFraction);
+
+	return FacingRotation.RotateVector(RelativeLocationFacingSpace);
+}
+
+// -------------------------------------------------------------------
+// FLayeredMove_Teleport
+// -------------------------------------------------------------------
+
+FLayeredMove_Teleport::FLayeredMove_Teleport()
+	: TargetLocation(FVector::ZeroVector)
+{
+	DurationMs = 0.f;
+	MixMode = EMoveMixMode::OverrideAll;
+	Priority = 10;
+}
+
+bool FLayeredMove_Teleport::GenerateMove(const FMoverTickStartData& SimState, const FMoverTimeStep& TimeStep, const UMoverComponent* MoverComp, UMoverBlackboard* SimBlackboard, FProposedMove& OutProposedMove)
+{
+	OutProposedMove = FProposedMove();
+	OutProposedMove.MixMode = EMoveMixMode::OverrideAll;
+	OutProposedMove.TargetLocation = TargetLocation;
+	OutProposedMove.bHasTargetLocation = true;
+
+	return true;
+}
+
+FLayeredMoveBase* FLayeredMove_Teleport::Clone() const
+{
+	FLayeredMove_Teleport* CopyPtr = new FLayeredMove_Teleport(*this);
+	return CopyPtr;
+}
+
+void FLayeredMove_Teleport::NetSerialize(FArchive& Ar)
+{
+	Super::NetSerialize(Ar);
+
+	Ar << TargetLocation;
+}
+
+UScriptStruct* FLayeredMove_Teleport::GetScriptStruct() const
+{
+	return FLayeredMove_Teleport::StaticStruct();
+}
+
+FString FLayeredMove_Teleport::ToSimpleString() const
+{
+	return FString::Printf(TEXT("Teleport"));
+}
+
+void FLayeredMove_Teleport::AddReferencedObjects(class FReferenceCollector& Collector)
+{
+	Super::AddReferencedObjects(Collector);
+}
+
+// -------------------------------------------------------------------
+// FLayeredMove_MoveTo
+// -------------------------------------------------------------------
+
+FLayeredMove_MoveTo::FLayeredMove_MoveTo()
+	: StartLocation(ForceInitToZero)
+	, TargetLocation(ForceInitToZero)
+	, bRestrictSpeedToExpected(false)
+	, PathOffsetCurve(nullptr)
+	, TimeMappingCurve(nullptr)
+{
+	DurationMs = 1000.f;
+	MixMode = EMoveMixMode::OverrideVelocity;
+}
+
+FVector FLayeredMove_MoveTo::GetPathOffsetInWorldSpace(const float MoveFraction) const
+{
+	if (PathOffsetCurve)
+	{
+		float MinCurveTime(0.f);
+		float MaxCurveTime(1.f);
+
+		PathOffsetCurve->GetTimeRange(MinCurveTime, MaxCurveTime);
+		
+		// Calculate path offset
+		const FVector PathOffsetInFacingSpace = PathOffsetCurve->GetVectorValue(FMath::GetRangeValue(FVector2f(MinCurveTime, MaxCurveTime), MoveFraction));;
+		FRotator FacingRotation((TargetLocation-StartLocation).Rotation());
+		FacingRotation.Pitch = 0.f; // By default we don't include pitch in the offset, but an option could be added if necessary
+		return FacingRotation.RotateVector(PathOffsetInFacingSpace);
+	}
+
+	return FVector::ZeroVector;
+}
+
+float FLayeredMove_MoveTo::EvaluateFloatCurveAtFraction(const UCurveFloat& Curve, const float Fraction) const
+{
+	float MinCurveTime(0.f);
+	float MaxCurveTime(1.f);
+
+	Curve.GetTimeRange(MinCurveTime, MaxCurveTime);
+	return Curve.GetFloatValue(FMath::GetRangeValue(FVector2f(MinCurveTime, MaxCurveTime), Fraction));
+}
+
+bool FLayeredMove_MoveTo::GenerateMove(const FMoverTickStartData& StartState, const FMoverTimeStep& TimeStep, const UMoverComponent* MoverComp, UMoverBlackboard* SimBlackboard, FProposedMove& OutProposedMove)
+{
+	OutProposedMove.MixMode = MixMode;
+	OutProposedMove.TargetLocation = TargetLocation;
+	OutProposedMove.bHasTargetLocation = true;
+
+	const float DeltaSeconds = TimeStep.StepMs / 1000.f;
+	
+	float MoveFraction = (TimeStep.BaseSimTimeMs - StartSimTimeMs) / DurationMs;
+
+	if (TimeMappingCurve)
+	{
+		MoveFraction = EvaluateFloatCurveAtFraction(*TimeMappingCurve, MoveFraction);
+	}
+	
+	const AActor* MoverActor = MoverComp->GetOwner();
+	
+	FVector CurrentTargetLocation = FMath::Lerp<FVector, float>(StartLocation, TargetLocation, MoveFraction);
+	FVector PathOffset = GetPathOffsetInWorldSpace(MoveFraction);
+	CurrentTargetLocation += PathOffset;
+
+	const FVector CurrentLocation = MoverActor->GetActorLocation();
+
+	FVector Velocity = (CurrentTargetLocation - CurrentLocation) / DeltaSeconds;
+
+	if (bRestrictSpeedToExpected && !Velocity.IsNearlyZero(UE_KINDA_SMALL_NUMBER))
+	{
+		// Calculate expected current location (if we didn't have collision and moved exactly where our velocity should have taken us)
+		const float PreviousMoveFraction = (TimeStep.BaseSimTimeMs - StartSimTimeMs - TimeStep.StepMs) / DurationMs;
+		FVector CurrentExpectedLocation = FMath::Lerp<FVector, float>(StartLocation, TargetLocation, PreviousMoveFraction);
+		CurrentExpectedLocation += GetPathOffsetInWorldSpace(PreviousMoveFraction);
+
+		// Restrict speed to the expected speed, allowing some small amount of error
+		const FVector ExpectedForce = (CurrentTargetLocation - CurrentExpectedLocation) / DeltaSeconds;
+		const float ExpectedSpeed = ExpectedForce.Size();
+		const float CurrentSpeedSqr = Velocity.SizeSquared();
+
+		const float ErrorAllowance = 0.5f; // in cm/s
+		if (CurrentSpeedSqr > FMath::Square(ExpectedSpeed + ErrorAllowance))
+		{
+			Velocity.Normalize();
+			Velocity *= ExpectedSpeed;
+		}
+	}
+	
+	OutProposedMove.LinearVelocity = Velocity;
+	
+	return true;
+}
+
+FLayeredMoveBase* FLayeredMove_MoveTo::Clone() const
+{
+	FLayeredMove_MoveTo* CopyPtr = new FLayeredMove_MoveTo(*this);
+	return CopyPtr;
+}
+
+void FLayeredMove_MoveTo::NetSerialize(FArchive& Ar)
+{
+	Super::NetSerialize(Ar);
+
+	Ar << StartLocation;
+	Ar << TargetLocation;
+	Ar << bRestrictSpeedToExpected;
+	Ar << PathOffsetCurve;
+	Ar << TimeMappingCurve;
+}
+
+UScriptStruct* FLayeredMove_MoveTo::GetScriptStruct() const
+{
+	return FLayeredMove_MoveTo::StaticStruct();
+}
+
+FString FLayeredMove_MoveTo::ToSimpleString() const
+{
+	return FString::Printf(TEXT("Move To"));
+}
+
+void FLayeredMove_MoveTo::AddReferencedObjects(FReferenceCollector& Collector)
+{
+	Super::AddReferencedObjects(Collector);
+}
+
+// -------------------------------------------------------------------
+// FLayeredMove_MoveToDynamic
+// -------------------------------------------------------------------
+
+FLayeredMove_MoveToDynamic::FLayeredMove_MoveToDynamic()
+	: LocationActor(nullptr)
+{
+	StartLocation = FVector::ZeroVector;
+	TargetLocation = FVector::ZeroVector;
+	bRestrictSpeedToExpected = false;
+	PathOffsetCurve = nullptr;
+	TimeMappingCurve = nullptr;
+	DurationMs = 1000.f;
+	MixMode = EMoveMixMode::OverrideVelocity;
+}
+
+bool FLayeredMove_MoveToDynamic::GenerateMove(const FMoverTickStartData& StartState, const FMoverTimeStep& TimeStep, const UMoverComponent* MoverComp, UMoverBlackboard* SimBlackboard, FProposedMove& OutProposedMove)
+{
+	if (LocationActor)
+	{
+		TargetLocation = LocationActor->GetActorLocation();
+	}
+	
+	return Super::GenerateMove(StartState, TimeStep, MoverComp, SimBlackboard, OutProposedMove);
+}
+
+FLayeredMoveBase* FLayeredMove_MoveToDynamic::Clone() const
+{
+	FLayeredMove_MoveToDynamic* CopyPtr = new FLayeredMove_MoveToDynamic(*this);
+	return CopyPtr;
+}
+
+void FLayeredMove_MoveToDynamic::NetSerialize(FArchive& Ar)
+{
+	Super::NetSerialize(Ar);
+
+	Ar << LocationActor;
+}
+
+UScriptStruct* FLayeredMove_MoveToDynamic::GetScriptStruct() const
+{
+	return FLayeredMove_MoveToDynamic::StaticStruct();
+}
+
+FString FLayeredMove_MoveToDynamic::ToSimpleString() const
+{
+	return FString::Printf(TEXT("Move To Dynamic"));
+}
+
+void FLayeredMove_MoveToDynamic::AddReferencedObjects(FReferenceCollector& Collector)
+{
+	Super::AddReferencedObjects(Collector);
+}
+
+// -------------------------------------------------------------------
+// FLayeredMove_RadialImpulse
+// -------------------------------------------------------------------
+
+FLayeredMove_RadialImpulse::FLayeredMove_RadialImpulse()
+	: Location(ForceInitToZero)
+	, LocationActor(nullptr)
+	, Radius(1.f)
+	, Magnitude(0.f)
+	, bIsPush(true)
+	, bNoVerticalVelocity(false)
+	, DistanceFalloff(nullptr)
+	, MagnitudeOverTime(nullptr)
+	, bUseFixedWorldDirection(false)
+	, FixedWorldDirection(ForceInitToZero)
+{
+	DurationMs = 0.f;
+	MixMode = EMoveMixMode::AdditiveVelocity;
+}
+
+bool FLayeredMove_RadialImpulse::GenerateMove(const FMoverTickStartData& StartState, const FMoverTimeStep& TimeStep, const UMoverComponent* MoverComp, UMoverBlackboard* SimBlackboard, FProposedMove& OutProposedMove)
+{
+	const FMoverDefaultSyncState* SyncState = StartState.SyncState.SyncStateCollection.FindDataByType<FMoverDefaultSyncState>();
+	check(SyncState);
+	
+	const FVector CharacterLocation = SyncState->GetLocation_WorldSpace();
+	FVector Velocity = FVector::ZeroVector;
+	const FVector VelocityLocation = LocationActor ? LocationActor->GetActorLocation() : Location;
+	float Distance = FVector::Dist(VelocityLocation, CharacterLocation);
+	
+	if (Distance < Radius)
+	{
+		// Calculate magnitude
+		float CurrentMagnitude = Magnitude;
+		{
+			float AdditiveMagnitudeFactor = 1.f;
+			if (DistanceFalloff)
+			{
+				const float DistanceFactor = DistanceFalloff->GetFloatValue(FMath::Clamp(Distance / Radius, 0.f, 1.f));
+				AdditiveMagnitudeFactor -= (1.f - DistanceFactor);
+			}
+
+			if (MagnitudeOverTime)
+			{
+				const float TimeValue = DurationMs > 0.f ? FMath::Clamp(TimeStep.BaseSimTimeMs / DurationMs, 0.f, 1.f) : TimeStep.BaseSimTimeMs;
+				const float TimeFactor = MagnitudeOverTime->GetFloatValue(TimeValue);
+				AdditiveMagnitudeFactor -= (1.f - TimeFactor);
+			}
+
+			CurrentMagnitude = Magnitude * FMath::Clamp(AdditiveMagnitudeFactor, 0.f, 1.f);
+		}
+
+		if (bUseFixedWorldDirection)
+		{
+			Velocity = FixedWorldDirection.Vector() * CurrentMagnitude;
+		}
+		else
+		{
+			Velocity = (VelocityLocation - CharacterLocation).GetSafeNormal() * CurrentMagnitude;
+			
+			if (bIsPush)
+			{
+				Velocity *= -1.f;
+			}
+		}
+		
+		if (bNoVerticalVelocity)
+        {
+       		const FPlane MovementPlane(FVector::ZeroVector, MoverComp->GetUpDirection());
+       		Velocity = UMovementUtils::ConstrainToPlane(Velocity, MovementPlane);
+       	}
+	}
+	else 
+	{
+	     return false;
+	}
+
+	OutProposedMove.LinearVelocity = Velocity;
+
+	return true;
+}
+
+FLayeredMoveBase* FLayeredMove_RadialImpulse::Clone() const
+{
+	FLayeredMove_RadialImpulse* CopyPtr = new FLayeredMove_RadialImpulse(*this);
+    return CopyPtr;
+}
+
+void FLayeredMove_RadialImpulse::NetSerialize(FArchive& Ar)
+{
+	Super::NetSerialize(Ar);
+	
+	Ar << Location;
+	Ar << LocationActor;
+	Ar << Radius;
+	Ar << Magnitude;
+	Ar << bIsPush;
+	Ar << bNoVerticalVelocity;
+	Ar << DistanceFalloff;
+	Ar << MagnitudeOverTime;
+	Ar << bUseFixedWorldDirection;
+	if (bUseFixedWorldDirection)
+	{
+		Ar << FixedWorldDirection;
+	}
+}
+
+UScriptStruct* FLayeredMove_RadialImpulse::GetScriptStruct() const
+{
+	return FLayeredMove_RadialImpulse::StaticStruct();
+}
+
+FString FLayeredMove_RadialImpulse::ToSimpleString() const
+{
+	return FString::Printf(TEXT("Radial Impulse"));
+}
+
+void FLayeredMove_RadialImpulse::AddReferencedObjects(FReferenceCollector& Collector)
+{
+	Super::AddReferencedObjects(Collector);
+}
diff --git a/Engine/Plugins/Experimental/Mover/Source/Mover/Private/DefaultMovementSet/LayeredMoves/LaunchMove.cpp b/Engine/Plugins/Experimental/Mover/Source/Mover/Private/DefaultMovementSet/LayeredMoves/LaunchMove.cpp
new file mode 100644
--- /dev/null
+++ b/Engine/Plugins/Experimental/Mover/Source/Mover/Private/DefaultMovementSet/LayeredMoves/LaunchMove.cpp
@@ -0,0 +1,60 @@
+// Copyright Epic Games, Inc. All Rights Reserved.
+
+#include "DefaultMovementSet/LayeredMoves/LaunchMove.h"
+#include "MoverComponent.h"
+#include "MoverTypes.h"
+#include "MoverLog.h"
+
+
+FLayeredMove_Launch::FLayeredMove_Launch()
+{
+	DurationMs = 0.f;
+	MixMode = EMoveMixMode::OverrideVelocity;
+}
+
+bool FLayeredMove_Launch::GenerateMove(const FMoverTickStartData& SimState, const FMoverTimeStep& TimeStep, const UMoverComponent* MoverComp, UMoverBlackboard* SimBlackboard, FProposedMove& OutProposedMove)
+{
+	OutProposedMove.MixMode = MixMode;
+	OutProposedMove.LinearVelocity = LaunchVelocity;
+	OutProposedMove.PreferredMode = ForceMovementMode;
+
+	return true;
+}
+
+FLayeredMoveBase* FLayeredMove_Launch::Clone() const
+{
+	FLayeredMove_Launch* CopyPtr = new FLayeredMove_Launch(*this);
+	return CopyPtr;
+}
+
+void FLayeredMove_Launch::NetSerialize(FArchive& Ar)
+{
+	Super::NetSerialize(Ar);
+
+	SerializePackedVector<10, 16>(LaunchVelocity, Ar);
+
+	bool bUsingForcedMovementMode = !ForceMovementMode.IsNone();
+
+	Ar.SerializeBits(&bUsingForcedMovementMode, 1);
+
+	if (bUsingForcedMovementMode)
+	{
+		Ar << ForceMovementMode;
+	}
+
+}
+
+UScriptStruct* FLayeredMove_Launch::GetScriptStruct() const
+{
+	return FLayeredMove_Launch::StaticStruct();
+}
+
+FString FLayeredMove_Launch::ToSimpleString() const
+{
+	return FString::Printf(TEXT("Launch"));
+}
+
+void FLayeredMove_Launch::AddReferencedObjects(class FReferenceCollector& Collector)
+{
+	Super::AddReferencedObjects(Collector);
+}
diff --git a/Engine/Plugins/Experimental/Mover/Source/Mover/Private/DefaultMovementSet/LayeredMoves/MultiJumpLayeredMove.cpp b/Engine/Plugins/Experimental/Mover/Source/Mover/Private/DefaultMovementSet/LayeredMoves/MultiJumpLayeredMove.cpp
new file mode 100644
--- /dev/null
+++ b/Engine/Plugins/Experimental/Mover/Source/Mover/Private/DefaultMovementSet/LayeredMoves/MultiJumpLayeredMove.cpp
@@ -0,0 +1,143 @@
+// Copyright Epic Games, Inc. All Rights Reserved.
+
+#include "DefaultMovementSet/LayeredMoves/MultiJumpLayeredMove.h"
+#include "MoverComponent.h"
+#include "MoverSimulationTypes.h"
+#include "DefaultMovementSet/Settings/CommonLegacyMovementSettings.h"
+#include "MoveLibrary/FloorQueryUtils.h"
+
+FLayeredMove_MultiJump::FLayeredMove_MultiJump()
+	: MaximumInAirJumps(1)
+	, UpwardsSpeed(800)
+	, JumpsInAirRemaining(-1)
+	, TimeOfLastJumpMS(0)
+{
+	DurationMs = -1.f;
+	MixMode = EMoveMixMode::OverrideVelocity;
+}
+
+bool FLayeredMove_MultiJump::WantsToJump(const FMoverInputCmdContext& InputCmd)
+{
+	if (const FCharacterDefaultInputs* CharacterInputs = InputCmd.InputCollection.FindDataByType<FCharacterDefaultInputs>())
+	{
+		return CharacterInputs->bIsJumpJustPressed;
+	}
+	
+	return false;
+}
+
+bool FLayeredMove_MultiJump::GenerateMove(const FMoverTickStartData& StartState, const FMoverTimeStep& TimeStep, const UMoverComponent* MoverComp, UMoverBlackboard* SimBlackboard, FProposedMove& OutProposedMove)
+{
+	const FCharacterDefaultInputs* CharacterInputs = StartState.InputCmd.InputCollection.FindDataByType<FCharacterDefaultInputs>();
+	const FMoverDefaultSyncState* SyncState = StartState.SyncState.SyncStateCollection.FindDataByType<FMoverDefaultSyncState>();
+	check(SyncState);
+
+	OutProposedMove.MixMode = MixMode;
+
+	FFloorCheckResult FloorHitResult;
+	bool bValidBlackboard = SimBlackboard->TryGet(CommonBlackboard::LastFloorResult, OUT FloorHitResult);
+
+	if (StartSimTimeMs == TimeStep.BaseSimTimeMs)
+	{
+		JumpsInAirRemaining = MaximumInAirJumps;
+	}
+	
+	bool bPerformedJump = false;
+	if (CharacterInputs && CharacterInputs->bIsJumpJustPressed)
+	{
+		if (StartSimTimeMs == TimeStep.BaseSimTimeMs)
+		{
+			// if this was the first jump and its a valid floor we do the initial jump from walking and back out so we don't get extra jump
+			if (bValidBlackboard && FloorHitResult.IsWalkableFloor())
+			{
+				bPerformedJump = PerformJump(SyncState, TimeStep, MoverComp, OutProposedMove);
+				return bPerformedJump;
+			}
+		}
+
+		// perform in air jump
+		if (TimeStep.BaseSimTimeMs > TimeOfLastJumpMS && JumpsInAirRemaining > 0)
+		{
+			JumpsInAirRemaining--;
+			bPerformedJump = PerformJump(SyncState, TimeStep, MoverComp, OutProposedMove);
+		}
+		else
+		{
+			// setting mix mode to additive when we're not adding any jump input so air movement acts as expected
+			OutProposedMove.MixMode = EMoveMixMode::AdditiveVelocity;
+		}
+	}
+	else
+	{
+		// setting mix mode to additive when we're not adding any jump input so air movement acts as expected
+		OutProposedMove.MixMode = EMoveMixMode::AdditiveVelocity;
+	}
+
+	// if we hit a valid floor and it's not the start of the move (since we could start this move on the ground) end this move
+	if ((bValidBlackboard && FloorHitResult.IsWalkableFloor() && StartSimTimeMs < TimeStep.BaseSimTimeMs) || JumpsInAirRemaining <= 0)
+	{
+		DurationMs = 0;
+	}
+
+	return bPerformedJump;
+}
+
+FLayeredMoveBase* FLayeredMove_MultiJump::Clone() const
+{
+	FLayeredMove_MultiJump* CopyPtr = new FLayeredMove_MultiJump(*this);
+	return CopyPtr;
+}
+
+void FLayeredMove_MultiJump::NetSerialize(FArchive& Ar)
+{
+	Super::NetSerialize(Ar);
+
+	Ar << MaximumInAirJumps;
+	Ar << UpwardsSpeed;
+	Ar << JumpsInAirRemaining;
+	Ar << TimeOfLastJumpMS;
+}
+
+UScriptStruct* FLayeredMove_MultiJump::GetScriptStruct() const
+{
+	return FLayeredMove_MultiJump::StaticStruct();
+}
+
+FString FLayeredMove_MultiJump::ToSimpleString() const
+{
+	return FString::Printf(TEXT("Multi-jump"));
+}
+
+void FLayeredMove_MultiJump::AddReferencedObjects(FReferenceCollector& Collector)
+{
+	Super::AddReferencedObjects(Collector);
+}
+
+bool FLayeredMove_MultiJump::PerformJump(const FMoverDefaultSyncState* SyncState, const FMoverTimeStep& TimeStep, const UMoverComponent* MoverComp, FProposedMove& OutProposedMove)
+{
+	TimeOfLastJumpMS = TimeStep.BaseSimTimeMs;
+	if (const TObjectPtr<const UCommonLegacyMovementSettings> CommonLegacySettings = MoverComp->FindSharedSettings<UCommonLegacyMovementSettings>())
+	{
+		OutProposedMove.PreferredMode = CommonLegacySettings->AirMovementModeName;
+	}
+
+	const FVector UpDir = MoverComp->GetUpDirection();
+
+	const FVector ImpulseVelocity = UpDir * UpwardsSpeed;
+
+	// Jump impulse overrides vertical velocity while maintaining the rest
+	if (MixMode == EMoveMixMode::OverrideVelocity)
+	{
+		const FVector PriorVelocityWS = SyncState->GetVelocity_WorldSpace();
+		const FVector StartingNonUpwardsVelocity = PriorVelocityWS - PriorVelocityWS.ProjectOnToNormal(UpDir);
+
+		OutProposedMove.LinearVelocity = StartingNonUpwardsVelocity + ImpulseVelocity;
+	}
+	else
+	{
+		ensureMsgf(false, TEXT("Multi-Jump layered move only supports Override Velocity mix mode and was queued with a different mix mode. Layered move will do nothing."));
+		return false;
+	}
+
+	return true;
+}
diff --git a/Engine/Plugins/Experimental/Mover/Source/Mover/Private/DefaultMovementSet/Modes/FallingMode.cpp b/Engine/Plugins/Experimental/Mover/Source/Mover/Private/DefaultMovementSet/Modes/FallingMode.cpp
new file mode 100644
--- /dev/null
+++ b/Engine/Plugins/Experimental/Mover/Source/Mover/Private/DefaultMovementSet/Modes/FallingMode.cpp
@@ -0,0 +1,348 @@
+// Copyright Epic Games, Inc. All Rights Reserved.
+
+#include "DefaultMovementSet/Modes/FallingMode.h"
+#include "Components/SkeletalMeshComponent.h"
+#include "MoverComponent.h"
+#include "MoveLibrary/MovementUtils.h"
+#include "MoveLibrary/BasedMovementUtils.h"
+#include "MoveLibrary/FloorQueryUtils.h"
+#include "DefaultMovementSet/Settings/CommonLegacyMovementSettings.h"
+#include "MoveLibrary/AirMovementUtils.h"
+
+
+#include UE_INLINE_GENERATED_CPP_BY_NAME(FallingMode)
+
+UFallingMode::UFallingMode(const FObjectInitializer& ObjectInitializer)
+	: Super(ObjectInitializer),
+	  AirControlPercentage(0.4f),
+	  FallingDeceleration(200.0f),
+	  OverTerminalSpeedFallingDeceleration(800.0f),
+	  TerminalMovementPlaneSpeed(1500.0f),
+	  bShouldClampTerminalVerticalSpeed(true),
+	  VerticalFallingDeceleration(4000.0f),
+	  TerminalVerticalSpeed(2000.0f)
+{
+	SharedSettingsClass = UCommonLegacyMovementSettings::StaticClass();
+}
+
+constexpr float VERTICAL_SLOPE_NORMAL_Z = 0.001f; // Slope is vertical if Abs(Normal.Z) <= this threshold. Accounts for precision problems that sometimes angle normals slightly off horizontal for vertical surface.
+
+void UFallingMode::OnGenerateMove(const FMoverTickStartData& StartState, const FMoverTimeStep& TimeStep, FProposedMove& OutProposedMove) const
+{
+	const FCharacterDefaultInputs* CharacterInputs = StartState.InputCmd.InputCollection.FindDataByType<FCharacterDefaultInputs>();
+	const FMoverDefaultSyncState* StartingSyncState = StartState.SyncState.SyncStateCollection.FindDataByType<FMoverDefaultSyncState>();
+	check(StartingSyncState);
+
+	const float DeltaSeconds = TimeStep.StepMs * 0.001f;
+
+	// We don't want velocity limits to take the falling velocity component into account, since it is handled 
+	//   separately by the terminal velocity of the environment.
+	const FVector StartVelocity = StartingSyncState->GetVelocity_WorldSpace();
+	const FVector StartHorizontalVelocity = FVector(StartVelocity.X, StartVelocity.Y, 0.f);
+
+	FFreeMoveParams Params;
+	if (CharacterInputs)
+	{
+		Params.MoveInputType = CharacterInputs->GetMoveInputType();
+		Params.MoveInput = CharacterInputs->GetMoveInput();
+	}
+	else
+	{
+		Params.MoveInputType = EMoveInputType::Invalid;
+		Params.MoveInput = FVector::ZeroVector;
+	}
+
+	FRotator IntendedOrientation_WorldSpace;
+	// If there's no intent from input to change orientation, use the current orientation
+	if (!CharacterInputs || CharacterInputs->OrientationIntent.IsNearlyZero())
+	{
+		IntendedOrientation_WorldSpace = StartingSyncState->GetOrientation_WorldSpace();
+	}
+	else
+	{
+		IntendedOrientation_WorldSpace = CharacterInputs->GetOrientationIntentDir_WorldSpace().ToOrientationRotator();
+	}
+
+	Params.OrientationIntent = IntendedOrientation_WorldSpace;
+	Params.PriorVelocity = StartHorizontalVelocity;
+	Params.PriorOrientation = StartingSyncState->GetOrientation_WorldSpace();
+	Params.DeltaSeconds = DeltaSeconds;
+	Params.TurningRate = CommonLegacySettings->TurningRate;
+	Params.TurningBoost = CommonLegacySettings->TurningBoost;
+	Params.MaxSpeed = CommonLegacySettings->MaxSpeed;
+	Params.Acceleration = CommonLegacySettings->Acceleration;
+	Params.Deceleration = FallingDeceleration;
+
+	Params.MoveInput *= AirControlPercentage;
+	// Don't care about Z axis input since falling - if Z input matters that should probably be a different movement mode
+	Params.MoveInput.Z = 0;
+
+	// Check if any current velocity values are over our terminal velocity - if so limit the move input in that direction and apply OverTerminalVelocityFallingDeceleration
+	if (Params.MoveInput.Dot(StartVelocity) > 0 && StartVelocity.Size2D() >= TerminalMovementPlaneSpeed)
+	{
+		const FPlane MovementNormalPlane(StartVelocity, StartVelocity.GetSafeNormal());
+		Params.MoveInput = Params.MoveInput.ProjectOnTo(MovementNormalPlane);
+
+		Params.Deceleration = OverTerminalSpeedFallingDeceleration;
+	}
+	
+	UMoverBlackboard* SimBlackboard = GetBlackboard_Mutable();
+	FFloorCheckResult LastFloorResult;
+	// limit our moveinput based on the floor we're on
+	if (SimBlackboard && SimBlackboard->TryGet(CommonBlackboard::LastFloorResult, LastFloorResult))
+	{
+		if (LastFloorResult.HitResult.IsValidBlockingHit() && LastFloorResult.HitResult.Normal.Z > VERTICAL_SLOPE_NORMAL_Z && !LastFloorResult.IsWalkableFloor())
+		{
+			// If acceleration is into the wall, limit contribution.
+			if (FVector::DotProduct(Params.MoveInput, LastFloorResult.HitResult.Normal) < 0.f)
+			{
+				// Allow movement parallel to the wall, but not into it because that may push us up.
+				const FVector Normal2D = LastFloorResult.HitResult.Normal.GetSafeNormal2D();
+				Params.MoveInput = FVector::VectorPlaneProject(Params.MoveInput, Normal2D);
+			}
+		}
+	}
+	
+	OutProposedMove = UAirMovementUtils::ComputeControlledFreeMove(Params);
+
+	//  If we are going faster than TerminalVerticalVelocity apply VerticalFallingDeceleration otherwise reset Z velocity to before we applied deceleration 
+	if (StartVelocity.GetAbs().Z > TerminalVerticalSpeed)
+	{
+		if (bShouldClampTerminalVerticalSpeed)
+		{
+			OutProposedMove.LinearVelocity.Z = FMath::Sign(StartVelocity.Z) * TerminalVerticalSpeed;
+		}
+		else
+		{
+			float DesiredDeceleration = FMath::Abs(TerminalVerticalSpeed - FMath::Abs(StartVelocity.Z)) / DeltaSeconds;
+			float DecelerationToApply = FMath::Min(DesiredDeceleration, VerticalFallingDeceleration);
+			DecelerationToApply = FMath::Sign(StartVelocity.Z) * DecelerationToApply * DeltaSeconds;
+			OutProposedMove.LinearVelocity.Z = StartVelocity.Z - DecelerationToApply;
+		}
+	}
+	else
+	{
+		OutProposedMove.LinearVelocity.Z = StartVelocity.Z;	// restore original vertical velocity
+	}
+
+	// Add velocity change due to gravity
+	OutProposedMove.LinearVelocity += UMovementUtils::ComputeVelocityFromGravity(GetMoverComponent()->GetGravityAcceleration(), DeltaSeconds);
+}
+
+void UFallingMode::OnSimulationTick(const FSimulationTickParams& Params, FMoverTickEndData& OutputState)
+{
+	const FMoverTickStartData& StartState = Params.StartState;
+	USceneComponent* UpdatedComponent = Params.UpdatedComponent;
+	UPrimitiveComponent* UpdatedPrimitive = Params.UpdatedPrimitive;
+	FProposedMove ProposedMove = Params.ProposedMove;
+
+	const FCharacterDefaultInputs* CharacterInputs = StartState.InputCmd.InputCollection.FindDataByType<FCharacterDefaultInputs>();
+	const FMoverDefaultSyncState* StartingSyncState = StartState.SyncState.SyncStateCollection.FindDataByType<FMoverDefaultSyncState>();
+	check(StartingSyncState);
+
+	FMoverDefaultSyncState& OutputSyncState = OutputState.SyncState.SyncStateCollection.FindOrAddMutableDataByType<FMoverDefaultSyncState>();
+
+	//const FVector GravityAccel = FVector(0.0f, 0.0f, -9800.0f);	// -9.8 m / sec^2
+	const float DeltaSeconds = Params.TimeStep.StepMs * 0.001f;
+	float PctTimeApplied = 0.f;
+
+	// Instantaneous movement changes that are executed and we exit before consuming any time
+	if (ProposedMove.bHasTargetLocation && AttemptTeleport(UpdatedComponent, ProposedMove.TargetLocation, UpdatedComponent->GetComponentRotation(), StartingSyncState->GetVelocity_WorldSpace(), OutputSyncState))
+	{
+		OutputState.MovementEndState.RemainingMs = Params.TimeStep.StepMs; 	// Give back all the time
+		return;
+	}
+
+	FMovementRecord MoveRecord;
+	MoveRecord.SetDeltaSeconds(DeltaSeconds);
+	
+	UMoverBlackboard* SimBlackboard = GetBlackboard_Mutable();
+
+	SimBlackboard->Invalidate(CommonBlackboard::LastFloorResult);	// falling = no valid floor
+	SimBlackboard->Invalidate(CommonBlackboard::LastMovementBase);
+
+	OutputSyncState.MoveDirectionIntent = (ProposedMove.bHasDirIntent ? ProposedMove.DirectionIntent : FVector::ZeroVector);
+
+
+	// Use the orientation intent directly. If no intent is provided, use last frame's orientation. Note that we are assuming rotation changes can't fail. 
+	FRotator TargetOrient = StartingSyncState->GetOrientation_WorldSpace();
+
+	// Apply orientation changes (if any)
+	if (!ProposedMove.AngularVelocity.IsZero())
+	{
+		TargetOrient += (ProposedMove.AngularVelocity * DeltaSeconds);
+	}
+	
+	const FVector PriorFallingVelocity = StartingSyncState->GetVelocity_WorldSpace();
+
+	//FVector MoveDelta = 0.5f * (PriorFallingVelocity + ProposedMove.LinearVelocity) * DeltaSeconds; 	// TODO: revive midpoint integration
+	FVector MoveDelta = ProposedMove.LinearVelocity * DeltaSeconds;
+
+	FHitResult Hit(1.f);
+	const FQuat TargetOrientQuat = TargetOrient.Quaternion();
+
+	UMovementUtils::TrySafeMoveUpdatedComponent(UpdatedComponent, UpdatedPrimitive, MoveDelta, TargetOrientQuat, true, Hit, ETeleportType::None, MoveRecord);
+
+	// Compute final velocity based on how long we actually go until we get a hit.
+	FVector NewFallingVelocity = StartingSyncState->GetVelocity_WorldSpace();
+	//NewFallingVelocity.Z = 0.5f * (PriorFallingVelocity.Z + (ProposedMove.LinearVelocity.Z * Hit.Time));	// TODO: revive midpoint integration
+	NewFallingVelocity.Z = ProposedMove.LinearVelocity.Z * Hit.Time;
+
+	FFloorCheckResult LandingFloor;
+
+	// Handle impact, whether it's a landing surface or something to slide on
+	if (Hit.IsValidBlockingHit() && UpdatedPrimitive)
+	{
+		float LastMoveTimeSlice = DeltaSeconds;
+		float SubTimeTickRemaining = LastMoveTimeSlice * (1.f - Hit.Time);
+
+		PctTimeApplied += Hit.Time * (1.f - PctTimeApplied);
+
+		// Check for hitting a landing surface
+		if (UAirMovementUtils::IsValidLandingSpot(UpdatedComponent, UpdatedPrimitive, UpdatedPrimitive->GetComponentLocation(), 
+			Hit, CommonLegacySettings->FloorSweepDistance, CommonLegacySettings->MaxWalkSlopeCosine, OUT LandingFloor))
+		{
+			CaptureFinalState(UpdatedComponent, *StartingSyncState, LandingFloor, DeltaSeconds, DeltaSeconds * PctTimeApplied, OutputSyncState, OutputState, MoveRecord);
+			return;
+		}
+		
+		LandingFloor.HitResult = Hit;
+		SimBlackboard->Set(CommonBlackboard::LastFloorResult, LandingFloor);
+
+		UMoverComponent* MoverComponent = GetMoverComponent();
+		FMoverOnImpactParams ImpactParams(DefaultModeNames::Falling, Hit, MoveDelta);
+		MoverComponent->HandleImpact(ImpactParams);
+
+		// We didn't land on a walkable surface, so let's try to slide along it
+		UAirMovementUtils::TryMoveToFallAlongSurface(UpdatedComponent, UpdatedPrimitive, MoverComponent, MoveDelta, 
+			(1.f - Hit.Time), TargetOrientQuat, Hit.Normal, Hit, true,
+			CommonLegacySettings->FloorSweepDistance, CommonLegacySettings->MaxWalkSlopeCosine, LandingFloor, MoveRecord);
+
+		PctTimeApplied += Hit.Time * (1.f - PctTimeApplied);
+
+		if (LandingFloor.IsWalkableFloor())
+		{
+			CaptureFinalState(UpdatedComponent, *StartingSyncState, LandingFloor, DeltaSeconds, DeltaSeconds * PctTimeApplied, OutputSyncState, OutputState, MoveRecord);
+			return;
+		}
+	}
+	else
+	{
+		// This indicates an unimpeded full move
+		PctTimeApplied = 1.f;
+	}
+	
+	CaptureFinalState(UpdatedComponent, *StartingSyncState, LandingFloor, DeltaSeconds, DeltaSeconds* PctTimeApplied, OutputSyncState, OutputState, MoveRecord);
+}
+
+
+void UFallingMode::OnRegistered(const FName ModeName)
+{
+	Super::OnRegistered(ModeName);
+
+	CommonLegacySettings = GetMoverComponent()->FindSharedSettings<UCommonLegacyMovementSettings>();
+	ensureMsgf(CommonLegacySettings, TEXT("Failed to find instance of CommonLegacyMovementSettings on %s. Movement may not function properly."), *GetPathNameSafe(this));
+}
+
+
+void UFallingMode::OnUnregistered()
+{
+	CommonLegacySettings = nullptr;
+
+	Super::OnUnregistered();
+}
+
+
+bool UFallingMode::AttemptTeleport(USceneComponent* UpdatedComponent, const FVector& TeleportPos, const FRotator& TeleportRot, const FVector& PriorVelocity, FMoverDefaultSyncState& OutputSyncState)
+{
+	if (UpdatedComponent->GetOwner()->TeleportTo(TeleportPos, TeleportRot))
+	{
+		OutputSyncState.SetTransforms_WorldSpace( UpdatedComponent->GetComponentLocation(),
+												  UpdatedComponent->GetComponentRotation(),
+												  PriorVelocity,
+												  nullptr); // no movement base
+
+		UpdatedComponent->ComponentVelocity = PriorVelocity;
+		return true;
+	}
+
+	return false;
+}
+
+void UFallingMode::ProcessLanded(const FFloorCheckResult& FloorResult, FVector& Velocity, FRelativeBaseInfo& BaseInfo, FMoverTickEndData& TickEndData) const
+{
+	UMoverBlackboard* SimBlackboard = GetBlackboard_Mutable();
+
+	FName NextMovementMode = NAME_None; 
+	// if we can walk on the floor we landed on
+	if (FloorResult.IsWalkableFloor())
+	{
+		// Transfer to LandingMovementMode (usually walking), and cache any floor / movement base info
+		Velocity.Z = 0.0;
+		NextMovementMode = CommonLegacySettings->GroundMovementModeName;
+
+		SimBlackboard->Set(CommonBlackboard::LastFloorResult, FloorResult);
+
+		if (UBasedMovementUtils::IsADynamicBase(FloorResult.HitResult.GetComponent()))
+		{
+			BaseInfo.SetFromFloorResult(FloorResult);
+		}
+	}
+	// we could check for other surfaces here (i.e. when swimming is implemented we can check the floor hit here and see if we need to go into swimming)
+
+	// This would also be a good spot for implementing some falling physics interactions (i.e. falling into a movable object and pushing it based off of this actors velocity)
+	
+	// if a new mode was set go ahead and switch to it after this tick and broadcast we landed
+	if (!NextMovementMode.IsNone())
+	{
+		TickEndData.MovementEndState.NextModeName = NextMovementMode;
+		OnLanded.Broadcast(NextMovementMode, FloorResult.HitResult);
+	}
+}
+
+void UFallingMode::CaptureFinalState(USceneComponent* UpdatedComponent, const FMoverDefaultSyncState& StartSyncState, const FFloorCheckResult& FloorResult, float DeltaSeconds, float DeltaSecondsUsed, FMoverDefaultSyncState& OutputSyncState, FMoverTickEndData& TickEndData, FMovementRecord& Record) const
+{
+	UMoverBlackboard* SimBlackboard = GetBlackboard_Mutable();
+
+	const FVector FinalLocation = UpdatedComponent->GetComponentLocation();
+
+	// Check for time refunds
+	constexpr float MinRemainingSecondsToRefund = 0.0001f;	// If we have this amount of time (or more) remaining, give it to the next simulation step.
+
+	if ((DeltaSeconds - DeltaSecondsUsed) >= MinRemainingSecondsToRefund)
+	{
+		const float PctOfTimeRemaining = (1.0f - (DeltaSecondsUsed / DeltaSeconds));
+		TickEndData.MovementEndState.RemainingMs = PctOfTimeRemaining * DeltaSeconds * 1000.f;
+	}
+	else
+	{
+		TickEndData.MovementEndState.RemainingMs = 0.f;
+	}
+	
+	Record.SetDeltaSeconds( DeltaSecondsUsed );
+	
+	FVector EffectiveVelocity = Record.GetRelevantVelocity();
+	// TODO: Update Main/large movement record with substeps from our local record
+
+	FRelativeBaseInfo MovementBaseInfo;
+	ProcessLanded(FloorResult, EffectiveVelocity, MovementBaseInfo, TickEndData);
+
+	if (MovementBaseInfo.HasRelativeInfo())
+	{
+		SimBlackboard->Set(CommonBlackboard::LastMovementBase, MovementBaseInfo);
+
+		OutputSyncState.SetTransforms_WorldSpace( FinalLocation,
+												  UpdatedComponent->GetComponentRotation(),
+												  EffectiveVelocity,
+												  MovementBaseInfo.MovementBase, MovementBaseInfo.BoneName);
+	}
+	else
+	{
+		OutputSyncState.SetTransforms_WorldSpace( FinalLocation,
+												  UpdatedComponent->GetComponentRotation(),
+												  EffectiveVelocity,
+												  nullptr); // no movement base
+	}
+
+	UpdatedComponent->ComponentVelocity = EffectiveVelocity;
+}
diff --git a/Engine/Plugins/Experimental/Mover/Source/Mover/Private/DefaultMovementSet/Modes/FlyingMode.cpp b/Engine/Plugins/Experimental/Mover/Source/Mover/Private/DefaultMovementSet/Modes/FlyingMode.cpp
new file mode 100644
--- /dev/null
+++ b/Engine/Plugins/Experimental/Mover/Source/Mover/Private/DefaultMovementSet/Modes/FlyingMode.cpp
@@ -0,0 +1,180 @@
+// Copyright Epic Games, Inc. All Rights Reserved.
+
+#include "DefaultMovementSet/Modes/FlyingMode.h"
+#include "MoveLibrary/AirMovementUtils.h"
+#include "MoveLibrary/MovementUtils.h"
+#include "MoverComponent.h"
+#include "DefaultMovementSet/Settings/CommonLegacyMovementSettings.h"
+
+#include UE_INLINE_GENERATED_CPP_BY_NAME(FlyingMode)
+
+
+UFlyingMode::UFlyingMode(const FObjectInitializer& ObjectInitializer)
+	: Super(ObjectInitializer)
+{
+}
+
+void UFlyingMode::OnGenerateMove(const FMoverTickStartData& StartState, const FMoverTimeStep& TimeStep, FProposedMove& OutProposedMove) const
+{
+	const FCharacterDefaultInputs* CharacterInputs = StartState.InputCmd.InputCollection.FindDataByType<FCharacterDefaultInputs>();
+	const FMoverDefaultSyncState* StartingSyncState = StartState.SyncState.SyncStateCollection.FindDataByType<FMoverDefaultSyncState>();
+	check(StartingSyncState);
+
+	const float DeltaSeconds = TimeStep.StepMs * 0.001f;
+
+	FFreeMoveParams Params;
+	if (CharacterInputs)
+	{
+		Params.MoveInputType = CharacterInputs->GetMoveInputType();
+		Params.MoveInput = CharacterInputs->GetMoveInput();
+	}
+	else
+	{
+		Params.MoveInputType = EMoveInputType::Invalid;
+		Params.MoveInput = FVector::ZeroVector;
+	}
+
+	FRotator IntendedOrientation_WorldSpace;
+	// If there's no intent from input to change orientation, use the current orientation
+	if (!CharacterInputs || CharacterInputs->OrientationIntent.IsNearlyZero())
+	{
+		IntendedOrientation_WorldSpace = StartingSyncState->GetOrientation_WorldSpace();
+	}
+	else
+	{
+		IntendedOrientation_WorldSpace = CharacterInputs->GetOrientationIntentDir_WorldSpace().ToOrientationRotator();
+	}
+	
+	Params.OrientationIntent = IntendedOrientation_WorldSpace;
+	Params.PriorVelocity = StartingSyncState->GetVelocity_WorldSpace();
+	Params.PriorOrientation = StartingSyncState->GetOrientation_WorldSpace();
+	Params.TurningRate = CommonLegacySettings->TurningRate;
+	Params.TurningBoost = CommonLegacySettings->TurningBoost;
+	Params.MaxSpeed = CommonLegacySettings->MaxSpeed;
+	Params.Acceleration = CommonLegacySettings->Acceleration;
+	Params.Deceleration = CommonLegacySettings->Deceleration;
+	Params.DeltaSeconds = DeltaSeconds;
+	
+	OutProposedMove = UAirMovementUtils::ComputeControlledFreeMove(Params);
+}
+
+void UFlyingMode::OnSimulationTick(const FSimulationTickParams& Params, FMoverTickEndData& OutputState)
+{
+	const FMoverTickStartData& StartState = Params.StartState;
+	USceneComponent* UpdatedComponent = Params.UpdatedComponent;
+	UPrimitiveComponent* UpdatedPrimitive = Params.UpdatedPrimitive;
+	FProposedMove ProposedMove = Params.ProposedMove;
+
+	const FCharacterDefaultInputs* CharacterInputs = StartState.InputCmd.InputCollection.FindDataByType<FCharacterDefaultInputs>();
+	const FMoverDefaultSyncState* StartingSyncState = StartState.SyncState.SyncStateCollection.FindDataByType<FMoverDefaultSyncState>();
+	check(StartingSyncState);
+
+	FMoverDefaultSyncState& OutputSyncState = OutputState.SyncState.SyncStateCollection.FindOrAddMutableDataByType<FMoverDefaultSyncState>();
+
+	const float DeltaSeconds = Params.TimeStep.StepMs * 0.001f;
+
+	// Instantaneous movement changes that are executed and we exit before consuming any time
+	if (ProposedMove.bHasTargetLocation && AttemptTeleport(UpdatedComponent, ProposedMove.TargetLocation, UpdatedComponent->GetComponentRotation(), *StartingSyncState, OutputState))
+	{
+		OutputState.MovementEndState.RemainingMs = Params.TimeStep.StepMs; 	// Give back all the time
+		return;
+	}
+
+	FMovementRecord MoveRecord;
+	MoveRecord.SetDeltaSeconds(DeltaSeconds);
+
+	UMoverBlackboard* SimBlackboard = GetBlackboard_Mutable();
+
+	SimBlackboard->Invalidate(CommonBlackboard::LastFloorResult);	// flying = no valid floor
+	SimBlackboard->Invalidate(CommonBlackboard::LastMovementBase);
+
+	OutputSyncState.MoveDirectionIntent = (ProposedMove.bHasDirIntent ? ProposedMove.DirectionIntent : FVector::ZeroVector);
+
+	// Use the orientation intent directly. If no intent is provided, use last frame's orientation. Note that we are assuming rotation changes can't fail. 
+	const FRotator StartingOrient = StartingSyncState->GetOrientation_WorldSpace();
+	FRotator TargetOrient = StartingOrient;
+
+	bool bIsOrientationChanging = false;
+
+	// Apply orientation changes (if any)
+	if (!ProposedMove.AngularVelocity.IsZero())
+	{
+		TargetOrient += (ProposedMove.AngularVelocity * DeltaSeconds);
+		bIsOrientationChanging = (TargetOrient != StartingOrient);
+	}
+	
+	FVector MoveDelta = ProposedMove.LinearVelocity * DeltaSeconds;
+	const FQuat OrientQuat = TargetOrient.Quaternion();
+	FHitResult Hit(1.f);
+
+	if (!MoveDelta.IsNearlyZero() || bIsOrientationChanging)
+	{
+		UMovementUtils::TrySafeMoveUpdatedComponent(UpdatedComponent, UpdatedPrimitive, MoveDelta, OrientQuat, true, Hit, ETeleportType::None, MoveRecord);
+	}
+
+	if (Hit.IsValidBlockingHit())
+	{
+		UMoverComponent* MoverComponent = GetMoverComponent();
+		FMoverOnImpactParams ImpactParams(DefaultModeNames::Flying, Hit, MoveDelta);
+		MoverComponent->HandleImpact(ImpactParams);
+		// Try to slide the remaining distance along the surface.
+		UMovementUtils::TryMoveToSlideAlongSurface(UpdatedComponent, UpdatedPrimitive, MoverComponent, MoveDelta, 1.f - Hit.Time, OrientQuat, Hit.Normal, Hit, true, MoveRecord);
+	}
+
+	CaptureFinalState(UpdatedComponent, MoveRecord, *StartingSyncState, OutputSyncState, DeltaSeconds);
+}
+
+
+bool UFlyingMode::AttemptTeleport(USceneComponent* UpdatedComponent, const FVector& TeleportPos, const FRotator& TeleportRot, const FMoverDefaultSyncState& StartingSyncState, FMoverTickEndData& Output)
+{
+	if (UpdatedComponent->GetOwner()->TeleportTo(TeleportPos, TeleportRot))
+	{
+		FMoverDefaultSyncState& OutputSyncState = Output.SyncState.SyncStateCollection.FindOrAddMutableDataByType<FMoverDefaultSyncState>();
+
+		OutputSyncState.SetTransforms_WorldSpace( UpdatedComponent->GetComponentLocation(),
+												  UpdatedComponent->GetComponentRotation(),
+												  StartingSyncState.GetVelocity_WorldSpace(),
+												  nullptr); // no movement base
+
+		UpdatedComponent->ComponentVelocity = StartingSyncState.GetVelocity_WorldSpace();
+
+		GetBlackboard_Mutable()->Invalidate(CommonBlackboard::LastFloorResult);
+
+		return true;
+	}
+
+	return false;
+}
+
+
+// TODO: replace this function with simply looking at/collapsing the MovementRecord
+void UFlyingMode::CaptureFinalState(USceneComponent* UpdatedComponent, FMovementRecord& Record, const FMoverDefaultSyncState& StartSyncState, FMoverDefaultSyncState& OutputSyncState, const float DeltaSeconds) const
+{
+	const FVector FinalLocation = UpdatedComponent->GetComponentLocation();
+	const FVector FinalVelocity = Record.GetRelevantVelocity();
+	
+	// TODO: Update Main/large movement record with substeps from our local record
+
+	OutputSyncState.SetTransforms_WorldSpace(FinalLocation,
+											  UpdatedComponent->GetComponentRotation(),
+											  FinalVelocity,
+											  nullptr); // no movement base
+
+	UpdatedComponent->ComponentVelocity = FinalVelocity;
+}
+
+void UFlyingMode::OnRegistered(const FName ModeName)
+{
+	Super::OnRegistered(ModeName);
+
+	CommonLegacySettings = GetMoverComponent()->FindSharedSettings<UCommonLegacyMovementSettings>();
+	ensureMsgf(CommonLegacySettings, TEXT("Failed to find instance of CommonLegacyMovementSettings on %s. Movement may not function properly."), *GetPathNameSafe(this));
+}
+
+
+void UFlyingMode::OnUnregistered()
+{
+	CommonLegacySettings = nullptr;
+
+	Super::OnUnregistered();
+}
\ No newline at end of file
diff --git a/Engine/Plugins/Experimental/Mover/Source/Mover/Private/DefaultMovementSet/Modes/SwimmingMode.cpp b/Engine/Plugins/Experimental/Mover/Source/Mover/Private/DefaultMovementSet/Modes/SwimmingMode.cpp
new file mode 100644
--- /dev/null
+++ b/Engine/Plugins/Experimental/Mover/Source/Mover/Private/DefaultMovementSet/Modes/SwimmingMode.cpp
@@ -0,0 +1,216 @@
+// Copyright Epic Games, Inc. All Rights Reserved.
+
+#include "DefaultMovementSet/Modes/SwimmingMode.h"
+#include "Components/SkeletalMeshComponent.h"
+#include "GameFramework/Pawn.h"
+#include "MoveLibrary/WaterMovementUtils.h"
+#include "DefaultMovementSet/LayeredMoves/BasicLayeredMoves.h"
+#include "MoverComponent.h"
+#include "DefaultMovementSet/Settings/CommonLegacyMovementSettings.h"
+#include "MoverLog.h"
+
+#include UE_INLINE_GENERATED_CPP_BY_NAME(SwimmingMode)
+
+
+USwimmingMode::USwimmingMode(const FObjectInitializer& ObjectInitializer)
+	: Super(ObjectInitializer)
+{
+	SharedSettingsClass = UCommonLegacyMovementSettings::StaticClass();
+}
+
+void USwimmingMode::OnGenerateMove(const FMoverTickStartData& StartState, const FMoverTimeStep& TimeStep, FProposedMove& OutProposedMove) const
+{
+	const UMoverComponent* MoverComp = GetMoverComponent();
+	const FCharacterDefaultInputs* CharacterInputs = StartState.InputCmd.InputCollection.FindDataByType<FCharacterDefaultInputs>();
+	const FMoverDefaultSyncState* StartingSyncState = StartState.SyncState.SyncStateCollection.FindDataByType<FMoverDefaultSyncState>();
+	check(StartingSyncState);
+	
+	const float DeltaSeconds = TimeStep.StepMs * 0.001f;
+	
+	UMoverBlackboard* SimBlackboard = GetBlackboard_Mutable();
+	FWaterCheckResult LastWaterResult;
+	
+	const float CapsuleHalfHeight = MoverComp->GetOwner()->GetSimpleCollisionHalfHeight();
+
+	if (SimBlackboard && SimBlackboard->TryGet(CommonBlackboard::LastWaterResult, LastWaterResult) && LastWaterResult.IsSwimmableVolume())
+	{
+		FUpdateWaterSplineDataParams Params;
+		Params.TargetImmersionDepth = CommonLegacySettings->SwimmingIdealImmersionDepth;
+		Params.WaterVelocityDepthForMax = SurfaceSwimmingWaterControlSettings.WaterVelocityDepthForMax;
+		Params.WaterVelocityMinMultiplier = SurfaceSwimmingWaterControlSettings.WaterVelocityMinMultiplier;
+		Params.PlayerVelocity = StartingSyncState->GetVelocity_WorldSpace();
+		Params.CapsuleHalfHeight = CapsuleHalfHeight;
+		Params.PlayerLocation = StartingSyncState->GetLocation_WorldSpace();
+		
+		UWaterMovementUtils::UpdateWaterSplineData(Params, LastWaterResult);
+	}
+	
+	FVector Velocity = StartingSyncState->GetVelocity_WorldSpace();
+	
+	// Buoyancy Bobbing
+	{
+		const FWaterFlowSplineData& WaterData = LastWaterResult.WaterSplineData;
+		
+		const float ReciprocalCapsuleHeight = 0.5f / CapsuleHalfHeight;
+		const FVector CapsuleLocation = StartingSyncState->GetLocation_WorldSpace();
+
+		const float ImmersionDepth = WaterData.ImmersionDepth + CapsuleHalfHeight;
+		const float ImmersionPercent = FMath::Clamp(ImmersionDepth * ReciprocalCapsuleHeight, UE_KINDA_SMALL_NUMBER, 1.f);
+		const float IdealDepth = CommonLegacySettings->SwimmingIdealImmersionDepth;
+		const float IdealImmersionPercent = FMath::Clamp(IdealDepth * ReciprocalCapsuleHeight, UE_KINDA_SMALL_NUMBER, 1.f);
+
+		const float GravityForce = GetMoverComponent()->GetGravityAcceleration().Z;
+
+		
+		// Using IdealDepth, we compute the Buoyancy such that (Buoyancy * IdealImmersionPercent) = -G force
+		const float BuoyancyForce = -GravityForce / IdealImmersionPercent;
+	
+		// Now calculate the actual force in the correct ratio
+		float BobbingForce = (BuoyancyForce * ImmersionPercent) + GravityForce;
+		const float MaxForce = SurfaceSwimmingWaterControlSettings.BobbingMaxForce;
+		BobbingForce = FMath::Clamp(BobbingForce, -MaxForce, MaxForce);
+
+		Velocity.Z += BobbingForce * DeltaSeconds;
+
+		// Vertical fluid friction for bobbing
+		if (!FMath::IsNearlyZero(Velocity.Z, (FVector::FReal)0.1f))
+		{
+			const float IdealDepthTolerance = SurfaceSwimmingWaterControlSettings.BobbingIdealDepthTolerance;
+			if (FMath::Sign(Velocity.Z) != FMath::Sign(BobbingForce) || FMath::IsNearlyEqual(ImmersionDepth, IdealDepth, IdealDepthTolerance))
+			{
+				float FluidFriction = 0.f;
+				float ExpDrag = 0.f;
+				const bool bFullySubmerged = ImmersionPercent > 1.f;
+				if (Velocity.Z > 0)
+				{
+					FluidFriction = SurfaceSwimmingWaterControlSettings.BobbingFrictionUp;
+					ExpDrag = SurfaceSwimmingWaterControlSettings.BobbingExpDragUp;
+				}
+				else
+				{
+					// Different drag when fully immersed and moving down (mainly controls how far you go when falling in fast)
+					FluidFriction = bFullySubmerged ? SurfaceSwimmingWaterControlSettings.BobbingFrictionDownSubmerged : SurfaceSwimmingWaterControlSettings.BobbingFrictionDown;
+					ExpDrag = bFullySubmerged ? SurfaceSwimmingWaterControlSettings.BobbingExpDragDownSubmerged : SurfaceSwimmingWaterControlSettings.BobbingExpDragDown;
+				}
+	
+				FluidFriction *= SurfaceSwimmingWaterControlSettings.BobbingFrictionMultiplier;
+				ExpDrag *= SurfaceSwimmingWaterControlSettings.BobbingExpDragMultiplier;
+	
+				Velocity.Z = Velocity.Z * (1.f - FMath::Min(FluidFriction * DeltaSeconds, 1.f));
+				Velocity.Z = Velocity.Z * (1.f - FMath::Min(FMath::Abs(Velocity.Z) * FMath::Square(ExpDrag) * DeltaSeconds, 1.f));
+			}
+		}
+	}
+
+	// Vertical speed limit in Water
+	{
+		const float MaxVerticalWaterSpeedUp = SurfaceSwimmingWaterControlSettings.MaxSpeedUp;
+		const float MaxVerticalWaterSpeedDown = SurfaceSwimmingWaterControlSettings.MaxSpeedDown;
+	
+		Velocity.Z = FMath::Clamp(Velocity.Z, -FMath::Abs(MaxVerticalWaterSpeedDown), FMath::Abs(MaxVerticalWaterSpeedUp));
+	}
+
+	// Calculate and apply the requested move here
+	{
+		// Force from Water flow's WaterVelocity 
+		const float MaxWaterForce = SurfaceSwimmingWaterControlSettings.MaxWaterForce;
+		const float WaterForceMultiplier = SurfaceSwimmingWaterControlSettings.WaterForceMultiplier * SurfaceSwimmingWaterControlSettings.WaterForceSecondMultiplier;
+		const FVector WaterVelocity = LastWaterResult.WaterSplineData.WaterVelocity;
+		const FVector WaterAcceleration = (WaterVelocity * WaterForceMultiplier).GetClampedToMaxSize(MaxWaterForce);
+		const float WaterSpeed = WaterVelocity.Size();
+	
+	    // Consider player input
+		FRotator IntendedOrientation_WorldSpace;
+		if (!CharacterInputs || CharacterInputs->OrientationIntent.IsNearlyZero())
+		{
+			IntendedOrientation_WorldSpace = StartingSyncState->GetOrientation_WorldSpace();
+		}
+		else
+		{
+			IntendedOrientation_WorldSpace = CharacterInputs->GetOrientationIntentDir_WorldSpace().ToOrientationRotator();
+		}
+		
+		FWaterMoveParams Params;
+		if (CharacterInputs)
+		{
+			Params.MoveInputType = CharacterInputs->GetMoveInputType();
+			Params.MoveInput = CharacterInputs->GetMoveInput_WorldSpace();
+		}
+		else
+		{
+			Params.MoveInputType = EMoveInputType::Invalid;
+			Params.MoveInput = FVector::ZeroVector;
+		}
+		Params.OrientationIntent = IntendedOrientation_WorldSpace;
+		Params.PriorVelocity = StartingSyncState->GetVelocity_WorldSpace();
+		Params.PriorOrientation = StartingSyncState->GetOrientation_WorldSpace();
+		Params.TurningRate = CommonLegacySettings->TurningRate;
+		Params.TurningBoost = CommonLegacySettings->TurningBoost;
+		Params.MaxSpeed = CommonLegacySettings->MaxSpeed;
+		Params.Acceleration = CommonLegacySettings->Acceleration;
+		Params.Deceleration = CommonLegacySettings->Deceleration;
+		Params.DeltaSeconds = DeltaSeconds;
+		Params.MoveSpeed = WaterSpeed;
+		Params.MoveAcceleration = WaterAcceleration;
+
+		// Calculate the move
+		OutProposedMove = UWaterMovementUtils::ComputeControlledWaterMove(Params);
+		
+		// Use Z Velocity calculated earlier (Buoyancy, Friction and Terminal Velocity) for the move's Z component
+		OutProposedMove.LinearVelocity.Z = Velocity.Z;
+	}
+}
+
+void USwimmingMode::OnSimulationTick(const FSimulationTickParams& Params, FMoverTickEndData& OutputState)
+{
+	// TODO Har
+}
+
+void USwimmingMode::OnRegistered(const FName ModeName)
+{
+	Super::OnRegistered(ModeName);
+
+	CommonLegacySettings = GetMoverComponent()->FindSharedSettings<UCommonLegacyMovementSettings>();
+	ensureMsgf(CommonLegacySettings, TEXT("Failed to find instance of CommonLegacyMovementSettings on %s. Movement may not function properly."), *GetPathNameSafe(this));
+}
+
+void USwimmingMode::OnUnregistered()
+{
+	CommonLegacySettings = nullptr;
+
+	Super::OnUnregistered();
+}
+
+bool USwimmingMode::AttemptJump(float UpwardsSpeed, FMoverTickEndData& OutputState)
+{
+	FMoverDefaultSyncState& OutputSyncState = OutputState.SyncState.SyncStateCollection.FindOrAddMutableDataByType<FMoverDefaultSyncState>();
+
+	// TODO: This should check if a jump is even allowed
+	TSharedPtr<FLayeredMove_JumpImpulse> JumpMove = MakeShared<FLayeredMove_JumpImpulse>();
+	JumpMove->UpwardsSpeed = UpwardsSpeed;
+	OutputSyncState.LayeredMoves.QueueLayeredMove(JumpMove);
+	OutputState.MovementEndState.NextModeName = DefaultModeNames::Falling;
+	return true;
+}
+
+bool USwimmingMode::AttemptTeleport(USceneComponent* UpdatedComponent, const FVector& TeleportPos, const FRotator& TeleportRot, const FVector& PriorVelocity, FMoverTickEndData& Output)
+{
+	if (UpdatedComponent->GetOwner()->TeleportTo(TeleportPos, TeleportRot))
+	{
+		FMoverDefaultSyncState& OutputSyncState = Output.SyncState.SyncStateCollection.FindOrAddMutableDataByType<FMoverDefaultSyncState>();
+
+		OutputSyncState.SetTransforms_WorldSpace( UpdatedComponent->GetComponentLocation(),
+												  UpdatedComponent->GetComponentRotation(),
+												  PriorVelocity,
+												  nullptr ); // no movement base
+		
+		// TODO: instead of invalidating it, consider checking for a floor. Possibly a dynamic base?
+		GetBlackboard_Mutable()->Invalidate(CommonBlackboard::LastFloorResult);
+		GetBlackboard_Mutable()->Invalidate(CommonBlackboard::LastWaterResult);
+		GetBlackboard_Mutable()->Invalidate(CommonBlackboard::LastMovementBase);
+
+		return true;
+	}
+
+	return false;
+}
\ No newline at end of file
diff --git a/Engine/Plugins/Experimental/Mover/Source/Mover/Private/DefaultMovementSet/Modes/WalkingMode.cpp b/Engine/Plugins/Experimental/Mover/Source/Mover/Private/DefaultMovementSet/Modes/WalkingMode.cpp
new file mode 100644
--- /dev/null
+++ b/Engine/Plugins/Experimental/Mover/Source/Mover/Private/DefaultMovementSet/Modes/WalkingMode.cpp
@@ -0,0 +1,411 @@
+// Copyright Epic Games, Inc. All Rights Reserved.
+
+#include "DefaultMovementSet/Modes/WalkingMode.h"
+#include "Components/SkeletalMeshComponent.h"
+#include "GameFramework/Pawn.h"
+#include "MoveLibrary/MovementUtils.h"
+#include "MoveLibrary/BasedMovementUtils.h"
+#include "MoveLibrary/ModularMovement.h"
+#include "MoveLibrary/FloorQueryUtils.h"
+#include "MoveLibrary/GroundMovementUtils.h"
+#include "DefaultMovementSet/LayeredMoves/BasicLayeredMoves.h"
+#include "MoverComponent.h"
+#include "DefaultMovementSet/Settings/CommonLegacyMovementSettings.h"
+#include "MoverLog.h"
+
+#include UE_INLINE_GENERATED_CPP_BY_NAME(WalkingMode)
+
+
+UWalkingMode::UWalkingMode(const FObjectInitializer& ObjectInitializer)
+	: Super(ObjectInitializer)
+{
+	SharedSettingsClass = UCommonLegacyMovementSettings::StaticClass();
+}
+
+void UWalkingMode::OnGenerateMove(const FMoverTickStartData& StartState, const FMoverTimeStep& TimeStep, FProposedMove& OutProposedMove) const
+{
+	const UMoverComponent* MoverComp = GetMoverComponent();
+	const FCharacterDefaultInputs* CharacterInputs = StartState.InputCmd.InputCollection.FindDataByType<FCharacterDefaultInputs>();
+	const FMoverDefaultSyncState* StartingSyncState = StartState.SyncState.SyncStateCollection.FindDataByType<FMoverDefaultSyncState>();
+	check(StartingSyncState);
+
+
+	const float DeltaSeconds = TimeStep.StepMs * 0.001f;
+	FFloorCheckResult LastFloorResult;
+	FVector MovementNormal;
+
+	UMoverBlackboard* SimBlackboard = GetBlackboard_Mutable();
+
+	// Try to use the floor as the basis for the intended move direction (i.e. try to walk along slopes, rather than into them)
+	if (SimBlackboard && SimBlackboard->TryGet(CommonBlackboard::LastFloorResult, LastFloorResult) && LastFloorResult.IsWalkableFloor())
+	{
+		MovementNormal = LastFloorResult.HitResult.ImpactNormal;
+	}
+	else
+	{
+		MovementNormal = MoverComp->GetUpDirection();
+	}
+
+	
+
+	FRotator IntendedOrientation_WorldSpace;
+	// If there's no intent from input to change orientation, use the current orientation
+	if (!CharacterInputs || CharacterInputs->OrientationIntent.IsNearlyZero())
+	{
+		IntendedOrientation_WorldSpace = StartingSyncState->GetOrientation_WorldSpace();
+	}
+	else
+	{
+		IntendedOrientation_WorldSpace = CharacterInputs->GetOrientationIntentDir_WorldSpace().ToOrientationRotator();
+	}
+
+	FGroundMoveParams Params;
+
+	if (CharacterInputs)
+	{
+		Params.MoveInputType = CharacterInputs->GetMoveInputType();
+		Params.MoveInput = CharacterInputs->GetMoveInput_WorldSpace();
+	}
+	else
+	{
+		Params.MoveInputType = EMoveInputType::Invalid;
+		Params.MoveInput = FVector::ZeroVector;
+	}
+
+	Params.OrientationIntent = IntendedOrientation_WorldSpace;
+	Params.PriorVelocity = StartingSyncState->GetVelocity_WorldSpace();
+	Params.PriorOrientation = StartingSyncState->GetOrientation_WorldSpace();
+	Params.GroundNormal = MovementNormal;
+	Params.TurningRate = CommonLegacySettings->TurningRate;
+	Params.TurningBoost = CommonLegacySettings->TurningBoost;
+	Params.MaxSpeed = CommonLegacySettings->MaxSpeed;
+	Params.Acceleration = CommonLegacySettings->Acceleration;
+	Params.Deceleration = CommonLegacySettings->Deceleration;
+	Params.DeltaSeconds = DeltaSeconds;
+	
+	if (Params.MoveInput.SizeSquared() > 0.f && !UMovementUtils::IsExceedingMaxSpeed(Params.PriorVelocity, CommonLegacySettings->MaxSpeed))
+	{
+		Params.Friction = CommonLegacySettings->GroundFriction;
+	}
+	else
+	{
+		Params.Friction = CommonLegacySettings->bUseSeparateBrakingFriction ? CommonLegacySettings->BrakingFriction : CommonLegacySettings->GroundFriction;
+		Params.Friction *= CommonLegacySettings->BrakingFrictionFactor;
+	}
+
+	OutProposedMove = UGroundMovementUtils::ComputeControlledGroundMove(Params);
+
+	if (TurnGenerator)
+	{
+		OutProposedMove.AngularVelocity = ITurnGeneratorInterface::Execute_GetTurn(TurnGenerator, IntendedOrientation_WorldSpace, StartState, *StartingSyncState, TimeStep, OutProposedMove, SimBlackboard);
+	}
+}
+
+void UWalkingMode::OnSimulationTick(const FSimulationTickParams& Params, FMoverTickEndData& OutputState)
+{
+	const UMoverComponent* MoverComp = GetMoverComponent();
+	const FMoverTickStartData& StartState = Params.StartState;
+	USceneComponent* UpdatedComponent = Params.UpdatedComponent;
+	UPrimitiveComponent* UpdatedPrimitive = Params.UpdatedPrimitive;
+	FProposedMove ProposedMove = Params.ProposedMove;
+
+	if (!UpdatedComponent || !UpdatedPrimitive)
+	{
+		return;
+	}
+
+	const FCharacterDefaultInputs* CharacterInputs = StartState.InputCmd.InputCollection.FindDataByType<FCharacterDefaultInputs>();
+	const FMoverDefaultSyncState* StartingSyncState = StartState.SyncState.SyncStateCollection.FindDataByType<FMoverDefaultSyncState>();
+	check(StartingSyncState);
+
+	FMoverDefaultSyncState& OutputSyncState = OutputState.SyncState.SyncStateCollection.FindOrAddMutableDataByType<FMoverDefaultSyncState>();
+
+
+	const float DeltaSeconds = Params.TimeStep.StepMs * 0.001f;
+
+	// Instantaneous movement changes that are executed and we exit before consuming any time
+	if ( (ProposedMove.bHasTargetLocation && AttemptTeleport(UpdatedComponent, ProposedMove.TargetLocation, UpdatedComponent->GetComponentRotation(), StartingSyncState->GetVelocity_WorldSpace(), OutputState)) ||	// Teleport
+		 (CharacterInputs && CharacterInputs->bIsJumpJustPressed && AttemptJump(CommonLegacySettings->JumpUpwardsSpeed, OutputState)) )	// Jump
+	{
+		UpdatedComponent->ComponentVelocity = StartingSyncState->GetVelocity_WorldSpace();
+		OutputState.MovementEndState.RemainingMs = Params.TimeStep.StepMs; 	// Give back all the time
+		return;
+	}
+
+	const FVector OrigMoveDelta = ProposedMove.LinearVelocity * DeltaSeconds;
+
+	TObjectPtr<AActor> OwnerActor = UpdatedComponent->GetOwner();
+	check(OwnerActor);
+	TObjectPtr<USkeletalMeshComponent> OwnerSkeletalMesh = OwnerActor->FindComponentByClass<USkeletalMeshComponent>();
+
+	FMovementRecord MoveRecord;
+	MoveRecord.SetDeltaSeconds(DeltaSeconds);
+
+	FFloorCheckResult CurrentFloor;
+	FRelativeBaseInfo OldRelativeBase;	// last known movement base (typically from a prior frame or prior substep)
+
+	UMoverBlackboard* SimBlackboard = GetBlackboard_Mutable();
+
+	// If we don't have cached floor information, we need to search for it again
+	if (!SimBlackboard->TryGet(CommonBlackboard::LastFloorResult, CurrentFloor))
+	{
+		UFloorQueryUtils::FindFloor(UpdatedComponent, UpdatedPrimitive,
+			CommonLegacySettings->FloorSweepDistance, CommonLegacySettings->MaxWalkSlopeCosine,
+			UpdatedPrimitive->GetComponentLocation(), CurrentFloor);
+	}
+
+	if (!SimBlackboard->TryGet(CommonBlackboard::LastMovementBase, OldRelativeBase))
+	{
+		OldRelativeBase = UpdateFloorAndBaseInfo(CurrentFloor);
+	}
+	
+	
+	// If we're on a dynamic movement base, attempt to move along with whatever motion is has performed since we last ticked
+	if (OldRelativeBase.UsesSameBase(StartingSyncState->GetMovementBase(), StartingSyncState->GetMovementBaseBoneName()))
+	{
+		FHitResult BasedMoveHitResult;
+		bool bDidMoveAlongWithBase = UBasedMovementUtils::TryMoveToStayWithBase(UpdatedComponent, UpdatedPrimitive, OldRelativeBase, MoveRecord, bIgnoreBaseRotation);
+
+	}
+
+	OutputSyncState.MoveDirectionIntent = (ProposedMove.bHasDirIntent ? ProposedMove.DirectionIntent : FVector::ZeroVector);
+
+	const FRotator StartingOrient = StartingSyncState->GetOrientation_WorldSpace();
+	FRotator TargetOrient = StartingOrient;
+
+	bool bIsOrientationChanging = false;
+
+	// Apply orientation changes (if any)
+	if (!ProposedMove.AngularVelocity.IsZero())
+	{ 
+		TargetOrient += (ProposedMove.AngularVelocity * DeltaSeconds);
+		bIsOrientationChanging = (TargetOrient != StartingOrient);
+	}
+
+	const FQuat TargetOrientQuat = TargetOrient.Quaternion();
+	FHitResult MoveHitResult(1.f);
+	
+	FVector CurMoveDelta = OrigMoveDelta;
+	
+	FOptionalFloorCheckResult StepUpFloorResult;	// passed to sub-operations, so we can use their final floor results if they did a test
+
+	float PercentTimeAppliedSoFar = MoveHitResult.Time;
+	bool bWasFirstMoveBlocked = false;
+	if (!CurMoveDelta.IsNearlyZero() || bIsOrientationChanging)
+	{
+		// Attempt to move the full amount first
+		bool bMoved = UMovementUtils::TrySafeMoveUpdatedComponent(UpdatedComponent, UpdatedPrimitive, CurMoveDelta, TargetOrientQuat, true, MoveHitResult, ETeleportType::None, MoveRecord);
+		float LastMoveSeconds = DeltaSeconds;
+
+		if (MoveHitResult.bStartPenetrating)
+		{
+			// We started by being stuck in geometry and need to resolve it first
+			// TODO: try to resolve starting stuck in geometry
+			bWasFirstMoveBlocked = true;
+		}
+		else if (MoveHitResult.IsValidBlockingHit())
+		{
+			// We impacted something (possibly a ramp, possibly a barrier)
+			bWasFirstMoveBlocked = true;
+			PercentTimeAppliedSoFar = MoveHitResult.Time;		
+
+			// Check if the blockage is a walkable ramp rising in front of us
+			// JAH TODO: change this normal.Z test to consider vs the movement plane (not just Z=0 plane)
+			if ((MoveHitResult.Time > 0.f) && (MoveHitResult.Normal.Z > UE_KINDA_SMALL_NUMBER) && 
+			    UFloorQueryUtils::IsHitSurfaceWalkable(MoveHitResult, CommonLegacySettings->MaxWalkSlopeCosine))
+			{
+				// It's a walkable ramp, so cut up the move and attempt to move the remainder of it along the ramp's surface, possibly generating another hit
+				const float PercentTimeRemaining = 1.f - PercentTimeAppliedSoFar;
+				CurMoveDelta = UGroundMovementUtils::ComputeDeflectedMoveOntoRamp(CurMoveDelta * PercentTimeRemaining, MoveHitResult, CommonLegacySettings->MaxWalkSlopeCosine, CurrentFloor.bLineTrace);
+				UMovementUtils::TrySafeMoveUpdatedComponent(UpdatedComponent, UpdatedPrimitive, CurMoveDelta, TargetOrientQuat, true, MoveHitResult, ETeleportType::None, MoveRecord);
+				LastMoveSeconds = PercentTimeRemaining * LastMoveSeconds;
+
+				const float SecondHitPercent = MoveHitResult.Time * PercentTimeRemaining;
+				PercentTimeAppliedSoFar = FMath::Clamp(PercentTimeAppliedSoFar + SecondHitPercent, 0.f, 1.f);
+			}
+
+			if (MoveHitResult.IsValidBlockingHit())
+			{
+				// If still blocked, try to step up onto the blocking object OR slide along it
+
+				UMoverComponent* MoverComponent = GetMoverComponent();
+				// JAH TODO: Take movement bases into account
+				if (UGroundMovementUtils::CanStepUpOnHitSurface(MoveHitResult)) // || (CharacterOwner->GetMovementBase() != nullptr && Hit.HitObjectHandle == CharacterOwner->GetMovementBase()->GetOwner()))
+				{
+					// hit a barrier or unwalkable surface, try to step up and onto it
+					const FVector PreStepUpLocation = UpdatedComponent->GetComponentLocation();
+					const FVector DownwardDir = -MoverComp->GetUpDirection();
+
+					if (!UGroundMovementUtils::TryMoveToStepUp(UpdatedComponent, UpdatedPrimitive, MoverComponent, DownwardDir, CommonLegacySettings->MaxStepHeight, CommonLegacySettings->MaxWalkSlopeCosine, CommonLegacySettings->FloorSweepDistance, OrigMoveDelta * (1.f - PercentTimeAppliedSoFar), MoveHitResult, CurrentFloor, false, &StepUpFloorResult, MoveRecord))
+					{
+						FMoverOnImpactParams ImpactParams(DefaultModeNames::Walking, MoveHitResult, OrigMoveDelta);
+						MoverComponent->HandleImpact(ImpactParams);
+						float PercentAvailableToSlide = 1.f - PercentTimeAppliedSoFar;
+						float SlideAmount = UGroundMovementUtils::TryWalkToSlideAlongSurface(UpdatedComponent, UpdatedPrimitive, MoverComponent, OrigMoveDelta, PercentAvailableToSlide, TargetOrientQuat, MoveHitResult.Normal, MoveHitResult, true, MoveRecord, CommonLegacySettings->MaxWalkSlopeCosine, CommonLegacySettings->MaxStepHeight);
+						PercentTimeAppliedSoFar += PercentAvailableToSlide * SlideAmount;
+					}
+				}
+				else if (MoveHitResult.Component.IsValid() && !MoveHitResult.Component.Get()->CanCharacterStepUp(Cast<APawn>(MoveHitResult.GetActor())))
+				{
+					FMoverOnImpactParams ImpactParams(DefaultModeNames::Walking, MoveHitResult, OrigMoveDelta);
+					MoverComponent->HandleImpact(ImpactParams);
+					float PercentAvailableToSlide = 1.f - PercentTimeAppliedSoFar;
+					float SlideAmount = UGroundMovementUtils::TryWalkToSlideAlongSurface(UpdatedComponent, UpdatedPrimitive, MoverComponent, OrigMoveDelta, 1.f - PercentTimeAppliedSoFar, TargetOrientQuat, MoveHitResult.Normal, MoveHitResult, true, MoveRecord, CommonLegacySettings->MaxWalkSlopeCosine, CommonLegacySettings->MaxStepHeight);
+					PercentTimeAppliedSoFar += PercentAvailableToSlide * SlideAmount;
+				}
+			}
+		}
+
+		// Search for the floor we've ended up on
+		UFloorQueryUtils::FindFloor(UpdatedComponent, UpdatedPrimitive,
+			CommonLegacySettings->FloorSweepDistance, CommonLegacySettings->MaxWalkSlopeCosine,
+			UpdatedPrimitive->GetComponentLocation(), CurrentFloor);
+
+		if (CurrentFloor.IsWalkableFloor())
+		{
+			UGroundMovementUtils::TryMoveToAdjustHeightAboveFloor(UpdatedComponent, UpdatedPrimitive, CurrentFloor, CommonLegacySettings->MaxWalkSlopeCosine, MoveRecord);
+		}
+    
+		if (!CurrentFloor.IsWalkableFloor() && !CurrentFloor.HitResult.bStartPenetrating)
+		{
+			// No floor or not walkable, so let's let the airborne movement mode deal with it
+			OutputState.MovementEndState.NextModeName = CommonLegacySettings->AirMovementModeName;
+			OutputState.MovementEndState.RemainingMs = Params.TimeStep.StepMs - (Params.TimeStep.StepMs * PercentTimeAppliedSoFar);
+			MoveRecord.SetDeltaSeconds((Params.TimeStep.StepMs - OutputState.MovementEndState.RemainingMs) * 0.001f);
+			CaptureFinalState(UpdatedComponent, CurrentFloor, MoveRecord, OutputSyncState);
+			return;
+		}
+	}
+	else
+	{
+		// If the actor isn't moving we still need to check if they have a valid floor
+		UFloorQueryUtils::FindFloor(UpdatedComponent, UpdatedPrimitive,
+			CommonLegacySettings->FloorSweepDistance, CommonLegacySettings->MaxWalkSlopeCosine,
+			UpdatedPrimitive->GetComponentLocation(), CurrentFloor);
+		
+		FHitResult Hit(CurrentFloor.HitResult);
+		if (Hit.bStartPenetrating)
+		{
+			// The floor check failed because it started in penetration
+			// We do not want to try to move downward because the downward sweep failed, rather we'd like to try to pop out of the floor.
+			Hit.TraceEnd = Hit.TraceStart + FVector(0.f, 0.f, 2.4f);
+			FVector RequestedAdjustment = UMovementUtils::ComputePenetrationAdjustment(Hit);
+			
+			const EMoveComponentFlags IncludeBlockingOverlapsWithoutEvents = (MOVECOMP_NeverIgnoreBlockingOverlaps | MOVECOMP_DisableBlockingOverlapDispatch);
+			EMoveComponentFlags MoveComponentFlags = MOVECOMP_NoFlags;
+			MoveComponentFlags = (MoveComponentFlags | IncludeBlockingOverlapsWithoutEvents);
+			UMovementUtils::TryMoveToResolvePenetration(UpdatedComponent, UpdatedPrimitive, MoveComponentFlags, RequestedAdjustment, Hit, UpdatedComponent->GetComponentQuat(), MoveRecord);
+		}
+		
+		if (!CurrentFloor.IsWalkableFloor() && !Hit.bStartPenetrating)
+		{
+			// No floor or not walkable, so let's let the airborne movement mode deal with it
+			OutputState.MovementEndState.NextModeName = CommonLegacySettings->AirMovementModeName;
+			OutputState.MovementEndState.RemainingMs = Params.TimeStep.StepMs;
+			MoveRecord.SetDeltaSeconds((Params.TimeStep.StepMs - OutputState.MovementEndState.RemainingMs) * 0.001f);
+			CaptureFinalState(UpdatedComponent, CurrentFloor, MoveRecord, OutputSyncState);
+			return;
+		}
+	}
+
+	CaptureFinalState(UpdatedComponent, CurrentFloor, MoveRecord, OutputSyncState);
+
+}
+
+
+void UWalkingMode::OnRegistered(const FName ModeName)
+{
+	Super::OnRegistered(ModeName);
+
+	CommonLegacySettings = GetMoverComponent()->FindSharedSettings<UCommonLegacyMovementSettings>();
+	ensureMsgf(CommonLegacySettings, TEXT("Failed to find instance of CommonLegacyMovementSettings on %s. Movement may not function properly."), *GetPathNameSafe(this));
+}
+
+void UWalkingMode::OnUnregistered()
+{
+	CommonLegacySettings = nullptr;
+
+	Super::OnUnregistered();
+}
+
+bool UWalkingMode::AttemptJump(float UpwardsSpeed, FMoverTickEndData& OutputState)
+{
+	FMoverDefaultSyncState& OutputSyncState = OutputState.SyncState.SyncStateCollection.FindOrAddMutableDataByType<FMoverDefaultSyncState>();
+
+	// TODO: This should check if a jump is even allowed
+	TSharedPtr<FLayeredMove_JumpImpulse> JumpMove = MakeShared<FLayeredMove_JumpImpulse>();
+	JumpMove->UpwardsSpeed = UpwardsSpeed;
+	OutputSyncState.LayeredMoves.QueueLayeredMove(JumpMove);
+	OutputState.MovementEndState.NextModeName = CommonLegacySettings->AirMovementModeName;
+	return true;
+}
+
+bool UWalkingMode::AttemptTeleport(USceneComponent* UpdatedComponent, const FVector& TeleportPos, const FRotator& TeleportRot, const FVector& PriorVelocity, FMoverTickEndData& Output)
+{
+	if (UpdatedComponent->GetOwner()->TeleportTo(TeleportPos, TeleportRot))
+	{
+		FMoverDefaultSyncState& OutputSyncState = Output.SyncState.SyncStateCollection.FindOrAddMutableDataByType<FMoverDefaultSyncState>();
+
+		OutputSyncState.SetTransforms_WorldSpace( UpdatedComponent->GetComponentLocation(),
+												  UpdatedComponent->GetComponentRotation(),
+												  PriorVelocity,
+												  nullptr ); // no movement base
+		
+		// TODO: instead of invalidating it, consider checking for a floor. Possibly a dynamic base?
+		GetBlackboard_Mutable()->Invalidate(CommonBlackboard::LastFloorResult);
+		GetBlackboard_Mutable()->Invalidate(CommonBlackboard::LastMovementBase);
+
+		return true;
+	}
+
+	return false;
+}
+
+// TODO: replace this function with simply looking at/collapsing the MovementRecord
+void UWalkingMode::CaptureFinalState(USceneComponent* UpdatedComponent, const FFloorCheckResult& FloorResult, const FMovementRecord& Record, FMoverDefaultSyncState& OutputSyncState) const
+{
+	FRelativeBaseInfo BaseInfo = UpdateFloorAndBaseInfo(FloorResult);
+
+	// TODO: Update Main/large movement record with substeps from our local record
+	
+	if (BaseInfo.HasRelativeInfo())
+	{
+		OutputSyncState.SetTransforms_WorldSpace( UpdatedComponent->GetComponentLocation(),
+												  UpdatedComponent->GetComponentRotation(),
+												  Record.GetRelevantVelocity(),
+												  BaseInfo.MovementBase, BaseInfo.BoneName);
+	}
+	else
+	{
+		OutputSyncState.SetTransforms_WorldSpace( UpdatedComponent->GetComponentLocation(),
+												  UpdatedComponent->GetComponentRotation(),
+												  Record.GetRelevantVelocity(),
+												  nullptr);	// no movement base
+	}
+
+	UpdatedComponent->ComponentVelocity = OutputSyncState.GetVelocity_WorldSpace();
+}
+
+
+FRelativeBaseInfo UWalkingMode::UpdateFloorAndBaseInfo(const FFloorCheckResult& FloorResult) const
+{
+	FRelativeBaseInfo ReturnBaseInfo;
+
+	UMoverBlackboard* SimBlackboard = GetBlackboard_Mutable();
+
+	SimBlackboard->Set(CommonBlackboard::LastFloorResult, FloorResult);
+
+	if (FloorResult.IsWalkableFloor() && UBasedMovementUtils::IsADynamicBase(FloorResult.HitResult.GetComponent()))
+	{
+		ReturnBaseInfo.SetFromFloorResult(FloorResult);
+
+		SimBlackboard->Set(CommonBlackboard::LastMovementBase, ReturnBaseInfo);
+	}
+	else
+	{
+
+		SimBlackboard->Invalidate(CommonBlackboard::LastMovementBase);
+	}
+
+	return ReturnBaseInfo;
+}
diff --git a/Engine/Plugins/Experimental/Mover/Source/Mover/Private/Kinematic/LayeredMoves/AnimRootMotionLayeredMove.cpp b/Engine/Plugins/Experimental/Mover/Source/Mover/Private/Kinematic/LayeredMoves/AnimRootMotionLayeredMove.cpp
deleted file mode 100644
diff --git a/Engine/Plugins/Experimental/Mover/Source/Mover/Private/Kinematic/LayeredMoves/BasicLayeredMoves.cpp b/Engine/Plugins/Experimental/Mover/Source/Mover/Private/Kinematic/LayeredMoves/BasicLayeredMoves.cpp
deleted file mode 100644
diff --git a/Engine/Plugins/Experimental/Mover/Source/Mover/Private/Kinematic/LayeredMoves/LaunchMove.cpp b/Engine/Plugins/Experimental/Mover/Source/Mover/Private/Kinematic/LayeredMoves/LaunchMove.cpp
deleted file mode 100644
diff --git a/Engine/Plugins/Experimental/Mover/Source/Mover/Private/Kinematic/LayeredMoves/MultiJumpLayeredMove.cpp b/Engine/Plugins/Experimental/Mover/Source/Mover/Private/Kinematic/LayeredMoves/MultiJumpLayeredMove.cpp
deleted file mode 100644
diff --git a/Engine/Plugins/Experimental/Mover/Source/Mover/Private/Kinematic/Modes/FallingMode.cpp b/Engine/Plugins/Experimental/Mover/Source/Mover/Private/Kinematic/Modes/FallingMode.cpp
deleted file mode 100644
diff --git a/Engine/Plugins/Experimental/Mover/Source/Mover/Private/Kinematic/Modes/FlyingMode.cpp b/Engine/Plugins/Experimental/Mover/Source/Mover/Private/Kinematic/Modes/FlyingMode.cpp
deleted file mode 100644
diff --git a/Engine/Plugins/Experimental/Mover/Source/Mover/Private/Kinematic/Modes/SwimmingMode.cpp b/Engine/Plugins/Experimental/Mover/Source/Mover/Private/Kinematic/Modes/SwimmingMode.cpp
deleted file mode 100644
diff --git a/Engine/Plugins/Experimental/Mover/Source/Mover/Private/Kinematic/Modes/WalkingMode.cpp b/Engine/Plugins/Experimental/Mover/Source/Mover/Private/Kinematic/Modes/WalkingMode.cpp
deleted file mode 100644
diff --git a/Engine/Plugins/Experimental/Mover/Source/Mover/Private/MoveLibrary/PlayMoverMontageCallbackProxy.cpp b/Engine/Plugins/Experimental/Mover/Source/Mover/Private/MoveLibrary/PlayMoverMontageCallbackProxy.cpp
--- a/Engine/Plugins/Experimental/Mover/Source/Mover/Private/MoveLibrary/PlayMoverMontageCallbackProxy.cpp
+++ b/Engine/Plugins/Experimental/Mover/Source/Mover/Private/MoveLibrary/PlayMoverMontageCallbackProxy.cpp
@@ -4,7 +4,7 @@
 #include "Animation/AnimMontage.h"
 #include "Components/SkeletalMeshComponent.h"
 #include "MoverComponent.h"
-#include "Kinematic/LayeredMoves/AnimRootMotionLayeredMove.h"
+#include "DefaultMovementSet/LayeredMoves/AnimRootMotionLayeredMove.h"
 
 #include UE_INLINE_GENERATED_CPP_BY_NAME(PlayMoverMontageCallbackProxy)
 
diff --git a/Engine/Plugins/Experimental/Mover/Source/Mover/Private/MoverComponent.cpp b/Engine/Plugins/Experimental/Mover/Source/Mover/Private/MoverComponent.cpp
--- a/Engine/Plugins/Experimental/Mover/Source/Mover/Private/MoverComponent.cpp
+++ b/Engine/Plugins/Experimental/Mover/Source/Mover/Private/MoverComponent.cpp
@@ -4,9 +4,9 @@
 #include "MoverComponent.h"
 #include "MoverSimulationTypes.h"
 #include "MovementModeStateMachine.h"
-#include "Kinematic/Modes/WalkingMode.h"
-#include "Kinematic/Modes/FallingMode.h"
-#include "Kinematic/Modes/FlyingMode.h"
+#include "DefaultMovementSet/Modes/WalkingMode.h"
+#include "DefaultMovementSet/Modes/FallingMode.h"
+#include "DefaultMovementSet/Modes/FlyingMode.h"
 #include "MoveLibrary/MovementMixer.h"
 #include "MoveLibrary/MovementUtils.h"
 #include "MoveLibrary/FloorQueryUtils.h"
@@ -49,11 +49,11 @@ UMoverComponent::UMoverComponent()
 	bAutoActivate = true;
 
 	// Default movement modes
-	MovementModes.Add(KinematicModeNames::Walking, CreateDefaultSubobject<UWalkingMode>(TEXT("DefaultWalkingMode")));
-	MovementModes.Add(KinematicModeNames::Falling, CreateDefaultSubobject<UFallingMode>(TEXT("DefaultFallingMode")));
-	MovementModes.Add(KinematicModeNames::Flying,  CreateDefaultSubobject<UFlyingMode>(TEXT("DefaultFlyingMode")));
+	MovementModes.Add(DefaultModeNames::Walking, CreateDefaultSubobject<UWalkingMode>(TEXT("DefaultWalkingMode")));
+	MovementModes.Add(DefaultModeNames::Falling, CreateDefaultSubobject<UFallingMode>(TEXT("DefaultFallingMode")));
+	MovementModes.Add(DefaultModeNames::Flying,  CreateDefaultSubobject<UFlyingMode>(TEXT("DefaultFlyingMode")));
 
-	StartingMovementMode = KinematicModeNames::Falling;
+	StartingMovementMode = DefaultModeNames::Falling;
 
 	PersistentSyncStateDataTypes.Add(FMoverDataPersistence(FMoverDefaultSyncState::StaticStruct(), true));
 
@@ -328,7 +328,7 @@ void UMoverComponent::SimulationTick(const FMoverTimeStep& InTimeStep, const FMo
 
 	SimOutput.AuxState = SimInput.AuxState;
 
-	FKinematicDefaultInputs* Input = SimInput.InputCmd.InputCollection.FindMutableDataByType<FKinematicDefaultInputs>();
+	FCharacterDefaultInputs* Input = SimInput.InputCmd.InputCollection.FindMutableDataByType<FCharacterDefaultInputs>();
 
 
 	if (Input && !Input->SuggestedMovementMode.IsNone())
@@ -956,7 +956,7 @@ FRotator UMoverComponent::GetTargetOrientation() const
 	if (bHasValidCachedUsedInput)
 	{
 		const FMoverInputCmdContext& LastInputCmd = GetLastInputCmd();
-		const FKinematicDefaultInputs* MoverInputs = LastInputCmd.InputCollection.FindDataByType<FKinematicDefaultInputs>();
+		const FCharacterDefaultInputs* MoverInputs = LastInputCmd.InputCollection.FindDataByType<FCharacterDefaultInputs>();
 
 		const FVector TargetOrientationDir = MoverInputs->GetOrientationIntentDir_WorldSpace();
 
@@ -1024,7 +1024,7 @@ bool UMoverComponent::IsFalling() const
 	{
 		if (const FMoverDefaultSyncState* SyncState = CachedLastSyncState.SyncStateCollection.FindDataByType<FMoverDefaultSyncState>())
 		{
-			return SyncState->MovementMode == KinematicModeNames::Falling;
+			return SyncState->MovementMode == DefaultModeNames::Falling;
 		}
 	}
 
@@ -1037,7 +1037,7 @@ bool UMoverComponent::IsAirborne() const
 	{
 		if (const FMoverDefaultSyncState* SyncState = CachedLastSyncState.SyncStateCollection.FindDataByType<FMoverDefaultSyncState>())
 		{
-			return SyncState->MovementMode == KinematicModeNames::Flying || SyncState->MovementMode == KinematicModeNames::Falling;
+			return SyncState->MovementMode == DefaultModeNames::Flying || SyncState->MovementMode == DefaultModeNames::Falling;
 		}
 	}
 
@@ -1050,7 +1050,7 @@ bool UMoverComponent::IsOnGround() const
 	{
 		if (const FMoverDefaultSyncState* SyncState = CachedLastSyncState.SyncStateCollection.FindDataByType<FMoverDefaultSyncState>())
 		{
-			return SyncState->MovementMode == KinematicModeNames::Walking;
+			return SyncState->MovementMode == DefaultModeNames::Walking;
 		}
 	}
 
@@ -1062,7 +1062,7 @@ bool UMoverComponent::IsSlopeSliding() const
 	if (IsAirborne())
 	{
 		FFloorCheckResult HitResult;
-		if (SimBlackboard->TryGet(KinematicBlackboard::LastFloorResult, HitResult))
+		if (SimBlackboard->TryGet(CommonBlackboard::LastFloorResult, HitResult))
 		{
 			return HitResult.bBlockingHit && !HitResult.bWalkableFloor;
 		}
@@ -1200,7 +1200,7 @@ const FMoverSyncState& UMoverComponent::GetSyncState() const
 bool UMoverComponent::TryGetFloorCheckHitResult(FHitResult& OutHitResult) const
 {
 	FFloorCheckResult FloorCheck;
-	if (SimBlackboard->TryGet(KinematicBlackboard::LastFloorResult, FloorCheck))
+	if (SimBlackboard->TryGet(CommonBlackboard::LastFloorResult, FloorCheck))
 	{
 		OutHitResult = FloorCheck.HitResult;
 		return true;
diff --git a/Engine/Plugins/Experimental/Mover/Source/Mover/Private/MoverDataModelTypes.cpp b/Engine/Plugins/Experimental/Mover/Source/Mover/Private/MoverDataModelTypes.cpp
--- a/Engine/Plugins/Experimental/Mover/Source/Mover/Private/MoverDataModelTypes.cpp
+++ b/Engine/Plugins/Experimental/Mover/Source/Mover/Private/MoverDataModelTypes.cpp
@@ -8,9 +8,9 @@
 
 
 
-// FKinematicDefaultInputs //////////////////////////////////////////////////////////////
+// FCharacterDefaultInputs //////////////////////////////////////////////////////////////
 
-void FKinematicDefaultInputs::SetMoveInput(EMoveInputType InMoveInputType, const FVector& InMoveInput)
+void FCharacterDefaultInputs::SetMoveInput(EMoveInputType InMoveInputType, const FVector& InMoveInput)
 {
 	MoveInputType = InMoveInputType;
 
@@ -23,7 +23,7 @@ void FKinematicDefaultInputs::SetMoveInput(EMoveInputType InMoveInputType, const
 }
 
 
-FVector FKinematicDefaultInputs::GetMoveInput_WorldSpace() const
+FVector FCharacterDefaultInputs::GetMoveInput_WorldSpace() const
 {
 	if (bUsingMovementBase && MovementBase)
 	{
@@ -36,7 +36,7 @@ FVector FKinematicDefaultInputs::GetMoveInput_WorldSpace() const
 }
 
 
-FVector FKinematicDefaultInputs::GetOrientationIntentDir_WorldSpace() const
+FVector FCharacterDefaultInputs::GetOrientationIntentDir_WorldSpace() const
 {
 	if (bUsingMovementBase && MovementBase)
 	{
@@ -49,14 +49,14 @@ FVector FKinematicDefaultInputs::GetOrientationIntentDir_WorldSpace() const
 }
 
 
-FMoverDataStructBase* FKinematicDefaultInputs::Clone() const
+FMoverDataStructBase* FCharacterDefaultInputs::Clone() const
 {
 	// TODO: ensure that this memory allocation jives with deletion method
-	FKinematicDefaultInputs* CopyPtr = new FKinematicDefaultInputs(*this);
+	FCharacterDefaultInputs* CopyPtr = new FCharacterDefaultInputs(*this);
 	return CopyPtr;
 }
 
-bool FKinematicDefaultInputs::NetSerialize(FArchive& Ar, UPackageMap* Map, bool& bOutSuccess)
+bool FCharacterDefaultInputs::NetSerialize(FArchive& Ar, UPackageMap* Map, bool& bOutSuccess)
 {
 	Super::NetSerialize(Ar, Map, bOutSuccess);
 
@@ -90,7 +90,7 @@ bool FKinematicDefaultInputs::NetSerialize(FArchive& Ar, UPackageMap* Map, bool&
 }
 
 
-void FKinematicDefaultInputs::ToString(FAnsiStringBuilderBase& Out) const
+void FCharacterDefaultInputs::ToString(FAnsiStringBuilderBase& Out) const
 {
 	Super::ToString(Out);
 
@@ -387,12 +387,12 @@ FRotator FMoverDefaultSyncState::GetOrientation_BaseSpace() const
 
 // UMoverDataModelBlueprintLibrary ///////////////////////////////////////////////////
 
-void UMoverDataModelBlueprintLibrary::SetMoveIntent(FKinematicDefaultInputs& Inputs, const FVector& WorldDirectionIntent)
+void UMoverDataModelBlueprintLibrary::SetMoveIntent(FCharacterDefaultInputs& Inputs, const FVector& WorldDirectionIntent)
 {
 	Inputs.SetMoveInput(EMoveInputType::DirectionalIntent, WorldDirectionIntent);
 }
 
-FVector UMoverDataModelBlueprintLibrary::GetMoveDirectionIntentFromInputs(const FKinematicDefaultInputs& Inputs)
+FVector UMoverDataModelBlueprintLibrary::GetMoveDirectionIntentFromInputs(const FCharacterDefaultInputs& Inputs)
 {
 	return Inputs.GetMoveInput_WorldSpace();
 }
diff --git a/Engine/Plugins/Experimental/Mover/Source/Mover/Private/PhysicsMover/Modes/PhysicsDrivenFallingMode.cpp b/Engine/Plugins/Experimental/Mover/Source/Mover/Private/PhysicsMover/Modes/PhysicsDrivenFallingMode.cpp
--- a/Engine/Plugins/Experimental/Mover/Source/Mover/Private/PhysicsMover/Modes/PhysicsDrivenFallingMode.cpp
+++ b/Engine/Plugins/Experimental/Mover/Source/Mover/Private/PhysicsMover/Modes/PhysicsDrivenFallingMode.cpp
@@ -5,7 +5,7 @@
 #include "Chaos/Character/CharacterGroundConstraint.h"
 #include "GameFramework/PhysicsVolume.h"
 #include "MoverComponent.h"
-#include "Kinematic/Settings/CommonLegacyMovementSettings.h"
+#include "DefaultMovementSet/Settings/CommonLegacyMovementSettings.h"
 #include "Math/UnitConversion.h"
 #include "MoveLibrary/FloorQueryUtils.h"
 #include "MoveLibrary/WaterMovementUtils.h"
@@ -33,7 +33,7 @@ void UPhysicsDrivenFallingMode::OnSimulationTick(const FSimulationTickParams& Pa
 	UPrimitiveComponent* UpdatedPrimitive = Params.UpdatedPrimitive;
 	FProposedMove ProposedMove = Params.ProposedMove;
 
-	const FKinematicDefaultInputs* KinematicInputs = StartState.InputCmd.InputCollection.FindDataByType<FKinematicDefaultInputs>();
+	const FCharacterDefaultInputs* CharacterInputs = StartState.InputCmd.InputCollection.FindDataByType<FCharacterDefaultInputs>();
 	const FMoverDefaultSyncState* StartingSyncState = StartState.SyncState.SyncStateCollection.FindDataByType<FMoverDefaultSyncState>();
 	check(StartingSyncState);
 
@@ -58,8 +58,8 @@ void UPhysicsDrivenFallingMode::OnSimulationTick(const FSimulationTickParams& Pa
 		return;
 	}
 
-	SimBlackboard->Invalidate(KinematicBlackboard::LastFloorResult);
-	SimBlackboard->Invalidate(KinematicBlackboard::LastWaterResult);
+	SimBlackboard->Invalidate(CommonBlackboard::LastFloorResult);
+	SimBlackboard->Invalidate(CommonBlackboard::LastWaterResult);
 
 	// Find floor
 
@@ -73,15 +73,15 @@ void UPhysicsDrivenFallingMode::OnSimulationTick(const FSimulationTickParams& Pa
 	UpdatedPrimitive, UpDir, PawnRadius, TargetHeight, CommonLegacySettings->MaxStepHeight,
 		CommonLegacySettings->MaxWalkSlopeCosine, FloorResult, WaterResult);
 
-	SimBlackboard->Set(KinematicBlackboard::LastFloorResult, FloorResult);
-	SimBlackboard->Set(KinematicBlackboard::LastWaterResult, WaterResult);
+	SimBlackboard->Set(CommonBlackboard::LastFloorResult, FloorResult);
+	SimBlackboard->Set(CommonBlackboard::LastWaterResult, WaterResult);
 
-	const bool bIsMovingUp = KinematicInputs->bIsJumpJustPressed || (UpDir.Dot(ProposedMove.LinearVelocity) > 0.0f);
+	const bool bIsMovingUp = CharacterInputs->bIsJumpJustPressed || (UpDir.Dot(ProposedMove.LinearVelocity) > 0.0f);
 	const bool bStartSwimming = WaterResult.WaterSplineData.ImmersionDepth > CommonLegacySettings->SwimmingStartImmersionDepth;
 
 	if (WaterResult.IsSwimmableVolume() && bStartSwimming && !bIsMovingUp)
 	{
-		OutputState.MovementEndState.NextModeName = KinematicModeNames::Swimming;
+		OutputState.MovementEndState.NextModeName = DefaultModeNames::Swimming;
 		OutputState.MovementEndState.RemainingMs = Params.TimeStep.StepMs;
 		return;
 	}
@@ -112,12 +112,12 @@ void UPhysicsDrivenFallingMode::OnSimulationTick(const FSimulationTickParams& Pa
 
 	if (FloorResult.IsWalkableFloor() && bIsFloorWithinReach && !bIsMovingUpRelativeToFloor)
 	{
-		OutputState.MovementEndState.NextModeName = KinematicModeNames::Walking;
+		OutputState.MovementEndState.NextModeName = DefaultModeNames::Walking;
 		TargetPos -= UpDir * (FloorResult.FloorDist - TargetHeight);
 	}
 	else
 	{
-		OutputState.MovementEndState.NextModeName = KinematicModeNames::Falling;
+		OutputState.MovementEndState.NextModeName = DefaultModeNames::Falling;
 	}
 
 	OutputState.MovementEndState.RemainingMs = 0.0f;
diff --git a/Engine/Plugins/Experimental/Mover/Source/Mover/Private/PhysicsMover/Modes/PhysicsDrivenFlyingMode.cpp b/Engine/Plugins/Experimental/Mover/Source/Mover/Private/PhysicsMover/Modes/PhysicsDrivenFlyingMode.cpp
--- a/Engine/Plugins/Experimental/Mover/Source/Mover/Private/PhysicsMover/Modes/PhysicsDrivenFlyingMode.cpp
+++ b/Engine/Plugins/Experimental/Mover/Source/Mover/Private/PhysicsMover/Modes/PhysicsDrivenFlyingMode.cpp
@@ -5,7 +5,7 @@
 #include "Chaos/Character/CharacterGroundConstraint.h"
 #include "GameFramework/PhysicsVolume.h"
 #include "MoverComponent.h"
-#include "Kinematic/Settings/CommonLegacyMovementSettings.h"
+#include "DefaultMovementSet/Settings/CommonLegacyMovementSettings.h"
 #include "Math/UnitConversion.h"
 #include "MoveLibrary/MovementUtils.h"
 #include "PhysicsMover/PhysicsMovementUtils.h"
@@ -32,7 +32,6 @@ void UPhysicsDrivenFlyingMode::OnSimulationTick(const FSimulationTickParams& Par
 	UPrimitiveComponent* UpdatedPrimitive = Params.UpdatedPrimitive;
 	FProposedMove ProposedMove = Params.ProposedMove;
 
-	const FKinematicDefaultInputs* KinematicInputs = StartState.InputCmd.InputCollection.FindDataByType<FKinematicDefaultInputs>();
 	const FMoverDefaultSyncState* StartingSyncState = StartState.SyncState.SyncStateCollection.FindDataByType<FMoverDefaultSyncState>();
 	check(StartingSyncState);
 
@@ -52,8 +51,8 @@ void UPhysicsDrivenFlyingMode::OnSimulationTick(const FSimulationTickParams& Par
 
 	if (UMoverBlackboard* SimBlackboard = GetBlackboard_Mutable())
 	{
-		SimBlackboard->Invalidate(KinematicBlackboard::LastFloorResult);
-		SimBlackboard->Invalidate(KinematicBlackboard::LastWaterResult);
+		SimBlackboard->Invalidate(CommonBlackboard::LastFloorResult);
+		SimBlackboard->Invalidate(CommonBlackboard::LastWaterResult);
 	}
 
 	// In air steering
@@ -73,7 +72,7 @@ void UPhysicsDrivenFlyingMode::OnSimulationTick(const FSimulationTickParams& Par
 
 	FVector TargetPos = StartingSyncState->GetLocation_WorldSpace() + TargetVel * DeltaSeconds;
 
-	OutputState.MovementEndState.NextModeName = KinematicModeNames::Flying;
+	OutputState.MovementEndState.NextModeName = DefaultModeNames::Flying;
 
 	OutputState.MovementEndState.RemainingMs = 0.0f;
 	OutputSyncState.MoveDirectionIntent = ProposedMove.bHasDirIntent ? ProposedMove.DirectionIntent : FVector::ZeroVector;
diff --git a/Engine/Plugins/Experimental/Mover/Source/Mover/Private/PhysicsMover/Modes/PhysicsDrivenSwimmingMode.cpp b/Engine/Plugins/Experimental/Mover/Source/Mover/Private/PhysicsMover/Modes/PhysicsDrivenSwimmingMode.cpp
--- a/Engine/Plugins/Experimental/Mover/Source/Mover/Private/PhysicsMover/Modes/PhysicsDrivenSwimmingMode.cpp
+++ b/Engine/Plugins/Experimental/Mover/Source/Mover/Private/PhysicsMover/Modes/PhysicsDrivenSwimmingMode.cpp
@@ -6,9 +6,9 @@
 #include "Math/UnitConversion.h"
 #include "MoveLibrary/FloorQueryUtils.h"
 #include "MoveLibrary/WaterMovementUtils.h"
-#include "Kinematic/LayeredMoves/BasicLayeredMoves.h"
+#include "DefaultMovementSet/LayeredMoves/BasicLayeredMoves.h"
 #include "MoverComponent.h"
-#include "Kinematic/Settings/CommonLegacyMovementSettings.h"
+#include "DefaultMovementSet/Settings/CommonLegacyMovementSettings.h"
 #include "GameFramework/PhysicsVolume.h"
 #include "PhysicsMover/PhysicsMovementUtils.h"
 
@@ -38,7 +38,7 @@ void UPhysicsDrivenSwimmingMode::OnSimulationTick(const FSimulationTickParams& P
 	
 	const FVector UpDir = GetMoverComponent()->GetUpDirection();
 	
-	const FKinematicDefaultInputs* KinematicInputs = StartState.InputCmd.InputCollection.FindDataByType<FKinematicDefaultInputs>();
+	const FCharacterDefaultInputs* CharacterInputs = StartState.InputCmd.InputCollection.FindDataByType<FCharacterDefaultInputs>();
 
 	const FMoverDefaultSyncState* StartingSyncState = StartState.SyncState.SyncStateCollection.FindDataByType<FMoverDefaultSyncState>();
 	check(StartingSyncState);
@@ -48,7 +48,7 @@ void UPhysicsDrivenSwimmingMode::OnSimulationTick(const FSimulationTickParams& P
 	const float DeltaSeconds = Params.TimeStep.StepMs * 0.001f;
 
 	if ((ProposedMove.bHasTargetLocation && AttemptTeleport(UpdatedComponent, ProposedMove.TargetLocation, UpdatedComponent->GetComponentRotation(), StartingSyncState->GetVelocity_WorldSpace(), OutputState)) ||	// Teleport
-		(KinematicInputs->bIsJumpJustPressed && AttemptJump(SurfaceSwimmingWaterControlSettings.JumpMultiplier*CommonLegacySettings->JumpUpwardsSpeed, OutputState)))
+		(CharacterInputs->bIsJumpJustPressed && AttemptJump(SurfaceSwimmingWaterControlSettings.JumpMultiplier*CommonLegacySettings->JumpUpwardsSpeed, OutputState)))
 	{
 		OutputState.MovementEndState.RemainingMs = Params.TimeStep.StepMs;
 		return;
@@ -60,8 +60,8 @@ void UPhysicsDrivenSwimmingMode::OnSimulationTick(const FSimulationTickParams& P
 		OutputSyncState = *StartingSyncState;
 		return;
 	}
-	SimBlackboard->Invalidate(KinematicBlackboard::LastFloorResult);
-	SimBlackboard->Invalidate(KinematicBlackboard::LastWaterResult);
+	SimBlackboard->Invalidate(CommonBlackboard::LastFloorResult);
+	SimBlackboard->Invalidate(CommonBlackboard::LastWaterResult);
 	
 	OutputSyncState.MoveDirectionIntent = ProposedMove.bHasDirIntent ? ProposedMove.DirectionIntent : FVector::ZeroVector;
 
@@ -77,8 +77,8 @@ void UPhysicsDrivenSwimmingMode::OnSimulationTick(const FSimulationTickParams& P
 		UpdatedPrimitive, UpDir, PawnRadius, TargetHeight, CommonLegacySettings->MaxStepHeight,
 		CommonLegacySettings->MaxWalkSlopeCosine, FloorResult, WaterResult);
 	
-	SimBlackboard->Set(KinematicBlackboard::LastFloorResult, FloorResult);
-	SimBlackboard->Set(KinematicBlackboard::LastWaterResult, WaterResult);
+	SimBlackboard->Set(CommonBlackboard::LastFloorResult, FloorResult);
+	SimBlackboard->Set(CommonBlackboard::LastWaterResult, WaterResult);
 	
 	if (WaterResult.IsSwimmableVolume())
 	{
@@ -109,20 +109,20 @@ void UPhysicsDrivenSwimmingMode::OnSimulationTick(const FSimulationTickParams& P
 	
 		if (bWalkTrigger && bIsWithinReach)
 		{
-			OutputState.MovementEndState.NextModeName = KinematicModeNames::Walking;
+			OutputState.MovementEndState.NextModeName = DefaultModeNames::Walking;
 		}
 		else if (bFallTrigger)
 		{
-			OutputState.MovementEndState.NextModeName = KinematicModeNames::Falling;
+			OutputState.MovementEndState.NextModeName = DefaultModeNames::Falling;
 		}
 		else
 		{
-			OutputState.MovementEndState.NextModeName = KinematicModeNames::Swimming;
+			OutputState.MovementEndState.NextModeName = DefaultModeNames::Swimming;
 		}
 	}
 	else
 	{
-		OutputState.MovementEndState.NextModeName = KinematicModeNames::Falling;
+		OutputState.MovementEndState.NextModeName = DefaultModeNames::Falling;
 	}
 	
 	OutputState.MovementEndState.RemainingMs = 0.0f;
@@ -136,7 +136,7 @@ bool UPhysicsDrivenSwimmingMode::AttemptJump(float UpwardsSpeed, FMoverTickEndDa
  	TSharedPtr<FLayeredMove_JumpImpulse> JumpMove = MakeShared<FLayeredMove_JumpImpulse>();
 	JumpMove->UpwardsSpeed = UpwardsSpeed;
 	OutputSyncState.LayeredMoves.QueueLayeredMove(JumpMove);
-	OutputState.MovementEndState.NextModeName = KinematicModeNames::Falling;
+	OutputState.MovementEndState.NextModeName = DefaultModeNames::Falling;
 	return true;
 }
 
@@ -150,8 +150,8 @@ bool UPhysicsDrivenSwimmingMode::AttemptTeleport(USceneComponent* UpdatedCompone
 		nullptr); // no movement base
 
 	// TODO: instead of invalidating it, consider checking for a floor. Possibly a dynamic base?
-	GetBlackboard_Mutable()->Invalidate(KinematicBlackboard::LastFloorResult);
-	GetBlackboard_Mutable()->Invalidate(KinematicBlackboard::LastMovementBase);
+	GetBlackboard_Mutable()->Invalidate(CommonBlackboard::LastFloorResult);
+	GetBlackboard_Mutable()->Invalidate(CommonBlackboard::LastMovementBase);
 
 	return true;
 }
\ No newline at end of file
diff --git a/Engine/Plugins/Experimental/Mover/Source/Mover/Private/PhysicsMover/Modes/PhysicsDrivenWalkingMode.cpp b/Engine/Plugins/Experimental/Mover/Source/Mover/Private/PhysicsMover/Modes/PhysicsDrivenWalkingMode.cpp
--- a/Engine/Plugins/Experimental/Mover/Source/Mover/Private/PhysicsMover/Modes/PhysicsDrivenWalkingMode.cpp
+++ b/Engine/Plugins/Experimental/Mover/Source/Mover/Private/PhysicsMover/Modes/PhysicsDrivenWalkingMode.cpp
@@ -6,9 +6,9 @@
 #include "Chaos/PhysicsObject.h"
 #include "Chaos/PhysicsObjectInternalInterface.h"
 #include "GameFramework/PhysicsVolume.h"
-#include "Kinematic/LayeredMoves/BasicLayeredMoves.h"
+#include "DefaultMovementSet/LayeredMoves/BasicLayeredMoves.h"
 #include "MoverComponent.h"
-#include "Kinematic/Settings/CommonLegacyMovementSettings.h"
+#include "DefaultMovementSet/Settings/CommonLegacyMovementSettings.h"
 #include "Math/UnitConversion.h"
 #include "MoveLibrary/WaterMovementUtils.h"
 #include "PhysicsMover/PhysicsMovementUtils.h"
@@ -42,7 +42,6 @@ void UPhysicsDrivenWalkingMode::OnSimulationTick(const FSimulationTickParams& Pa
 
 	const FVector UpDir = GetMoverComponent()->GetUpDirection();
 
-	const FKinematicDefaultInputs* KinematicInputs = StartState.InputCmd.InputCollection.FindDataByType<FKinematicDefaultInputs>();
 	const FMoverDefaultSyncState* StartingSyncState = StartState.SyncState.SyncStateCollection.FindDataByType<FMoverDefaultSyncState>();
 	check(StartingSyncState);
 
@@ -67,7 +66,7 @@ void UPhysicsDrivenWalkingMode::OnSimulationTick(const FSimulationTickParams& Pa
 	// Store the previous ground normal that was used to compute the proposed move
 	FFloorCheckResult PrevFloorResult;
 	FVector PrevGroundNormal = UpDir;
-	if (SimBlackboard->TryGet(KinematicBlackboard::LastFloorResult, PrevFloorResult))
+	if (SimBlackboard->TryGet(CommonBlackboard::LastFloorResult, PrevFloorResult))
 	{
 		PrevGroundNormal = PrevFloorResult.HitResult.ImpactNormal;
 	}
@@ -83,14 +82,14 @@ void UPhysicsDrivenWalkingMode::OnSimulationTick(const FSimulationTickParams& Pa
 		UpdatedPrimitive, UpDir, PawnRadius, TargetHeight, CommonLegacySettings->MaxStepHeight,
 		CommonLegacySettings->MaxWalkSlopeCosine, FloorResult, WaterResult);
 
-	SimBlackboard->Set(KinematicBlackboard::LastFloorResult, FloorResult);
-	SimBlackboard->Set(KinematicBlackboard::LastWaterResult, WaterResult);
+	SimBlackboard->Set(CommonBlackboard::LastFloorResult, FloorResult);
+	SimBlackboard->Set(CommonBlackboard::LastWaterResult, WaterResult);
 
 	const bool bStartSwimming = (WaterResult.WaterSplineData.ImmersionDepth+TargetHeight) > CommonLegacySettings->SwimmingIdealImmersionDepth;
 	
 	if (WaterResult.IsSwimmableVolume() && bStartSwimming)
 	{
-		SwitchToState(KinematicModeNames::Swimming, Params, OutputState);
+		SwitchToState(DefaultModeNames::Swimming, Params, OutputState);
 	}
 	else if (FloorResult.IsWalkableFloor())
 	{
@@ -145,15 +144,15 @@ void UPhysicsDrivenWalkingMode::OnSimulationTick(const FSimulationTickParams& Pa
 		// If the character is unsupported allow some grace period before falling
 		bool bIsSupported = bIsWithinReach && !bIsLiftingOffSurface;
 		float TimeSinceSupported = MaxUnsupportedTimeBeforeFalling;
-		SimBlackboard->TryGet(KinematicBlackboard::TimeSinceSupported, TimeSinceSupported);
+		SimBlackboard->TryGet(CommonBlackboard::TimeSinceSupported, TimeSinceSupported);
 		if (bIsSupported)
 		{
-			SimBlackboard->Set(KinematicBlackboard::TimeSinceSupported, 0.0f);
+			SimBlackboard->Set(CommonBlackboard::TimeSinceSupported, 0.0f);
 		}
 		else
 		{
 			TimeSinceSupported += DeltaSeconds;
-			SimBlackboard->Set(KinematicBlackboard::TimeSinceSupported, TimeSinceSupported);
+			SimBlackboard->Set(CommonBlackboard::TimeSinceSupported, TimeSinceSupported);
 			bIsSupported = TimeSinceSupported < MaxUnsupportedTimeBeforeFalling;
 		}
 
@@ -174,7 +173,7 @@ void UPhysicsDrivenWalkingMode::OnSimulationTick(const FSimulationTickParams& Pa
 
 		if (bIsSupported)
 		{
-			OutputState.MovementEndState.NextModeName = KinematicModeNames::Walking;
+			OutputState.MovementEndState.NextModeName = DefaultModeNames::Walking;
 			OutputState.MovementEndState.RemainingMs = 0.0f;
 
 			OutputSyncState.MoveDirectionIntent = ProposedMove.bHasDirIntent ? ProposedMove.DirectionIntent : FVector::ZeroVector;
@@ -187,13 +186,13 @@ void UPhysicsDrivenWalkingMode::OnSimulationTick(const FSimulationTickParams& Pa
 		else
 		{
 			// Blocking hit but not supported
-			SwitchToState(KinematicModeNames::Falling, Params, OutputState);
+			SwitchToState(DefaultModeNames::Falling, Params, OutputState);
 		}
 	}
 	else
 	{
 		// No water or floor not found
-		SwitchToState(KinematicModeNames::Falling, Params, OutputState);
+		SwitchToState(DefaultModeNames::Falling, Params, OutputState);
 	}
 }
 
@@ -207,8 +206,8 @@ bool UPhysicsDrivenWalkingMode::AttemptTeleport(USceneComponent* UpdatedComponen
 		nullptr); // no movement base
 
 	// TODO: instead of invalidating it, consider checking for a floor. Possibly a dynamic base?
-	GetBlackboard_Mutable()->Invalidate(KinematicBlackboard::LastFloorResult);
-	GetBlackboard_Mutable()->Invalidate(KinematicBlackboard::LastMovementBase);
+	GetBlackboard_Mutable()->Invalidate(CommonBlackboard::LastFloorResult);
+	GetBlackboard_Mutable()->Invalidate(CommonBlackboard::LastMovementBase);
 
 	return true;
 }
diff --git a/Engine/Plugins/Experimental/Mover/Source/Mover/Private/PhysicsMover/Transitions/PhysicsJumpCheck.cpp b/Engine/Plugins/Experimental/Mover/Source/Mover/Private/PhysicsMover/Transitions/PhysicsJumpCheck.cpp
--- a/Engine/Plugins/Experimental/Mover/Source/Mover/Private/PhysicsMover/Transitions/PhysicsJumpCheck.cpp
+++ b/Engine/Plugins/Experimental/Mover/Source/Mover/Private/PhysicsMover/Transitions/PhysicsJumpCheck.cpp
@@ -1,7 +1,7 @@
 // Copyright Epic Games, Inc. All Rights Reserved.
 
 #include "PhysicsMover/Transitions/PhysicsJumpCheck.h"
-#include "Kinematic/LayeredMoves/BasicLayeredMoves.h"
+#include "DefaultMovementSet/LayeredMoves/BasicLayeredMoves.h"
 #include "MoverComponent.h"
 
 #if WITH_EDITOR
@@ -11,7 +11,7 @@
 UPhysicsJumpCheck::UPhysicsJumpCheck(const FObjectInitializer& ObjectInitializer)
 	: Super(ObjectInitializer)
 {
-	TransitionToMode = KinematicModeNames::Falling;
+	TransitionToMode = DefaultModeNames::Falling;
 }
 
 FTransitionEvalResult UPhysicsJumpCheck::OnEvaluate(const FSimulationTickParams& Params) const
@@ -19,10 +19,10 @@ FTransitionEvalResult UPhysicsJumpCheck::OnEvaluate(const FSimulationTickParams&
 	FTransitionEvalResult EvalResult; 
 
 	const FMoverTickStartData& StartState = Params.StartState;
-	const FKinematicDefaultInputs* KinematicInputs = StartState.InputCmd.InputCollection.FindDataByType<FKinematicDefaultInputs>();
-	check(KinematicInputs);
+	const FCharacterDefaultInputs* CharacterInputs = StartState.InputCmd.InputCollection.FindDataByType<FCharacterDefaultInputs>();
+	check(CharacterInputs);
 
-	if (KinematicInputs->bIsJumpJustPressed)
+	if (CharacterInputs->bIsJumpJustPressed)
 	{
 		EvalResult.NextMode = TransitionToMode;
 	}
diff --git a/Engine/Plugins/Experimental/Mover/Source/Mover/Public/DefaultMovementSet/LayeredMoves/AnimRootMotionLayeredMove.h b/Engine/Plugins/Experimental/Mover/Source/Mover/Public/DefaultMovementSet/LayeredMoves/AnimRootMotionLayeredMove.h
new file mode 100644
--- /dev/null
+++ b/Engine/Plugins/Experimental/Mover/Source/Mover/Public/DefaultMovementSet/LayeredMoves/AnimRootMotionLayeredMove.h
@@ -0,0 +1,44 @@
+// Copyright Epic Games, Inc. All Rights Reserved.
+
+#pragma once
+
+
+#include "LayeredMove.h"
+#include "AnimRootMotionLayeredMove.generated.h"
+
+class UAnimMontage;
+
+
+/** Anim Root Motion Move: handles root motion from a montage played on the associated mesh */
+USTRUCT(BlueprintType)
+struct MOVER_API FLayeredMove_AnimRootMotion : public FLayeredMoveBase
+{
+	GENERATED_USTRUCT_BODY()
+
+	FLayeredMove_AnimRootMotion();
+	virtual ~FLayeredMove_AnimRootMotion() {}
+
+	UPROPERTY(EditAnywhere, BlueprintReadWrite, Category = Mover)
+	TObjectPtr<UAnimMontage> Montage;
+
+	// Montage position when started (in unscaled seconds). 
+	UPROPERTY(EditAnywhere, BlueprintReadWrite, Category = Mover)
+	float StartingMontagePosition;
+
+	// Rate at which this montage is intended to play
+	UPROPERTY(EditAnywhere, BlueprintReadWrite, Category = Mover)
+	float PlayRate;
+
+	// Generate a movement 
+	virtual bool GenerateMove(const FMoverTickStartData& StartState, const FMoverTimeStep& TimeStep, const UMoverComponent* MoverComp, UMoverBlackboard* SimBlackboard, FProposedMove& OutProposedMove) override;
+
+	virtual FLayeredMoveBase* Clone() const override;
+
+	virtual void NetSerialize(FArchive& Ar) override;
+
+	virtual UScriptStruct* GetScriptStruct() const override;
+
+	virtual FString ToSimpleString() const override;
+
+	virtual void AddReferencedObjects(class FReferenceCollector& Collector) override;
+};
diff --git a/Engine/Plugins/Experimental/Mover/Source/Mover/Public/DefaultMovementSet/LayeredMoves/BasicLayeredMoves.h b/Engine/Plugins/Experimental/Mover/Source/Mover/Public/DefaultMovementSet/LayeredMoves/BasicLayeredMoves.h
new file mode 100644
--- /dev/null
+++ b/Engine/Plugins/Experimental/Mover/Source/Mover/Public/DefaultMovementSet/LayeredMoves/BasicLayeredMoves.h
@@ -0,0 +1,386 @@
+// Copyright Epic Games, Inc. All Rights Reserved.
+
+#pragma once
+
+
+#include "LayeredMove.h"
+#include "BasicLayeredMoves.generated.h"
+
+class UCurveVector;
+class UCurveFloat;
+
+UENUM(BlueprintType)
+enum class ELayeredMove_ConstantVelocitySettingsFlags : uint8
+{
+	NoFlags						= 0,
+
+	// Velocity direction is determined when this move starts, relative to the target actor. This flag is incompatible with other Velocity**** flags
+	VelocityStartRelative		= (1 << 0),
+	// Velocity direction is always relative to the target actor, even if its orientation changes. This flag is incompatible with other Velocity**** flags
+	VelocityAlwaysRelative		= (1 << 1),
+
+};
+
+/** Linear Velocity: A method of inducing a straight-line velocity on an actor over time  */
+USTRUCT(BlueprintType)
+struct MOVER_API FLayeredMove_LinearVelocity : public FLayeredMoveBase
+{
+	GENERATED_USTRUCT_BODY()
+
+	FLayeredMove_LinearVelocity();
+
+	virtual ~FLayeredMove_LinearVelocity() {}
+
+	// Units per second, could be worldspace vs relative depending on Flags
+	UPROPERTY(EditAnywhere, BlueprintReadWrite, Category = Mover)
+	FVector Velocity;
+
+	// Optional curve float for controlling the magnitude of the velocity applied to the actor over the duration of the move
+	UPROPERTY(EditAnywhere, BlueprintReadWrite, Category = Mover)
+	TObjectPtr<UCurveFloat> MagnitudeOverTime;
+	
+	// @see ELayeredMove_ConstantVelocitySettingsFlags
+	UPROPERTY(EditAnywhere, BlueprintReadWrite, Category = Mover)
+	uint8 SettingsFlags;
+
+	// Generate a movement 
+	virtual bool GenerateMove(const FMoverTickStartData& StartState, const FMoverTimeStep& TimeStep, const UMoverComponent* MoverComp, UMoverBlackboard* SimBlackboard, FProposedMove& OutProposedMove) override;
+
+	virtual FLayeredMoveBase* Clone() const override;
+
+	virtual void NetSerialize(FArchive& Ar) override;
+
+	virtual UScriptStruct* GetScriptStruct() const override;
+
+	virtual FString ToSimpleString() const override;
+
+	virtual void AddReferencedObjects(class FReferenceCollector& Collector) override;
+};
+
+template<>
+struct TStructOpsTypeTraits< FLayeredMove_LinearVelocity > : public TStructOpsTypeTraitsBase2< FLayeredMove_LinearVelocity >
+{
+	enum
+	{
+		//WithNetSerializer = true,
+		WithCopy = true
+	};
+};
+
+
+/** Jump Impulse: introduces an instantaneous upwards change in velocity. This overrides the existing 'up' component of the actor's current velocity */
+USTRUCT(BlueprintType)
+struct MOVER_API FLayeredMove_JumpImpulse : public FLayeredMoveBase
+{
+	GENERATED_USTRUCT_BODY()
+
+	FLayeredMove_JumpImpulse();
+
+	virtual ~FLayeredMove_JumpImpulse() {}
+
+	// Units per second, in whatever direction the target actor considers 'up'
+	UPROPERTY(EditAnywhere, BlueprintReadWrite, Category = Mover)
+	float UpwardsSpeed;
+
+	// Generate a movement 
+	virtual bool GenerateMove(const FMoverTickStartData& StartState, const FMoverTimeStep& TimeStep, const UMoverComponent* MoverComp, UMoverBlackboard* SimBlackboard, FProposedMove& OutProposedMove) override;
+
+	virtual FLayeredMoveBase* Clone() const override;
+
+	virtual void NetSerialize(FArchive& Ar) override;
+
+	virtual UScriptStruct* GetScriptStruct() const override;
+
+	virtual FString ToSimpleString() const override;
+
+	virtual void AddReferencedObjects(class FReferenceCollector& Collector) override;
+};
+
+template<>
+struct TStructOpsTypeTraits< FLayeredMove_JumpImpulse > : public TStructOpsTypeTraitsBase2< FLayeredMove_JumpImpulse >
+{
+	enum
+	{
+		//WithNetSerializer = true,
+		WithCopy = true
+	};
+};
+
+/** JumpTo: Moves this actor in a more jump-like manner - specifying height and distance of jump rather than just upwards speed
+  * Note: this layered move is only intended for Mover actors using Z as it's up direction
+  */
+USTRUCT(BlueprintType)
+struct MOVER_API FLayeredMove_JumpTo : public FLayeredMoveBase
+{
+	GENERATED_USTRUCT_BODY()
+
+	FLayeredMove_JumpTo();
+	virtual ~FLayeredMove_JumpTo() {}
+	
+	// Distance this jump impulse is supposed to reach
+	UPROPERTY(EditAnywhere, BlueprintReadWrite, Category = Mover)
+	float JumpDistance;
+
+	// Height this jump impulse is supposed to reach
+	UPROPERTY(EditAnywhere, BlueprintReadWrite, Category = Mover)
+	float JumpHeight;
+
+	// If true this jump will use the actor's rotation for jump direction instead of the Jump Rotation property.
+	UPROPERTY(EditAnywhere, BlueprintReadWrite, Category = Mover)
+	bool bUseActorRotation;
+	
+	// Direction to jump in. Only used if bUseActorRotation is false
+	UPROPERTY(EditAnywhere, BlueprintReadWrite, Category = Mover)
+	FRotator JumpRotation;
+	
+	// Optional CurveVector used to offset the actor from the path
+	UPROPERTY(EditAnywhere, BlueprintReadWrite, Category = Mover)
+	TObjectPtr<UCurveVector> PathOffsetCurve;
+	
+	// Optional CurveFloat to modify jump impulse over the duration of the impulse
+	UPROPERTY(EditAnywhere, BlueprintReadWrite, Category = Mover)
+	TObjectPtr<UCurveFloat> TimeMappingCurve;
+	
+	// Generate a movement 
+	virtual bool GenerateMove(const FMoverTickStartData& StartState, const FMoverTimeStep& TimeStep, const UMoverComponent* MoverComp, UMoverBlackboard* SimBlackboard, FProposedMove& OutProposedMove) override;
+	
+	virtual FLayeredMoveBase* Clone() const override;
+	
+	virtual void NetSerialize(FArchive& Ar) override;
+	
+	virtual UScriptStruct* GetScriptStruct() const override;
+	
+	virtual FString ToSimpleString() const override;
+
+	virtual void AddReferencedObjects(class FReferenceCollector& Collector) override;
+protected:
+	// helper function to apply movement vector offset from the PathOffsetCurve
+	FVector GetPathOffset(const float MoveFraction) const;
+
+	// Helper function to apply TimeMappingCurve to the layered move
+	float EvaluateFloatCurveAtFraction(const UCurveFloat& Curve, const float Fraction) const;
+
+	// Helper function for calculating position in jump given move fraction of total time for move
+	FVector GetRelativeLocation(float MoveFraction, const FRotator& Rotator) const;
+};
+
+template<>
+struct TStructOpsTypeTraits< FLayeredMove_JumpTo > : public TStructOpsTypeTraitsBase2< FLayeredMove_JumpTo >
+{
+	enum
+	{
+		//WithNetSerializer = true,
+		WithCopy = true
+	};
+};
+
+/** Teleport: instantly moves an actor to a new location */
+USTRUCT(BlueprintType)
+struct MOVER_API FLayeredMove_Teleport : public FLayeredMoveBase
+{
+	GENERATED_USTRUCT_BODY()
+
+	FLayeredMove_Teleport();
+	virtual ~FLayeredMove_Teleport() {}
+
+	// Location to teleport to, in world space
+	UPROPERTY(EditAnywhere, BlueprintReadWrite, Category = Mover)
+	FVector TargetLocation;
+
+	// Generate a movement 
+	virtual bool GenerateMove(const FMoverTickStartData& StartState, const FMoverTimeStep& TimeStep, const UMoverComponent* MoverComp, UMoverBlackboard* SimBlackboard, FProposedMove& OutProposedMove) override;
+
+	virtual FLayeredMoveBase* Clone() const override;
+
+	virtual void NetSerialize(FArchive& Ar) override;
+
+	virtual UScriptStruct* GetScriptStruct() const override;
+
+	virtual FString ToSimpleString() const override;
+
+	virtual void AddReferencedObjects(class FReferenceCollector& Collector) override;
+};
+
+template<>
+struct TStructOpsTypeTraits< FLayeredMove_Teleport > : public TStructOpsTypeTraitsBase2< FLayeredMove_Teleport >
+{
+	enum
+	{
+		//WithNetSerializer = true,
+		WithCopy = true
+	};
+};
+
+
+/** MoveTo: Move Actor from the starting location to the target location over a duration of time.*/
+USTRUCT(BlueprintType)
+struct MOVER_API FLayeredMove_MoveTo : public FLayeredMoveBase
+{
+	GENERATED_USTRUCT_BODY()
+
+	FLayeredMove_MoveTo();
+	virtual ~FLayeredMove_MoveTo() {}
+
+	// Location to Start the MoveTo move from
+	UPROPERTY(EditAnywhere, BlueprintReadWrite, Category = Mover)
+	FVector StartLocation;
+	
+	// Location to move towards
+	UPROPERTY(EditAnywhere, BlueprintReadWrite, Category = Mover)
+	FVector TargetLocation;
+
+	// if true, will restrict speed to where the actor is expected to be (in regard to start, end and duration)
+	UPROPERTY(EditAnywhere, BlueprintReadWrite, Category = Mover)
+	bool bRestrictSpeedToExpected;
+	
+	// Optional CurveVector used to offset the actor from the path
+	UPROPERTY(EditAnywhere, BlueprintReadWrite, Category = Mover)
+	TObjectPtr<UCurveVector> PathOffsetCurve;
+
+	// Optional CurveFloat to apply to how fast the actor moves as they get closer to the target location
+	UPROPERTY(EditAnywhere, BlueprintReadWrite, Category = Mover)
+	TObjectPtr<UCurveFloat> TimeMappingCurve;
+	
+	// Generate a movement 
+	virtual bool GenerateMove(const FMoverTickStartData& StartState, const FMoverTimeStep& TimeStep, const UMoverComponent* MoverComp, UMoverBlackboard* SimBlackboard, FProposedMove& OutProposedMove) override;
+
+	virtual FLayeredMoveBase* Clone() const override;
+
+	virtual void NetSerialize(FArchive& Ar) override;
+
+	virtual UScriptStruct* GetScriptStruct() const override;
+
+	virtual FString ToSimpleString() const override;
+
+	virtual void AddReferencedObjects(class FReferenceCollector& Collector) override;
+	
+protected:
+	// helper function to apply movement vector offset from the PathOffsetCurve
+	FVector GetPathOffsetInWorldSpace(const float MoveFraction) const;
+
+	// Helper function to apply TimeMappingCurve to the layered move
+	float EvaluateFloatCurveAtFraction(const UCurveFloat& Curve, const float Fraction) const;
+};
+
+template<>
+struct TStructOpsTypeTraits< FLayeredMove_MoveTo > : public TStructOpsTypeTraitsBase2< FLayeredMove_MoveTo >
+{
+	enum
+	{
+		//WithNetSerializer = true,
+		WithCopy = true
+	};
+};
+
+
+/** MoveToDynamic: Move Actor from the starting location to the target location over a duration of time with a moving target location
+ *  You can either update TargetLocation manually to update the location this actor is moving towards or set an actor for this actor to move towards.
+ */
+USTRUCT(BlueprintType)
+struct MOVER_API FLayeredMove_MoveToDynamic : public FLayeredMove_MoveTo
+{
+	GENERATED_USTRUCT_BODY()
+
+	FLayeredMove_MoveToDynamic();
+	virtual ~FLayeredMove_MoveToDynamic() {}
+	
+	// Optional actor to move to. Note: this overrides the TargetLocation Property
+	UPROPERTY(EditAnywhere, BlueprintReadWrite, Category = Mover)
+	TObjectPtr<AActor> LocationActor;
+
+	// Generate a movement 
+	virtual bool GenerateMove(const FMoverTickStartData& StartState, const FMoverTimeStep& TimeStep, const UMoverComponent* MoverComp, UMoverBlackboard* SimBlackboard, FProposedMove& OutProposedMove) override;
+
+	virtual FLayeredMoveBase* Clone() const override;
+
+	virtual void NetSerialize(FArchive& Ar) override;
+
+	virtual UScriptStruct* GetScriptStruct() const override;
+
+	virtual FString ToSimpleString() const override;
+
+	virtual void AddReferencedObjects(class FReferenceCollector& Collector) override;
+};
+
+template<>
+struct TStructOpsTypeTraits< FLayeredMove_MoveToDynamic > : public TStructOpsTypeTraitsBase2< FLayeredMove_MoveToDynamic >
+{
+	enum
+	{
+		//WithNetSerializer = true,
+		WithCopy = true
+	};
+};
+
+/** RadialImpulse applies a velocity pulling or pushing away from a given world location to the target actor */
+USTRUCT(BlueprintType)
+struct MOVER_API FLayeredMove_RadialImpulse : public FLayeredMoveBase
+{
+	GENERATED_USTRUCT_BODY()
+
+	FLayeredMove_RadialImpulse();
+	virtual ~FLayeredMove_RadialImpulse() {}
+
+	// Location to pull or push actor from
+	UPROPERTY(EditAnywhere, BlueprintReadWrite, Category = Mover)
+	FVector Location;
+
+	// Optional Actor to pull or push this actor from. Note: this overrides the Location Property
+	UPROPERTY(EditAnywhere, BlueprintReadWrite, Category = Mover)
+	TObjectPtr<AActor> LocationActor;
+
+	// Radius from the location to apply radial velocity
+	UPROPERTY(EditAnywhere, BlueprintReadWrite, Category = Mover)
+	float Radius;
+
+	// Magnitude of velocity applied to actors in range
+	UPROPERTY(EditAnywhere, BlueprintReadWrite, Category = Mover)
+	float Magnitude;
+
+	// If true the velocity applied will push actor away from location, otherwise this layered move wil pull the actor towards it
+	UPROPERTY(EditAnywhere, BlueprintReadWrite, Category = Mover)
+	bool bIsPush;
+
+	// Whether to apply vertical velocity (based off mover components up direction) to the actor affected by this layered move
+	UPROPERTY(EditAnywhere, BlueprintReadWrite, Category = Mover)
+	bool bNoVerticalVelocity;
+
+	// Optional Curvefloat to dictate falloff of velocity as you get further from the velocity location
+	UPROPERTY(EditAnywhere, BlueprintReadWrite, Category = Mover)
+	TObjectPtr<UCurveFloat> DistanceFalloff;
+
+	// Optional Curvefloat to dictate magnitude of velocity applied over time
+	UPROPERTY(EditAnywhere, BlueprintReadWrite, Category = Mover)
+	TObjectPtr<UCurveFloat> MagnitudeOverTime;
+
+	// If true velocity added to the actor will be in the direction of FixedWorldDirection, otherwise it will be calculated from the position of the actor and velocity location.
+	UPROPERTY(EditAnywhere, BlueprintReadWrite, Category = Mover)
+	bool bUseFixedWorldDirection;
+
+	// Direction to apply velocity from the radial impulse in. This is only taken into account if UseFixedWorldDirection is true.
+	UPROPERTY(EditAnywhere, BlueprintReadWrite, Category = Mover)
+	FRotator FixedWorldDirection;
+	
+	// Generate a movement 
+	virtual bool GenerateMove(const FMoverTickStartData& StartState, const FMoverTimeStep& TimeStep, const UMoverComponent* MoverComp, UMoverBlackboard* SimBlackboard, FProposedMove& OutProposedMove) override;
+
+	virtual FLayeredMoveBase* Clone() const override;
+
+	virtual void NetSerialize(FArchive& Ar) override;
+
+	virtual UScriptStruct* GetScriptStruct() const override;
+
+	virtual FString ToSimpleString() const override;
+
+	virtual void AddReferencedObjects(class FReferenceCollector& Collector) override;
+};
+
+template<>
+struct TStructOpsTypeTraits< FLayeredMove_RadialImpulse > : public TStructOpsTypeTraitsBase2< FLayeredMove_RadialImpulse >
+{
+	enum
+	{
+		//WithNetSerializer = true,
+		WithCopy = true
+	};
+};
diff --git a/Engine/Plugins/Experimental/Mover/Source/Mover/Public/DefaultMovementSet/LayeredMoves/LaunchMove.h b/Engine/Plugins/Experimental/Mover/Source/Mover/Public/DefaultMovementSet/LayeredMoves/LaunchMove.h
new file mode 100644
--- /dev/null
+++ b/Engine/Plugins/Experimental/Mover/Source/Mover/Public/DefaultMovementSet/LayeredMoves/LaunchMove.h
@@ -0,0 +1,38 @@
+// Copyright Epic Games, Inc. All Rights Reserved.
+
+#pragma once
+
+#include "LayeredMove.h"
+#include "LaunchMove.generated.h"
+
+
+/** Launch Move: provides an impulse velocity to the actor after (optionally) forcing them into a particular movement mode */
+USTRUCT(BlueprintType)
+struct MOVER_API FLayeredMove_Launch : public FLayeredMoveBase
+{
+	GENERATED_USTRUCT_BODY()
+
+	FLayeredMove_Launch();
+	virtual ~FLayeredMove_Launch() {}
+
+	// Velocity to apply to the actor. Could be additive or overriding depending on MixMode setting.
+	UPROPERTY(EditAnywhere, BlueprintReadWrite, Category = Mover, meta=(ForceUnits="cm/s"))
+	FVector LaunchVelocity = FVector::ZeroVector;
+
+	// Optional movement mode name to force the actor into before applying the impulse velocity.
+	UPROPERTY(EditAnywhere, BlueprintReadWrite, Category = Mover)
+	FName ForceMovementMode = NAME_None;
+
+	// Generate a movement 
+	virtual bool GenerateMove(const FMoverTickStartData& StartState, const FMoverTimeStep& TimeStep, const UMoverComponent* MoverComp, UMoverBlackboard* SimBlackboard, FProposedMove& OutProposedMove) override;
+
+	virtual FLayeredMoveBase* Clone() const override;
+
+	virtual void NetSerialize(FArchive& Ar) override;
+
+	virtual UScriptStruct* GetScriptStruct() const override;
+
+	virtual FString ToSimpleString() const override;
+
+	virtual void AddReferencedObjects(class FReferenceCollector& Collector) override;
+};
diff --git a/Engine/Plugins/Experimental/Mover/Source/Mover/Public/DefaultMovementSet/LayeredMoves/MultiJumpLayeredMove.h b/Engine/Plugins/Experimental/Mover/Source/Mover/Public/DefaultMovementSet/LayeredMoves/MultiJumpLayeredMove.h
new file mode 100644
--- /dev/null
+++ b/Engine/Plugins/Experimental/Mover/Source/Mover/Public/DefaultMovementSet/LayeredMoves/MultiJumpLayeredMove.h
@@ -0,0 +1,64 @@
+﻿// Copyright Epic Games, Inc. All Rights Reserved.
+
+#pragma once
+
+#include "LayeredMove.h"
+#include "MultiJumpLayeredMove.generated.h"
+
+struct FMoverInputCmdContext;
+struct FMoverDefaultSyncState;
+class UCommonLegacyMovementSettings;
+
+/** MultiJump: layered move for handling multiple jumps without touching the ground i.e. a double jump.
+  * Note: This layered move ends automatically when the actor hits a valid floor so duration is expected to be less than 0. */
+USTRUCT(BlueprintType)
+struct MOVER_API FLayeredMove_MultiJump : public FLayeredMoveBase
+{
+	GENERATED_BODY()
+
+	FLayeredMove_MultiJump();
+	virtual ~FLayeredMove_MultiJump() {}
+
+	/** Maximum Jumps an actor could perform while in the air */
+	UPROPERTY(BlueprintReadWrite, EditAnywhere, Category = "Jumping", meta = (ClampMin = "1", UIMin = "1"))
+	int32 MaximumInAirJumps;
+
+	// Units per second, in whatever direction the target actor considers 'up'
+	UPROPERTY(EditAnywhere, BlueprintReadWrite, Category = "Jumping", meta = (ClampMin = "0", UIMin = "0", ForceUnits = "cm/s"))
+	float UpwardsSpeed;
+
+	// returns true if input context has state that the actor wants to jump
+	virtual bool WantsToJump(const FMoverInputCmdContext& InputCmd);
+	
+	// Generate a movement 
+	virtual bool GenerateMove(const FMoverTickStartData& StartState, const FMoverTimeStep& TimeStep, const UMoverComponent* MoverComp, UMoverBlackboard* SimBlackboard, FProposedMove& OutProposedMove) override;
+
+	virtual FLayeredMoveBase* Clone() const override;
+
+	virtual void NetSerialize(FArchive& Ar) override;
+
+	virtual UScriptStruct* GetScriptStruct() const override;
+
+	virtual FString ToSimpleString() const override;
+
+	virtual void AddReferencedObjects(class FReferenceCollector& Collector) override;
+
+protected:
+	UPROPERTY(BlueprintReadOnly, EditAnywhere, Category = "Jumping")
+	int32 JumpsInAirRemaining;
+
+	UPROPERTY(BlueprintReadOnly, EditAnywhere, Category = "Jumping")
+	float TimeOfLastJumpMS;
+
+	bool PerformJump(const FMoverDefaultSyncState* SyncState, const FMoverTimeStep& TimeStep, const UMoverComponent* MoverComp, FProposedMove& OutProposedMove);
+};
+
+template<>
+struct TStructOpsTypeTraits< FLayeredMove_MultiJump > : public TStructOpsTypeTraitsBase2< FLayeredMove_MultiJump >
+{
+	enum
+	{
+		//WithNetSerializer = true,
+		WithCopy = true
+	};
+};
diff --git a/Engine/Plugins/Experimental/Mover/Source/Mover/Public/DefaultMovementSet/Modes/FallingMode.h b/Engine/Plugins/Experimental/Mover/Source/Mover/Public/DefaultMovementSet/Modes/FallingMode.h
new file mode 100644
--- /dev/null
+++ b/Engine/Plugins/Experimental/Mover/Source/Mover/Public/DefaultMovementSet/Modes/FallingMode.h
@@ -0,0 +1,94 @@
+// Copyright Epic Games, Inc. All Rights Reserved.
+
+#pragma once
+
+#include "CoreMinimal.h"
+#include "MovementMode.h"
+#include "FallingMode.generated.h"
+
+class UCommonLegacyMovementSettings;
+struct FFloorCheckResult;
+
+// Fired after the actor lands on a valid surface. First param is the name of the mode this actor is in after landing. Second param is the hit result from hitting the floor.
+DECLARE_DYNAMIC_MULTICAST_DELEGATE_TwoParams(FMover_OnLanded, const FName&, NextMovementModeName, const FHitResult&, HitResult);
+
+/**
+ * FallingMode: a default movement mode for moving through the air and jumping, typically influenced by gravity and air control
+ */
+UCLASS(Blueprintable, BlueprintType)
+class MOVER_API UFallingMode : public UBaseMovementMode
+{
+	GENERATED_UCLASS_BODY()
+
+
+public:
+	UFUNCTION(BlueprintCallable, Category=Mover)
+	virtual void OnGenerateMove(const FMoverTickStartData& StartState, const FMoverTimeStep& TimeStep, FProposedMove& OutProposedMove) const override;
+
+	UFUNCTION(BlueprintCallable, Category=Mover)
+	virtual void OnSimulationTick(const FSimulationTickParams& Params, FMoverTickEndData& OutputState) override;
+
+	// Broadcast when this actor lands on a valid surface.
+	UPROPERTY(BlueprintAssignable, Category = Mover)
+	FMover_OnLanded OnLanded;
+	
+	/**
+	 * When falling, amount of movement control available to the actor.
+	 * 0 = no control, 1 = full control
+	 */
+	UPROPERTY(Category=Mover, EditAnywhere, BlueprintReadWrite, meta=(ClampMin="0", ClampMax="1.0"))
+	float AirControlPercentage;
+	
+	/**
+ 	 * Deceleration to apply to air movement when falling slower than terminal velocity.
+ 	 * Note: This is NOT applied to vertical velocity, only movement plane velocity
+ 	 */
+	UPROPERTY(Category=Mover, EditAnywhere, BlueprintReadWrite, meta=(ClampMin="0", ForceUnits = "cm/s^2"))
+	float FallingDeceleration;
+
+	/**
+     * Deceleration to apply to air movement when falling faster than terminal velocity
+	 * Note: This is NOT applied to vertical velocity, only movement plane velocity
+     */
+    UPROPERTY(Category=Mover, EditAnywhere, BlueprintReadWrite, meta=(ClampMin="0", ForceUnits = "cm/s^2"))
+    float OverTerminalSpeedFallingDeceleration;
+	
+	/**
+	 * If the actor's movement plane velocity is greater than this speed falling will start applying OverTerminalSpeedFallingDeceleration instead of FallingDeceleration
+	 * The expected behavior is to set OverTerminalSpeedFallingDeceleration higher than FallingDeceleration so the actor will slow down faster
+	 * when going over TerminalMovementPlaneSpeed.
+	 */
+	UPROPERTY(Category=Mover, EditAnywhere, BlueprintReadWrite, meta=(ClampMin="0", ForceUnits = "cm/s"))
+	float TerminalMovementPlaneSpeed;
+
+	/** When exceeding maximum vertical speed, should it be enforced via a hard clamp? If false, VerticalFallingDeceleration will be used for a smoother transition to the terminal speed limit. */
+	UPROPERTY(Category = Mover, EditAnywhere, BlueprintReadWrite)
+	bool bShouldClampTerminalVerticalSpeed;
+
+	/** Deceleration to apply to vertical velocity when it's greater than TerminalVerticalSpeed. Only used if bShouldClampTerminalVerticalSpeed is false. */
+	UPROPERTY(Category=Mover, EditAnywhere, BlueprintReadWrite, meta=(EditCondition="!bShouldClampTerminalVerticalSpeed", ClampMin="0", ForceUnits = "cm/s^2"))
+	float VerticalFallingDeceleration;
+	
+	/**
+	 * If the actors vertical velocity is greater than this speed VerticalFallingDeceleration will be applied to vertical velocity
+	 */
+	UPROPERTY(Category=Mover, EditAnywhere, BlueprintReadWrite, meta=(ClampMin="0", ForceUnits = "cm/s"))
+	float TerminalVerticalSpeed;
+	
+protected:
+	/**
+	 * Is called at the end of the tick in falling mode. Handles checking any landings that should occur and switching to specific modes
+	 * (i.e. landing on a walkable surface would switch to the walking movement mode) 
+	 */
+	UFUNCTION(BlueprintCallable, Category=Mover)
+	virtual void ProcessLanded(const FFloorCheckResult& FloorResult, FVector& Velocity, FRelativeBaseInfo& BaseInfo, FMoverTickEndData& TickEndData) const;
+	
+	virtual void OnRegistered(const FName ModeName) override;
+	virtual void OnUnregistered() override;
+
+	virtual bool AttemptTeleport(USceneComponent* UpdatedComponent, const FVector& TeleportPos, const FRotator& TeleportRot, const FVector& PriorVelocity, FMoverDefaultSyncState& OutputSyncState);
+
+	void CaptureFinalState(USceneComponent* UpdatedComponent, const FMoverDefaultSyncState& StartSyncState, const FFloorCheckResult& FloorResult, float DeltaSeconds, float DeltaSecondsUsed, FMoverDefaultSyncState& OutputSyncState, FMoverTickEndData& TickEndData, FMovementRecord& Record) const;
+
+	TObjectPtr<const UCommonLegacyMovementSettings> CommonLegacySettings;
+};
diff --git a/Engine/Plugins/Experimental/Mover/Source/Mover/Public/DefaultMovementSet/Modes/FlyingMode.h b/Engine/Plugins/Experimental/Mover/Source/Mover/Public/DefaultMovementSet/Modes/FlyingMode.h
new file mode 100644
--- /dev/null
+++ b/Engine/Plugins/Experimental/Mover/Source/Mover/Public/DefaultMovementSet/Modes/FlyingMode.h
@@ -0,0 +1,38 @@
+// Copyright Epic Games, Inc. All Rights Reserved.
+
+#pragma once
+
+#include "CoreMinimal.h"
+#include "MovementMode.h"
+#include "FlyingMode.generated.h"
+
+class UCommonLegacyMovementSettings;
+
+
+/**
+ * FlyingMode: a default movement mode for moving through the air freely, but still interacting with blocking geometry. The
+ * moving actor will remain upright vs the movement plane.
+ */
+UCLASS(Blueprintable, BlueprintType)
+class MOVER_API UFlyingMode : public UBaseMovementMode
+{
+	GENERATED_UCLASS_BODY()
+
+
+public:
+	UFUNCTION(BlueprintCallable, Category=Mover)
+	virtual void OnGenerateMove(const FMoverTickStartData& StartState, const FMoverTimeStep& TimeStep, FProposedMove& OutProposedMove) const override;
+
+	UFUNCTION(BlueprintCallable, Category=Mover)
+	virtual void OnSimulationTick(const FSimulationTickParams& Params, FMoverTickEndData& OutputState) override;
+
+protected:
+	virtual void OnRegistered(const FName ModeName) override;
+	virtual void OnUnregistered() override;
+
+	virtual bool AttemptTeleport(USceneComponent* UpdatedComponent, const FVector& TeleportPos, const FRotator& TeleportRot, const FMoverDefaultSyncState& StartingSyncState, FMoverTickEndData& Output);
+
+	void CaptureFinalState(USceneComponent* UpdatedComponent, FMovementRecord& Record, const FMoverDefaultSyncState& StartSyncState, FMoverDefaultSyncState& OutputSyncState, const float DeltaSeconds) const;
+
+	TObjectPtr<const UCommonLegacyMovementSettings> CommonLegacySettings;
+};
diff --git a/Engine/Plugins/Experimental/Mover/Source/Mover/Public/DefaultMovementSet/Modes/SwimmingMode.h b/Engine/Plugins/Experimental/Mover/Source/Mover/Public/DefaultMovementSet/Modes/SwimmingMode.h
new file mode 100644
--- /dev/null
+++ b/Engine/Plugins/Experimental/Mover/Source/Mover/Public/DefaultMovementSet/Modes/SwimmingMode.h
@@ -0,0 +1,221 @@
+// Copyright Epic Games, Inc. All Rights Reserved.
+
+#pragma once
+
+#include "CoreMinimal.h"
+#include "MovementMode.h"
+
+#include "SwimmingMode.generated.h"
+
+class UCommonLegacyMovementSettings;
+struct FFloorCheckResult;
+struct FRelativeBaseInfo;
+struct FMovementRecord;
+
+// Controls for the Swimming Movement
+USTRUCT(BlueprintType)
+struct FSwimmingControlSettings
+{
+	GENERATED_BODY()
+
+	FSwimmingControlSettings()
+	{
+	}
+
+	// At or below this depth, cancel and disallow crouching.
+	UPROPERTY(EditAnywhere, BlueprintReadWrite, Category = SwimmingControl)
+	float CancelCrouchImmersionDepth = 28.50f;
+
+	// At or below this depth, start swimming in DBNO
+	UPROPERTY(EditAnywhere, BlueprintReadWrite, Category = SwimmingControl)
+	float DBNOSwimImmersionDepth = 50.00f;
+
+	// Max acceleration.
+	UPROPERTY(EditAnywhere, BlueprintReadWrite, Category = SwimmingControl)
+	float MaxAcceleration = 350.00f;
+
+	// Max acceleration while sprinting.
+	UPROPERTY(EditAnywhere, BlueprintReadWrite, Category = SwimmingControl)
+	float MaxAccelerationSprinting = 490.00f;
+
+	// Braking deceleration (decel when letting go of input)
+	UPROPERTY(EditAnywhere, BlueprintReadWrite, Category = SwimmingControl)
+	float BrakingDeceleration = 400.00f;
+
+	// Max speed when not sprinting and moving normally (before water velocity is applied)
+	UPROPERTY(EditAnywhere, BlueprintReadWrite, Category = SwimmingControl)
+	float MaxNormalSpeed = 0.50f;
+
+	// Max sprint speed (before water velocity is applied)
+	UPROPERTY(EditAnywhere, BlueprintReadWrite, Category = SwimmingControl)
+	float MaxSprintSpeed = 1.00f;
+
+	// Min speed required (relative to water) to maintain sprint while jumping, otherwise will stop sprint (and change anims).
+	UPROPERTY(EditAnywhere, BlueprintReadWrite, Category = SwimmingControl)
+	float MinSprintJumpSpeed = -0.03f;
+
+	// If accel deviates from velocity by this angle while sprint jumping in air, stop sprint (and change anims).
+	UPROPERTY(EditAnywhere, BlueprintReadWrite, Category = SwimmingControl)
+	float SprintJumpAirAccelAngleLimit = 0.25f;
+
+	UPROPERTY(EditAnywhere, BlueprintReadWrite, Category = SwimmingControl)
+	float SprintRetriggerDelay = 210.00f;
+
+	UPROPERTY(EditAnywhere, BlueprintReadWrite, Category = SwimmingControl)
+	float SprintDelayAfterFiring = 385.00f;
+
+	// Max speed when targeting (before water velocity is applied)
+	UPROPERTY(EditAnywhere, BlueprintReadWrite, Category = SwimmingControl)
+	float MaxTargetingSpeed = 300.00f;
+
+	// Directional multiplier when moving mostly backwards
+	UPROPERTY(EditAnywhere, BlueprintReadWrite, Category = SwimmingControl)
+	float BackwardsSpeedMultiplier = 50.00f;
+
+	// Directional multiplier is applied when dot product of velocity and player facing direction is < this.
+	UPROPERTY(EditAnywhere, BlueprintReadWrite, Category = SwimmingControl)
+	float BackwardsSpeedCosAngle = 0.10f;
+
+	// Speed multiplier when moving off angle (velocity and acceleration divergent)
+	UPROPERTY(EditAnywhere, BlueprintReadWrite, Category = SwimmingControl)
+	float AngledSpeedMultiplier = 0.70f;
+
+	// Angled multiplier is applied when dot product of velocity and acceleration is < this.
+	UPROPERTY(EditAnywhere, BlueprintReadWrite, Category = SwimmingControl)
+	float AngledSpeedCosAngle = 250.00f;
+
+	// Friction, ie how floaty or snappy is changing direction.
+	UPROPERTY(EditAnywhere, BlueprintReadWrite, Category = SwimmingControl)
+	float Friction = 0.50f;
+
+	// Friction, ie how floaty or snappy is changing direction.
+	UPROPERTY(EditAnywhere, BlueprintReadWrite, Category = SwimmingControl)
+	float FrictionSprinting = -0.30f;
+
+	// If Dot(Velocity, Acceleration) < this value, apply FrictionDirectionChangeMultiplier to friction value used. Allows lower friction when changing direction hard, which slows velocity change.
+	UPROPERTY(EditAnywhere, BlueprintReadWrite, Category = SwimmingControl)
+	float FrictionDirectionChangeDot = 0.40f;
+
+	// Friction multiplier (usually < 1) when changing direction.
+	UPROPERTY(EditAnywhere, BlueprintReadWrite, Category = SwimmingControl)
+	float FrictionDirectionChangeMultiplier = 0.75f;
+
+	UPROPERTY(EditAnywhere, BlueprintReadWrite, Category = SwimmingControl)
+	float MaxSpeedUp = 500.00f;
+
+	UPROPERTY(EditAnywhere, BlueprintReadWrite, Category = SwimmingControl)
+	float MaxSpeedDown = 1000.00f;
+
+	UPROPERTY(EditAnywhere, BlueprintReadWrite, Category = SwimmingControl)
+	float MaxHorizontalEntrySpeed = 1200.00f;
+
+	// Multiplier to water force acceleration in direction of current.
+	UPROPERTY(EditAnywhere, BlueprintReadWrite, Category = SwimmingControl)
+	float WaterForceMultiplier = 2.00f;
+
+	// Multiplier applied on the top of WaterForceMultiplier, to water force acceleration in direction of current. Used only for inherited objects.
+	UPROPERTY(EditAnywhere, BlueprintReadWrite, Category = SwimmingControl)
+	float WaterForceSecondMultiplier = 1.00f;
+
+	// Max water force, after WaterVelocity * (WaterForceMultiplier * WaterForceSecondMultiplier).
+	UPROPERTY(EditAnywhere, BlueprintReadWrite, Category = SwimmingControl)
+	float MaxWaterForce = 400.00f;
+
+	// At or above this depth, use max velocity. Interps down to WaterVelocityMinMultiplier at wading depth (where player can start swimming)
+	UPROPERTY(EditAnywhere, BlueprintReadWrite, Category = SwimmingControl)
+	float WaterVelocityDepthForMax = 175.00f;
+
+	// Min velocity multiplier applied when depth equals min swimming depth (where they transition from wading to swimming). Interps between this and 1.0 at WaterVelocityDepthForMax.
+	UPROPERTY(EditAnywhere, BlueprintReadWrite, Category = SwimmingControl)
+	float WaterVelocityMinMultiplier = 0.50f;
+
+	// Max time step allowed, to prevent huge forces on hitches.
+	UPROPERTY(EditAnywhere, BlueprintReadWrite, Category = SwimmingControl)
+	float WaterSimMaxTimeStep = 0.10f;
+
+	// Simulation sub-step time allowed for higher quality movement (local player and server).
+	UPROPERTY(EditAnywhere, BlueprintReadWrite, Category = SwimmingControl)
+	float WaterSimSubStepTime = 0.05f;
+
+	// Bobbing: 
+	UPROPERTY(EditAnywhere, BlueprintReadWrite, Category = SwimmingControl)
+	float BobbingMaxForce = 3800.00f;
+
+	// Bobbing: Slow down strongly when within this tolerance of the ideal immersion depth. Normally we apply drag only when going away from the ideal depth, this allows some slowdown when approaching it.
+	UPROPERTY(EditAnywhere, BlueprintReadWrite, Category = SwimmingControl)
+	float BobbingIdealDepthTolerance = 7.50f;
+
+	// Bobbing: friction/drag opposed to downward velocity, linear multiplier per second
+	UPROPERTY(EditAnywhere, BlueprintReadWrite, Category = SwimmingControl)
+	float BobbingFrictionDown = 5.00f;
+
+	// Bobbing: friction/drag opposed to downward velocity, squared with velocity per second. Ramps up faster with higher speeds, less effect at low speeds.
+	UPROPERTY(EditAnywhere, BlueprintReadWrite, Category = SwimmingControl)
+	float BobbingExpDragDown = 0.05f;
+
+	// Bobbing: friction/drag opposed to downward velocity, linear multiplier per second. Only used when fully submerged (replaces other value).
+	UPROPERTY(EditAnywhere, BlueprintReadWrite, Category = SwimmingControl)
+	float BobbingFrictionDownSubmerged = 7.50f;
+
+	// Bobbing: friction/drag opposed to upward velocity, squared with velocity per second. Ramps up faster with higher speeds, less effect at low speeds.
+	UPROPERTY(EditAnywhere, BlueprintReadWrite, Category = SwimmingControl)
+	float BobbingExpDragDownSubmerged = 0.10f;
+
+	// Bobbing: friction/drag opposed to upward velocity, linear multiplier per second
+	UPROPERTY(EditAnywhere, BlueprintReadWrite, Category = SwimmingControl)
+	float BobbingFrictionUp = 7.00f;
+
+	// Bobbing: friction/drag opposed to upward velocity, squared with velocity per second. Ramps up faster with higher speeds, less effect at low speeds.
+	UPROPERTY(EditAnywhere, BlueprintReadWrite, Category = SwimmingControl)
+	float BobbingExpDragUp = 0.20f;
+
+	// Bobbing: friction multiplier, multiplies the fluid friction value.
+	UPROPERTY(EditAnywhere, BlueprintReadWrite, Category = SwimmingControl)
+	float BobbingFrictionMultiplier = 1.f;
+	
+	// Bobbing: multiplier for the exponential drag.
+	UPROPERTY(EditAnywhere, BlueprintReadWrite, Category = SwimmingControl)
+	float BobbingExpDragMultiplier = 1.f;
+
+	// Bobbing: multiplier when in sprint boost to keep from popping up and out as much.
+	UPROPERTY(EditAnywhere, BlueprintReadWrite, Category = SwimmingControl)
+	float BoostDragMultiplier = 1.85f;
+
+	// Multiplies ground's jump speed
+	UPROPERTY(EditAnywhere, BlueprintReadWrite, Category = SwimmingControl)
+	float JumpMultiplier = 1.5f;
+};
+
+/**
+ * SwimmingMode: a default movement mode for traversing water volumes
+ */
+UCLASS(Blueprintable, BlueprintType)
+class MOVER_API USwimmingMode : public UBaseMovementMode
+{
+	GENERATED_UCLASS_BODY()
+
+public:
+
+	UFUNCTION(BlueprintCallable, Category=Mover)
+	virtual void OnGenerateMove(const FMoverTickStartData& StartState, const FMoverTimeStep& TimeStep, FProposedMove& OutProposedMove) const override;
+
+	UFUNCTION(BlueprintCallable, Category=Mover)
+	virtual void OnSimulationTick(const FSimulationTickParams& Params, FMoverTickEndData& OutputState) override;
+
+	UPROPERTY(EditAnywhere, BlueprintReadOnly, Category = "Water Control")
+	FSwimmingControlSettings SurfaceSwimmingWaterControlSettings;
+
+protected:
+	virtual void OnRegistered(const FName ModeName) override; 
+	virtual void OnUnregistered() override;
+
+	virtual bool AttemptJump(float UpwardsSpeed, FMoverTickEndData& Output);
+	virtual bool AttemptTeleport(USceneComponent* UpdatedComponent, const FVector& TeleportPos, const FRotator& TeleportRot, const FVector& PriorVelocity, FMoverTickEndData& Output);
+
+
+	void CaptureFinalState(USceneComponent* UpdatedComponent, const FFloorCheckResult& FloorResult, const FMovementRecord& Record, FMoverDefaultSyncState& OutputSyncState) const;
+
+	FRelativeBaseInfo UpdateFloorAndBaseInfo(const FFloorCheckResult& FloorResult) const;
+
+	TObjectPtr<const UCommonLegacyMovementSettings> CommonLegacySettings;
+};
diff --git a/Engine/Plugins/Experimental/Mover/Source/Mover/Public/DefaultMovementSet/Modes/WalkingMode.h b/Engine/Plugins/Experimental/Mover/Source/Mover/Public/DefaultMovementSet/Modes/WalkingMode.h
new file mode 100644
--- /dev/null
+++ b/Engine/Plugins/Experimental/Mover/Source/Mover/Public/DefaultMovementSet/Modes/WalkingMode.h
@@ -0,0 +1,56 @@
+// Copyright Epic Games, Inc. All Rights Reserved.
+
+#pragma once
+
+#include "CoreMinimal.h"
+#include "MovementMode.h"
+#include "MoveLibrary/ModularMovement.h"
+#include "WalkingMode.generated.h"
+
+class UCommonLegacyMovementSettings;
+struct FFloorCheckResult;
+struct FRelativeBaseInfo;
+struct FMovementRecord;
+
+/**
+ * WalkingMode: a default movement mode for traversing surfaces and movement bases (walking, running, sneaking, etc.)
+ */
+UCLASS(Blueprintable, BlueprintType)
+class MOVER_API UWalkingMode : public UBaseMovementMode
+{
+	GENERATED_UCLASS_BODY()
+
+public:
+	/**
+ 	 * Whether the actor ignores changes in rotation of the base it is standing on.
+ 	 * If true, the actor maintains current world rotation.
+ 	 * If false, the actor rotates with the moving base.
+ 	 */
+	UPROPERTY(BlueprintReadWrite, EditAnywhere, Category=Mover)
+	bool bIgnoreBaseRotation = true;
+	
+	UFUNCTION(BlueprintCallable, Category=Mover)
+	virtual void OnGenerateMove(const FMoverTickStartData& StartState, const FMoverTimeStep& TimeStep, FProposedMove& OutProposedMove) const override;
+
+	UFUNCTION(BlueprintCallable, Category=Mover)
+	virtual void OnSimulationTick(const FSimulationTickParams& Params, FMoverTickEndData& OutputState) override;
+
+
+protected:
+	/** Optional modular object for generating rotation towards desired orientation. If not specified, linear interpolation will be used. */
+	UPROPERTY(BlueprintReadWrite, EditAnywhere, Instanced, Category=Mover, meta=(MustImplement="/Script/Mover.TurnGeneratorInterface"))
+	TObjectPtr<UObject> TurnGenerator;
+
+	virtual void OnRegistered(const FName ModeName) override; 
+	virtual void OnUnregistered() override;
+
+	virtual bool AttemptJump(float UpwardsSpeed, FMoverTickEndData& Output);
+	virtual bool AttemptTeleport(USceneComponent* UpdatedComponent, const FVector& TeleportPos, const FRotator& TeleportRot, const FVector& PriorVelocity, FMoverTickEndData& Output);
+
+
+	void CaptureFinalState(USceneComponent* UpdatedComponent, const FFloorCheckResult& FloorResult, const FMovementRecord& Record, FMoverDefaultSyncState& OutputSyncState) const;
+
+	FRelativeBaseInfo UpdateFloorAndBaseInfo(const FFloorCheckResult& FloorResult) const;
+
+	TObjectPtr<const UCommonLegacyMovementSettings> CommonLegacySettings;
+};
diff --git a/Engine/Plugins/Experimental/Mover/Source/Mover/Public/DefaultMovementSet/Settings/CommonLegacyMovementSettings.h b/Engine/Plugins/Experimental/Mover/Source/Mover/Public/DefaultMovementSet/Settings/CommonLegacyMovementSettings.h
new file mode 100644
--- /dev/null
+++ b/Engine/Plugins/Experimental/Mover/Source/Mover/Public/DefaultMovementSet/Settings/CommonLegacyMovementSettings.h
@@ -0,0 +1,111 @@
+// Copyright Epic Games, Inc. All Rights Reserved.
+
+#pragma once
+
+#include "MovementMode.h"
+#include "CommonLegacyMovementSettings.generated.h"
+
+
+/**
+ * CommonLegacyMovementSettings: collection of settings that are shared between several of the legacy movement modes
+ */ 
+UCLASS(BlueprintType)
+class MOVER_API UCommonLegacyMovementSettings : public UObject, public IMovementSettingsInterface
+{
+	GENERATED_BODY()
+
+	virtual FString GetDisplayName() const override { return GetName(); }
+
+public:
+	// What movement mode to use when on the ground.
+	UPROPERTY(Category="General", EditAnywhere, BlueprintReadWrite)
+	FName GroundMovementModeName = DefaultModeNames::Walking;
+
+	// What movement mode to use when airborne.
+	UPROPERTY(Category="General", EditAnywhere, BlueprintReadWrite)
+	FName AirMovementModeName = DefaultModeNames::Falling;
+
+	/** Walkable slope angle, represented as cosine(max slope angle) for performance reasons. Ex: for max slope angle of 30 degrees, value is cosine(30 deg) = 0.866 */
+	UPROPERTY(BlueprintReadWrite, EditAnywhere, Category="Ground Movement")
+	float MaxWalkSlopeCosine = 0.71f;
+
+	/** Max distance to scan for floor surfaces under a Mover actor */
+	UPROPERTY(BlueprintReadWrite, EditAnywhere, Category="Ground Movement", meta = (ClampMin = "0", UIMin = "0", ForceUnits = "cm"))
+	float FloorSweepDistance = 40.0f;
+
+	/** Mover actors will be able to step up onto or over obstacles shorter than this */
+	UPROPERTY(BlueprintReadWrite, EditAnywhere, Category="Ground Movement", meta = (ClampMin = "0", UIMin = "0", ForceUnits = "cm"))
+	float MaxStepHeight = 40.0f;
+
+	/** Maximum speed in the movement plane */
+	UPROPERTY(BlueprintReadWrite, EditAnywhere, Category="General", meta = (ClampMin = "0", UIMin = "0", ForceUnits = "cm/s"))
+	float MaxSpeed = 800.f;
+
+	/**
+	 * Setting that affects movement control. Higher values allow faster changes in direction. This can be used to simulate slippery
+	 * surfaces such as ice or oil by lowering the value (possibly based on the material the actor is standing on).
+	 */
+	UPROPERTY(BlueprintReadWrite, EditAnywhere, Category="General|Friction", meta = (ClampMin = "0", UIMin = "0"))
+	float GroundFriction = 8.0f;
+
+	/**
+	  * If true, BrakingFriction will be used to slow the character to a stop (when there is no Acceleration).
+	  * If false, braking uses the same friction passed to CalcVelocity() (ie GroundFriction when walking), multiplied by BrakingFrictionFactor.
+	  * This setting applies to all movement modes; if only desired in certain modes, consider toggling it when movement modes change.
+	  * @see BrakingFriction
+	  */
+	UPROPERTY(Category="General|Friction", EditDefaultsOnly, BlueprintReadWrite)
+	uint8 bUseSeparateBrakingFriction:1;
+
+	/**
+	 * Friction (drag) coefficient applied when braking (whenever Acceleration = 0, or if character is exceeding max speed); actual value used is this multiplied by BrakingFrictionFactor.
+	 * When braking, this property allows you to control how much friction is applied when moving across the ground, applying an opposing force that scales with current velocity.
+	 * Braking is composed of friction (velocity-dependent drag) and constant deceleration.
+	 * This is the current value, used in all movement modes; if this is not desired, override it or bUseSeparateBrakingFriction when movement mode changes.
+	 * @note Only used if bUseSeparateBrakingFriction setting is true, otherwise current friction such as GroundFriction is used.
+	 * @see bUseSeparateBrakingFriction, BrakingFrictionFactor, GroundFriction, BrakingDecelerationWalking
+	 */
+	UPROPERTY(Category="General|Friction", EditAnywhere, BlueprintReadWrite, meta=(ClampMin="0", UIMin="0", EditCondition="bUseSeparateBrakingFriction"))
+	float BrakingFriction = 8.0f;
+
+	/**
+	 * Factor used to multiply actual value of friction used when braking.
+	 * This applies to any friction value that is currently used, which may depend on bUseSeparateBrakingFriction.
+	 * @note This is 2 by default for historical reasons, a value of 1 gives the true drag equation.
+	 */
+	UPROPERTY(BlueprintReadWrite, EditAnywhere, Category="General|Friction", meta=(ClampMin="0", UIMin="0"))
+	float BrakingFrictionFactor = 2.0f;
+	
+	/** Default max linear rate of deceleration when there is no controlled input */
+	UPROPERTY(BlueprintReadWrite, EditAnywhere, Category="General", meta = (ClampMin = "0", UIMin = "0", ForceUnits = "cm/s^2"))
+	float Deceleration = 8000.f;
+
+	/** Default max linear rate of acceleration for controlled input. May be scaled based on magnitude of input. */
+	UPROPERTY(BlueprintReadWrite, EditAnywhere, Category="General", meta = (ClampMin = "0", UIMin = "0", ForceUnits = "cm/s^2"))
+	float Acceleration = 4000.f;
+
+	/** Maximum rate of turning rotation (degrees per second). Negative numbers indicate instant rotation and should cause rotation to snap instantly to desired direction. */
+	UPROPERTY(BlueprintReadWrite, EditAnywhere, Category="General", meta = (ClampMin = "-1", UIMin = "0", ForceUnits = "degrees/s"))
+	float TurningRate = 500.f;
+
+	/** Speeds velocity direction changes while turning, to reduce sliding */
+	UPROPERTY(BlueprintReadWrite, EditAnywhere, Category="General", meta = (ClampMin = "0", UIMin = "0", ForceUnits = "Multiplier"))
+	float TurningBoost = 8.f;
+
+	/** Instantaneous speed induced in an actor upon jumping */
+	UPROPERTY(BlueprintReadWrite, EditAnywhere, Category="Jumping", meta = (ClampMin = "0", UIMin = "0", ForceUnits = "cm/s"))
+	float JumpUpwardsSpeed = 500.0f;
+	
+	/** Depth at which the pawn starts swimming */
+	UPROPERTY(BlueprintReadWrite, EditAnywhere, Category = "Swimming", meta = (Units = "cm"))
+	float SwimmingStartImmersionDepth = 64.5f;
+
+	/** Depth at which the pawn will float when in water */
+	UPROPERTY(BlueprintReadWrite, EditAnywhere, Category = "Swimming", meta = (Units = "cm"))
+	float SwimmingIdealImmersionDepth = 51.66f;
+
+	/** Depth at which the pawn stops swimming */
+	UPROPERTY(BlueprintReadWrite, EditAnywhere, Category = "Swimming", meta = (Units = "cm"))
+	float SwimmingStopImmersionDepth = 39.9f;
+	
+};
\ No newline at end of file
diff --git a/Engine/Plugins/Experimental/Mover/Source/Mover/Public/Kinematic/LayeredMoves/AnimRootMotionLayeredMove.h b/Engine/Plugins/Experimental/Mover/Source/Mover/Public/Kinematic/LayeredMoves/AnimRootMotionLayeredMove.h
deleted file mode 100644
diff --git a/Engine/Plugins/Experimental/Mover/Source/Mover/Public/Kinematic/LayeredMoves/BasicLayeredMoves.h b/Engine/Plugins/Experimental/Mover/Source/Mover/Public/Kinematic/LayeredMoves/BasicLayeredMoves.h
deleted file mode 100644
diff --git a/Engine/Plugins/Experimental/Mover/Source/Mover/Public/Kinematic/LayeredMoves/LaunchMove.h b/Engine/Plugins/Experimental/Mover/Source/Mover/Public/Kinematic/LayeredMoves/LaunchMove.h
deleted file mode 100644
diff --git a/Engine/Plugins/Experimental/Mover/Source/Mover/Public/Kinematic/LayeredMoves/MultiJumpLayeredMove.h b/Engine/Plugins/Experimental/Mover/Source/Mover/Public/Kinematic/LayeredMoves/MultiJumpLayeredMove.h
deleted file mode 100644
diff --git a/Engine/Plugins/Experimental/Mover/Source/Mover/Public/Kinematic/Modes/FallingMode.h b/Engine/Plugins/Experimental/Mover/Source/Mover/Public/Kinematic/Modes/FallingMode.h
deleted file mode 100644
diff --git a/Engine/Plugins/Experimental/Mover/Source/Mover/Public/Kinematic/Modes/FlyingMode.h b/Engine/Plugins/Experimental/Mover/Source/Mover/Public/Kinematic/Modes/FlyingMode.h
deleted file mode 100644
diff --git a/Engine/Plugins/Experimental/Mover/Source/Mover/Public/Kinematic/Modes/SwimmingMode.h b/Engine/Plugins/Experimental/Mover/Source/Mover/Public/Kinematic/Modes/SwimmingMode.h
deleted file mode 100644
diff --git a/Engine/Plugins/Experimental/Mover/Source/Mover/Public/Kinematic/Modes/WalkingMode.h b/Engine/Plugins/Experimental/Mover/Source/Mover/Public/Kinematic/Modes/WalkingMode.h
deleted file mode 100644
diff --git a/Engine/Plugins/Experimental/Mover/Source/Mover/Public/Kinematic/Settings/CommonLegacyMovementSettings.h b/Engine/Plugins/Experimental/Mover/Source/Mover/Public/Kinematic/Settings/CommonLegacyMovementSettings.h
deleted file mode 100644
diff --git a/Engine/Plugins/Experimental/Mover/Source/Mover/Public/MoverDataModelTypes.h b/Engine/Plugins/Experimental/Mover/Source/Mover/Public/MoverDataModelTypes.h
--- a/Engine/Plugins/Experimental/Mover/Source/Mover/Public/MoverDataModelTypes.h
+++ b/Engine/Plugins/Experimental/Mover/Source/Mover/Public/MoverDataModelTypes.h
@@ -24,9 +24,9 @@ enum class EMoveInputType : uint8
 };
 
 
-// Data block containing all inputs that need to be authored and consumed for the kinematic Mover simulation
+// Data block containing all inputs that need to be authored and consumed for the default Mover character simulation
 USTRUCT(BlueprintType)
-struct MOVER_API FKinematicDefaultInputs : public FMoverDataStructBase
+struct MOVER_API FCharacterDefaultInputs : public FMoverDataStructBase
 {
 	GENERATED_USTRUCT_BODY()
 
@@ -78,7 +78,7 @@ public:
 	FVector GetOrientationIntentDir_WorldSpace() const;
 
 
-	FKinematicDefaultInputs()
+	FCharacterDefaultInputs()
 		: MoveInputType(EMoveInputType::Invalid)
 		, MoveInput(ForceInitToZero)
 		, OrientationIntent(ForceInitToZero)
@@ -92,9 +92,9 @@ public:
 	{
 	}
 
-	virtual ~FKinematicDefaultInputs() {}
+	virtual ~FCharacterDefaultInputs() {}
 
-	// @return newly allocated copy of this FKinematicDefaultInputs. Must be overridden by child classes
+	// @return newly allocated copy of this FCharacterDefaultInputs. Must be overridden by child classes
 	virtual FMoverDataStructBase* Clone() const override;
 
 	virtual bool NetSerialize(FArchive& Ar, UPackageMap* Map, bool& bOutSuccess) override;
@@ -109,7 +109,7 @@ public:
 };
 
 template<>
-struct TStructOpsTypeTraits< FKinematicDefaultInputs > : public TStructOpsTypeTraitsBase2< FKinematicDefaultInputs >
+struct TStructOpsTypeTraits< FCharacterDefaultInputs > : public TStructOpsTypeTraitsBase2< FCharacterDefaultInputs >
 {
 	enum
 	{
@@ -245,15 +245,15 @@ class MOVER_API UMoverDataModelBlueprintLibrary : public UBlueprintFunctionLibra
 {
 	GENERATED_BODY()
 
-public:	// FKinematicDefaultInputs
+public:	// FCharacterDefaultInputs
 
 	/** Sets move inputs from worldspace intent, as a per-axis magnitude in the range [-1,1] Zero vector indicates intent to stop. */
 	UFUNCTION(BlueprintCallable, Category = Mover)
-	static void SetMoveIntent(UPARAM(Ref) FKinematicDefaultInputs& Inputs, const FVector& WorldDirectionIntent);
+	static void SetMoveIntent(UPARAM(Ref) FCharacterDefaultInputs& Inputs, const FVector& WorldDirectionIntent);
 
 	/** Returns the move direction intent, if any, in world space */
 	UFUNCTION(BlueprintCallable, Category = Mover)
-	static FVector GetMoveDirectionIntentFromInputs(const FKinematicDefaultInputs& Inputs);
+	static FVector GetMoveDirectionIntentFromInputs(const FCharacterDefaultInputs& Inputs);
 
 
 public:	// FMoverDefaultSyncState
diff --git a/Engine/Plugins/Experimental/Mover/Source/Mover/Public/MoverSimulationTypes.h b/Engine/Plugins/Experimental/Mover/Source/Mover/Public/MoverSimulationTypes.h
--- a/Engine/Plugins/Experimental/Mover/Source/Mover/Public/MoverSimulationTypes.h
+++ b/Engine/Plugins/Experimental/Mover/Source/Mover/Public/MoverSimulationTypes.h
@@ -12,8 +12,8 @@
 #include "UObject/Interface.h"
 #include "MoverSimulationTypes.generated.h"
 
-
-namespace KinematicModeNames
+// Names for our default modes
+namespace DefaultModeNames
 {
 	const FName Walking = TEXT("Walking");
 	const FName Falling = TEXT("Falling");
@@ -22,7 +22,7 @@ namespace KinematicModeNames
 }
 
 // Commonly-used blackboard object keys
-namespace KinematicBlackboard
+namespace CommonBlackboard
 {
 	const FName LastFloorResult = TEXT("LastFloor");
 	const FName LastWaterResult = TEXT("LastWater");
@@ -77,7 +77,7 @@ public:
 	FMoverInputCmdContext()
 	{
 		// TODO: Consider eliminating this default input, considering non-character / non-controlled Mover actors (like platforms)
-		InputCollection.FindOrAddDataByType<FKinematicDefaultInputs>();
+		InputCollection.FindOrAddDataByType<FCharacterDefaultInputs>();
 	}
 
 	UScriptStruct* GetStruct() const { return StaticStruct(); }
diff --git a/Engine/Plugins/Experimental/Mover/Source/Mover/Public/PhysicsMover/Modes/PhysicsDrivenFallingMode.h b/Engine/Plugins/Experimental/Mover/Source/Mover/Public/PhysicsMover/Modes/PhysicsDrivenFallingMode.h
--- a/Engine/Plugins/Experimental/Mover/Source/Mover/Public/PhysicsMover/Modes/PhysicsDrivenFallingMode.h
+++ b/Engine/Plugins/Experimental/Mover/Source/Mover/Public/PhysicsMover/Modes/PhysicsDrivenFallingMode.h
@@ -3,7 +3,7 @@
 #pragma once
 
 #include "CoreMinimal.h"
-#include "Kinematic/Modes/FallingMode.h"
+#include "DefaultMovementSet/Modes/FallingMode.h"
 #include "PhysicsMover/PhysicsMoverSimulationTypes.h"
 
 #include "PhysicsDrivenFallingMode.generated.h"
diff --git a/Engine/Plugins/Experimental/Mover/Source/Mover/Public/PhysicsMover/Modes/PhysicsDrivenFlyingMode.h b/Engine/Plugins/Experimental/Mover/Source/Mover/Public/PhysicsMover/Modes/PhysicsDrivenFlyingMode.h
--- a/Engine/Plugins/Experimental/Mover/Source/Mover/Public/PhysicsMover/Modes/PhysicsDrivenFlyingMode.h
+++ b/Engine/Plugins/Experimental/Mover/Source/Mover/Public/PhysicsMover/Modes/PhysicsDrivenFlyingMode.h
@@ -3,7 +3,7 @@
 #pragma once
 
 #include "CoreMinimal.h"
-#include "Kinematic/Modes/FlyingMode.h"
+#include "DefaultMovementSet/Modes/FlyingMode.h"
 #include "PhysicsMover/PhysicsMoverSimulationTypes.h"
 
 #include "PhysicsDrivenFlyingMode.generated.h"
diff --git a/Engine/Plugins/Experimental/Mover/Source/Mover/Public/PhysicsMover/Modes/PhysicsDrivenSwimmingMode.h b/Engine/Plugins/Experimental/Mover/Source/Mover/Public/PhysicsMover/Modes/PhysicsDrivenSwimmingMode.h
--- a/Engine/Plugins/Experimental/Mover/Source/Mover/Public/PhysicsMover/Modes/PhysicsDrivenSwimmingMode.h
+++ b/Engine/Plugins/Experimental/Mover/Source/Mover/Public/PhysicsMover/Modes/PhysicsDrivenSwimmingMode.h
@@ -3,7 +3,7 @@
 #pragma once
 
 #include "CoreMinimal.h"
-#include "Kinematic/Modes/SwimmingMode.h"
+#include "DefaultMovementSet/Modes/SwimmingMode.h"
 #include "PhysicsMover/PhysicsMoverSimulationTypes.h"
 
 #include "PhysicsDrivenSwimmingMode.generated.h"
diff --git a/Engine/Plugins/Experimental/Mover/Source/Mover/Public/PhysicsMover/Modes/PhysicsDrivenWalkingMode.h b/Engine/Plugins/Experimental/Mover/Source/Mover/Public/PhysicsMover/Modes/PhysicsDrivenWalkingMode.h
--- a/Engine/Plugins/Experimental/Mover/Source/Mover/Public/PhysicsMover/Modes/PhysicsDrivenWalkingMode.h
+++ b/Engine/Plugins/Experimental/Mover/Source/Mover/Public/PhysicsMover/Modes/PhysicsDrivenWalkingMode.h
@@ -3,7 +3,7 @@
 #pragma once
 
 #include "CoreMinimal.h"
-#include "Kinematic/Modes/WalkingMode.h"
+#include "DefaultMovementSet/Modes/WalkingMode.h"
 #include "PhysicsMover/PhysicsMoverSimulationTypes.h"
 
 #include "PhysicsDrivenWalkingMode.generated.h"
diff --git a/Engine/Plugins/Experimental/Mover/Source/Mover/Public/PhysicsMover/PhysicsMoverSimulationTypes.h b/Engine/Plugins/Experimental/Mover/Source/Mover/Public/PhysicsMover/PhysicsMoverSimulationTypes.h
--- a/Engine/Plugins/Experimental/Mover/Source/Mover/Public/PhysicsMover/PhysicsMoverSimulationTypes.h
+++ b/Engine/Plugins/Experimental/Mover/Source/Mover/Public/PhysicsMover/PhysicsMoverSimulationTypes.h
@@ -109,7 +109,7 @@ public:
 
 	virtual ~FMovementSettingsInputs() {}
 
-	// @return newly allocated copy of this FKinematicDefaultInputs. Must be overridden by child classes
+	// @return newly allocated copy of this FMovementSettingsInputs. Must be overridden by child classes
 	virtual FMoverDataStructBase* Clone() const override;
 
 	virtual bool NetSerialize(FArchive& Ar, UPackageMap* Map, bool& bOutSuccess) override;
diff --git a/Engine/Plugins/Experimental/MoverExamples/Source/MoverExamples/Private/CharacterVariants/Ziplining/ZipliningMode.cpp b/Engine/Plugins/Experimental/MoverExamples/Source/MoverExamples/Private/CharacterVariants/Ziplining/ZipliningMode.cpp
--- a/Engine/Plugins/Experimental/MoverExamples/Source/MoverExamples/Private/CharacterVariants/Ziplining/ZipliningMode.cpp
+++ b/Engine/Plugins/Experimental/MoverExamples/Source/MoverExamples/Private/CharacterVariants/Ziplining/ZipliningMode.cpp
@@ -6,7 +6,7 @@
 #include "CharacterVariants/Ziplining/ZiplineInterface.h"
 #include "CharacterVariants/Ziplining/ZipliningTransitions.h"
 #include "Kismet/KismetSystemLibrary.h"
-#include "Kinematic/Settings/CommonLegacyMovementSettings.h"
+#include "DefaultMovementSet/Settings/CommonLegacyMovementSettings.h"
 #include "MoverLog.h"
 
 
@@ -155,13 +155,13 @@ void UZipliningMode::OnSimulationTick(const FSimulationTickParams& Params, FMove
 		// If we were unable to find a valid target zipline, refund all the time and let the actor fall
 		if (!StartPoint || !EndPoint)
 		{
-			FName DefaultAirMode = KinematicModeNames::Falling;
+			FName DefaultAirMode = DefaultModeNames::Falling;
 			if (UCommonLegacyMovementSettings* LegacySettings = MoverComp->FindSharedSettings_Mutable<UCommonLegacyMovementSettings>())
 			{
 				DefaultAirMode = LegacySettings->AirMovementModeName;
 			}
 
-			OutputState.MovementEndState.NextModeName = KinematicModeNames::Falling;
+			OutputState.MovementEndState.NextModeName = DefaultModeNames::Falling;
 			OutputState.MovementEndState.RemainingMs = Params.TimeStep.StepMs;
 			return;
 		}
@@ -231,7 +231,7 @@ void UZipliningMode::OnSimulationTick(const FSimulationTickParams& Params, FMove
 
 	if (bWillReachEndPosition)
 	{
-		FName DefaultAirMode = KinematicModeNames::Falling;
+		FName DefaultAirMode = DefaultModeNames::Falling;
 		if (UCommonLegacyMovementSettings* LegacySettings = MoverComp->FindSharedSettings_Mutable<UCommonLegacyMovementSettings>())
 		{
 			DefaultAirMode = LegacySettings->AirMovementModeName;
diff --git a/Engine/Plugins/Experimental/MoverExamples/Source/MoverExamples/Private/CharacterVariants/Ziplining/ZipliningTransitions.cpp b/Engine/Plugins/Experimental/MoverExamples/Source/MoverExamples/Private/CharacterVariants/Ziplining/ZipliningTransitions.cpp
--- a/Engine/Plugins/Experimental/MoverExamples/Source/MoverExamples/Private/CharacterVariants/Ziplining/ZipliningTransitions.cpp
+++ b/Engine/Plugins/Experimental/MoverExamples/Source/MoverExamples/Private/CharacterVariants/Ziplining/ZipliningTransitions.cpp
@@ -63,7 +63,7 @@ FTransitionEvalResult UZiplineEndTransition::OnEvaluate(const FSimulationTickPar
 {
 	FTransitionEvalResult EvalResult = FTransitionEvalResult::NoTransition;
 
-	if (const FKinematicDefaultInputs* DefaultInputs = Params.StartState.InputCmd.InputCollection.FindDataByType<FKinematicDefaultInputs>())
+	if (const FCharacterDefaultInputs* DefaultInputs = Params.StartState.InputCmd.InputCollection.FindDataByType<FCharacterDefaultInputs>())
 	{
 		if (DefaultInputs->bIsJumpJustPressed)
 		{
diff --git a/Engine/Plugins/Experimental/MoverExamples/Source/MoverExamples/Private/MoverExamplesCharacter.cpp b/Engine/Plugins/Experimental/MoverExamples/Source/MoverExamples/Private/MoverExamplesCharacter.cpp
--- a/Engine/Plugins/Experimental/MoverExamples/Source/MoverExamples/Private/MoverExamplesCharacter.cpp
+++ b/Engine/Plugins/Experimental/MoverExamples/Source/MoverExamples/Private/MoverExamplesCharacter.cpp
@@ -110,15 +110,15 @@ void AMoverExamplesCharacter::OnProduceInput(float DeltaMs, FMoverInputCmdContex
 	// answers and will depend on the game and its specific needs. In general, at this time, I'd recommend aim assist and lock on 
 	// targeting systems to happen /outside/ of the system, i.e, here. But I can think of scenarios where that may not be ideal too.
 
-	FKinematicDefaultInputs& DefaultKinematicInputs = OutInputCmd.InputCollection.FindOrAddMutableDataByType<FKinematicDefaultInputs>();
+	FCharacterDefaultInputs& CharacterInputs = OutInputCmd.InputCollection.FindOrAddMutableDataByType<FCharacterDefaultInputs>();
 
 	if (Controller == nullptr)
 	{
 		if (GetLocalRole() == ENetRole::ROLE_Authority && GetRemoteRole() == ENetRole::ROLE_SimulatedProxy)
 		{
-			static const FKinematicDefaultInputs DoNothingInput;
+			static const FCharacterDefaultInputs DoNothingInput;
 			// If we get here, that means this pawn is not currently possessed and we're choosing to provide default do-nothing input
-			DefaultKinematicInputs = DoNothingInput;
+			CharacterInputs = DoNothingInput;
 		}
 
 		// We don't have a local controller so we can't run the code below. This is ok. Simulated proxies will just use previous input when extrapolating
@@ -131,12 +131,12 @@ void AMoverExamplesCharacter::OnProduceInput(float DeltaMs, FMoverInputCmdContex
 		SpringComp->bUsePawnControlRotation = true;
 	}
 
-	DefaultKinematicInputs.ControlRotation = FRotator::ZeroRotator;
+	CharacterInputs.ControlRotation = FRotator::ZeroRotator;
 
 	APlayerController* PC = Cast<APlayerController>(Controller);
 	if (PC)
 	{
-		DefaultKinematicInputs.ControlRotation = PC->GetControlRotation();
+		CharacterInputs.ControlRotation = PC->GetControlRotation();
 	}
 
 	// Favor velocity input 
@@ -144,7 +144,7 @@ void AMoverExamplesCharacter::OnProduceInput(float DeltaMs, FMoverInputCmdContex
 
 	if (bUsingInputIntentForMove)
 	{
-		FRotator Rotator = DefaultKinematicInputs.ControlRotation;
+		FRotator Rotator = CharacterInputs.ControlRotation;
 		FVector FinalDirectionalIntent;
 		if (const UMoverComponent* MoverComponent = GetMoverComponent())
 		{
@@ -157,19 +157,19 @@ void AMoverExamplesCharacter::OnProduceInput(float DeltaMs, FMoverInputCmdContex
 			FinalDirectionalIntent = Rotator.RotateVector(CachedMoveInputIntent);
 		}
 		
-		DefaultKinematicInputs.SetMoveInput(EMoveInputType::DirectionalIntent, FinalDirectionalIntent);
+		CharacterInputs.SetMoveInput(EMoveInputType::DirectionalIntent, FinalDirectionalIntent);
 	}
 	else
 	{
-		DefaultKinematicInputs.SetMoveInput(EMoveInputType::Velocity, CachedMoveInputVelocity);
+		CharacterInputs.SetMoveInput(EMoveInputType::Velocity, CachedMoveInputVelocity);
 	}
 
 	static float RotationMagMin(1e-3);
 
-	const bool bHasAffirmativeMoveInput = (DefaultKinematicInputs.GetMoveInput().Size() >= RotationMagMin);
+	const bool bHasAffirmativeMoveInput = (CharacterInputs.GetMoveInput().Size() >= RotationMagMin);
 	
 	// Figure out intended orientation
-	DefaultKinematicInputs.OrientationIntent = FVector::ZeroVector;
+	CharacterInputs.OrientationIntent = FVector::ZeroVector;
 
 
 	if (bUsingInputIntentForMove && bHasAffirmativeMoveInput)
@@ -177,52 +177,52 @@ void AMoverExamplesCharacter::OnProduceInput(float DeltaMs, FMoverInputCmdContex
 		if (bOrientRotationToMovement)
 		{
 			// set the intent to the actors movement direction
-			DefaultKinematicInputs.OrientationIntent = DefaultKinematicInputs.GetMoveInput();
+			CharacterInputs.OrientationIntent = CharacterInputs.GetMoveInput();
 		}
 		else
 		{
 			// set intent to the the control rotation - often a player's camera rotation
-			DefaultKinematicInputs.OrientationIntent = DefaultKinematicInputs.ControlRotation.Vector();
+			CharacterInputs.OrientationIntent = CharacterInputs.ControlRotation.Vector();
 		}
 
-		LastAffirmativeMoveInput = DefaultKinematicInputs.GetMoveInput();
+		LastAffirmativeMoveInput = CharacterInputs.GetMoveInput();
 
 	}
 	else if (bMaintainLastInputOrientation)
 	{
 		// There is no movement intent, so use the last-known affirmative move input
-		DefaultKinematicInputs.OrientationIntent = LastAffirmativeMoveInput;
+		CharacterInputs.OrientationIntent = LastAffirmativeMoveInput;
 	}
 	
 	if (bShouldRemainVertical)
 	{
 		// canceling out any z intent if the actor is supposed to remain vertical
-		DefaultKinematicInputs.OrientationIntent = DefaultKinematicInputs.OrientationIntent.GetSafeNormal2D();
+		CharacterInputs.OrientationIntent = CharacterInputs.OrientationIntent.GetSafeNormal2D();
 	}
 
-	DefaultKinematicInputs.bIsJumpPressed = bIsJumpPressed;
-	DefaultKinematicInputs.bIsJumpJustPressed = bIsJumpJustPressed;
+	CharacterInputs.bIsJumpPressed = bIsJumpPressed;
+	CharacterInputs.bIsJumpJustPressed = bIsJumpJustPressed;
 
 	if (bShouldToggleFlying)
 	{
 		if (!bIsFlyingActive)
 		{
-			DefaultKinematicInputs.SuggestedMovementMode = KinematicModeNames::Flying;
+			CharacterInputs.SuggestedMovementMode = DefaultModeNames::Flying;
 		}
 		else
 		{
-			DefaultKinematicInputs.SuggestedMovementMode = KinematicModeNames::Falling;
+			CharacterInputs.SuggestedMovementMode = DefaultModeNames::Falling;
 		}
 
 		bIsFlyingActive = !bIsFlyingActive;
 	}
 	else
 	{
-		DefaultKinematicInputs.SuggestedMovementMode = NAME_None;
+		CharacterInputs.SuggestedMovementMode = NAME_None;
 	}
 
 	// Convert inputs to be relative to the current movement base (depending on options and state)
-	DefaultKinematicInputs.bUsingMovementBase = false;
+	CharacterInputs.bUsingMovementBase = false;
 
 	if (bUseBaseRelativeMovement)
 	{
@@ -234,15 +234,15 @@ void AMoverExamplesCharacter::OnProduceInput(float DeltaMs, FMoverInputCmdContex
 
 				FVector RelativeMoveInput, RelativeOrientDir;
 
-				UBasedMovementUtils::TransformWorldDirectionToBased(MovementBase, MovementBaseBoneName, DefaultKinematicInputs.GetMoveInput(), RelativeMoveInput);
-				UBasedMovementUtils::TransformWorldDirectionToBased(MovementBase, MovementBaseBoneName, DefaultKinematicInputs.OrientationIntent, RelativeOrientDir);
+				UBasedMovementUtils::TransformWorldDirectionToBased(MovementBase, MovementBaseBoneName, CharacterInputs.GetMoveInput(), RelativeMoveInput);
+				UBasedMovementUtils::TransformWorldDirectionToBased(MovementBase, MovementBaseBoneName, CharacterInputs.OrientationIntent, RelativeOrientDir);
 
-				DefaultKinematicInputs.SetMoveInput(DefaultKinematicInputs.GetMoveInputType(), RelativeMoveInput);
-				DefaultKinematicInputs.OrientationIntent = RelativeOrientDir;
+				CharacterInputs.SetMoveInput(CharacterInputs.GetMoveInputType(), RelativeMoveInput);
+				CharacterInputs.OrientationIntent = RelativeOrientDir;
 
-				DefaultKinematicInputs.bUsingMovementBase = true;
-				DefaultKinematicInputs.MovementBase = MovementBase;
-				DefaultKinematicInputs.MovementBaseBoneName = MovementBaseBoneName;
+				CharacterInputs.bUsingMovementBase = true;
+				CharacterInputs.MovementBase = MovementBase;
+				CharacterInputs.MovementBaseBoneName = MovementBaseBoneName;
 			}
 		}
 	}
diff --git a/Engine/Plugins/Experimental/MoverExamples/Source/MoverExamples/Public/CharacterVariants/AbilityInputs.h b/Engine/Plugins/Experimental/MoverExamples/Source/MoverExamples/Public/CharacterVariants/AbilityInputs.h
--- a/Engine/Plugins/Experimental/MoverExamples/Source/MoverExamples/Public/CharacterVariants/AbilityInputs.h
+++ b/Engine/Plugins/Experimental/MoverExamples/Source/MoverExamples/Public/CharacterVariants/AbilityInputs.h
@@ -25,7 +25,7 @@ struct MOVEREXAMPLES_API FMoverExampleAbilityInputs : public FMoverDataStructBas
 	bool bWantsToStartZiplining = false;
 
 
-	// @return newly allocated copy of this FKinematicDefaultInputs. Must be overridden by child classes
+	// @return newly allocated copy of this FMoverExampleAbilityInputs. Must be overridden by child classes
 	virtual FMoverDataStructBase* Clone() const override
 	{
 		// TODO: ensure that this memory allocation jives with deletion method
diff --git a/Engine/Plugins/Experimental/MoverExamples/Source/MoverExamples/Public/CharacterVariants/Ziplining/ZipliningTransitions.h b/Engine/Plugins/Experimental/MoverExamples/Source/MoverExamples/Public/CharacterVariants/Ziplining/ZipliningTransitions.h
--- a/Engine/Plugins/Experimental/MoverExamples/Source/MoverExamples/Public/CharacterVariants/Ziplining/ZipliningTransitions.h
+++ b/Engine/Plugins/Experimental/MoverExamples/Source/MoverExamples/Public/CharacterVariants/Ziplining/ZipliningTransitions.h
@@ -36,5 +36,5 @@ public:
 
 	// Mode to enter when exiting the zipline
 	UPROPERTY(EditAnywhere, Category = "Ziplining")
-	FName AutoExitToMode = KinematicModeNames::Falling;
+	FName AutoExitToMode = DefaultModeNames::Falling;
 };
