Converted some header deprecation from 5.5 to 5.4

diff --git a/Engine/Plugins/Experimental/ChaosVehiclesPlugin/Source/ChaosVehicles/Public/ChaosVehicleWheel.h b/Engine/Plugins/Experimental/ChaosVehiclesPlugin/Source/ChaosVehicles/Public/ChaosVehicleWheel.h
--- a/Engine/Plugins/Experimental/ChaosVehiclesPlugin/Source/ChaosVehicles/Public/ChaosVehicleWheel.h
+++ b/Engine/Plugins/Experimental/ChaosVehiclesPlugin/Source/ChaosVehicles/Public/ChaosVehicleWheel.h
@@ -13,9 +13,9 @@
 #include "EngineDefines.h"
 #include "SimpleVehicle.h"
 #include "VehicleUtility.h"
-#if UE_ENABLE_INCLUDE_ORDER_DEPRECATED_IN_5_5
+#if UE_ENABLE_INCLUDE_ORDER_DEPRECATED_IN_5_4
 #include "Engine/HitResult.h"
-#endif // UE_ENABLE_INCLUDE_ORDER_DEPRECATED_IN_5_5
+#endif // UE_ENABLE_INCLUDE_ORDER_DEPRECATED_IN_5_4
 #include "Curves/CurveFloat.h"
 
 #include "ChaosVehicleWheel.generated.h"
diff --git a/Engine/Plugins/Experimental/Mover/Source/Mover/Public/MoverComponent.h b/Engine/Plugins/Experimental/Mover/Source/Mover/Public/MoverComponent.h
--- a/Engine/Plugins/Experimental/Mover/Source/Mover/Public/MoverComponent.h
+++ b/Engine/Plugins/Experimental/Mover/Source/Mover/Public/MoverComponent.h
@@ -2,16 +2,16 @@
 
 #pragma once
 
-#if UE_ENABLE_INCLUDE_ORDER_DEPRECATED_IN_5_5
+#if UE_ENABLE_INCLUDE_ORDER_DEPRECATED_IN_5_4
 #include "CoreMinimal.h"
-#endif // UE_ENABLE_INCLUDE_ORDER_DEPRECATED_IN_5_5
+#endif // UE_ENABLE_INCLUDE_ORDER_DEPRECATED_IN_5_4
 #include "Components/ActorComponent.h"
 #include "MovementMode.h"
 #include "MoverTypes.h"
 #include "LayeredMove.h"
-#if UE_ENABLE_INCLUDE_ORDER_DEPRECATED_IN_5_5
+#if UE_ENABLE_INCLUDE_ORDER_DEPRECATED_IN_5_4
 #include "Engine/HitResult.h"
-#endif // UE_ENABLE_INCLUDE_ORDER_DEPRECATED_IN_5_5
+#endif // UE_ENABLE_INCLUDE_ORDER_DEPRECATED_IN_5_4
 #include "Backends/MoverBackendLiaison.h"
 #include "UObject/WeakInterfacePtr.h"
 #include "MoverComponent.generated.h"
diff --git a/Engine/Plugins/Experimental/XRCreativeFramework/Source/XRCreative/Private/ITF/SelectionInteraction.h b/Engine/Plugins/Experimental/XRCreativeFramework/Source/XRCreative/Private/ITF/SelectionInteraction.h
--- a/Engine/Plugins/Experimental/XRCreativeFramework/Source/XRCreative/Private/ITF/SelectionInteraction.h
+++ b/Engine/Plugins/Experimental/XRCreativeFramework/Source/XRCreative/Private/ITF/SelectionInteraction.h
@@ -6,9 +6,9 @@
 #include "BaseBehaviors/SingleClickBehavior.h"
 #include "Elements/Actor/ActorElementData.h"
 #include "Elements/Framework/TypedElementOwnerStore.h"
-#if UE_ENABLE_INCLUDE_ORDER_DEPRECATED_IN_5_5
+#if UE_ENABLE_INCLUDE_ORDER_DEPRECATED_IN_5_4
 #include "Engine/HitResult.h"
-#endif // UE_ENABLE_INCLUDE_ORDER_DEPRECATED_IN_5_5
+#endif // UE_ENABLE_INCLUDE_ORDER_DEPRECATED_IN_5_4
 #include "InputBehaviorSet.h"
 #include "SelectionInteraction.generated.h"
 
diff --git a/Engine/Plugins/Runtime/GeoReferencing/Source/GeoReferencingEditor/Public/GeoReferencingEditorBPLibrary.h b/Engine/Plugins/Runtime/GeoReferencing/Source/GeoReferencingEditor/Public/GeoReferencingEditorBPLibrary.h
--- a/Engine/Plugins/Runtime/GeoReferencing/Source/GeoReferencingEditor/Public/GeoReferencingEditorBPLibrary.h
+++ b/Engine/Plugins/Runtime/GeoReferencing/Source/GeoReferencingEditor/Public/GeoReferencingEditorBPLibrary.h
@@ -2,10 +2,10 @@
 
 #pragma once
 
-#if UE_ENABLE_INCLUDE_ORDER_DEPRECATED_IN_5_5
+#if UE_ENABLE_INCLUDE_ORDER_DEPRECATED_IN_5_4
 #include "CoreMinimal.h"
 #include "Engine/HitResult.h"
-#endif // UE_ENABLE_INCLUDE_ORDER_DEPRECATED_IN_5_5
+#endif // UE_ENABLE_INCLUDE_ORDER_DEPRECATED_IN_5_4
 #include "Kismet/BlueprintFunctionLibrary.h"
 
 #include "GeoReferencingEditorBPLibrary.generated.h"
diff --git a/Engine/Plugins/Runtime/MeshModelingToolset/Source/ModelingComponents/Public/Selection/SelectClickedAction.h b/Engine/Plugins/Runtime/MeshModelingToolset/Source/ModelingComponents/Public/Selection/SelectClickedAction.h
--- a/Engine/Plugins/Runtime/MeshModelingToolset/Source/ModelingComponents/Public/Selection/SelectClickedAction.h
+++ b/Engine/Plugins/Runtime/MeshModelingToolset/Source/ModelingComponents/Public/Selection/SelectClickedAction.h
@@ -5,9 +5,9 @@
 #include "CoreMinimal.h"
 #include "BaseBehaviors/BehaviorTargetInterfaces.h"
 #include "ToolSceneQueriesUtil.h"
-#if UE_ENABLE_INCLUDE_ORDER_DEPRECATED_IN_5_5
+#if UE_ENABLE_INCLUDE_ORDER_DEPRECATED_IN_5_4
 #include "Engine/HitResult.h"
-#endif // UE_ENABLE_INCLUDE_ORDER_DEPRECATED_IN_5_5
+#endif // UE_ENABLE_INCLUDE_ORDER_DEPRECATED_IN_5_4
 
 struct FHitResult;
 
diff --git a/Engine/Plugins/Runtime/NetworkPredictionExtras/Source/NetworkPredictionExtras/Public/BaseMovementComponent.h b/Engine/Plugins/Runtime/NetworkPredictionExtras/Source/NetworkPredictionExtras/Public/BaseMovementComponent.h
--- a/Engine/Plugins/Runtime/NetworkPredictionExtras/Source/NetworkPredictionExtras/Public/BaseMovementComponent.h
+++ b/Engine/Plugins/Runtime/NetworkPredictionExtras/Source/NetworkPredictionExtras/Public/BaseMovementComponent.h
@@ -3,9 +3,9 @@
 #pragma once
 
 #include "NetworkPredictionComponent.h"
-#if UE_ENABLE_INCLUDE_ORDER_DEPRECATED_IN_5_5
+#if UE_ENABLE_INCLUDE_ORDER_DEPRECATED_IN_5_4
 #include "Engine/HitResult.h"
-#endif // UE_ENABLE_INCLUDE_ORDER_DEPRECATED_IN_5_5
+#endif // UE_ENABLE_INCLUDE_ORDER_DEPRECATED_IN_5_4
 #include "BaseMovementComponent.generated.h"
 
 struct FHitResult;
diff --git a/Engine/Source/Editor/ClothPainter/Private/ClothPainter.h b/Engine/Source/Editor/ClothPainter/Private/ClothPainter.h
--- a/Engine/Source/Editor/ClothPainter/Private/ClothPainter.h
+++ b/Engine/Source/Editor/ClothPainter/Private/ClothPainter.h
@@ -7,9 +7,9 @@
 #include "Containers/Map.h"
 #include "Delegates/IDelegateInstance.h"
 #include "Engine/EngineBaseTypes.h"
-#if UE_ENABLE_INCLUDE_ORDER_DEPRECATED_IN_5_5
+#if UE_ENABLE_INCLUDE_ORDER_DEPRECATED_IN_5_4
 #include "Engine/HitResult.h"
-#endif // UE_ENABLE_INCLUDE_ORDER_DEPRECATED_IN_5_5
+#endif // UE_ENABLE_INCLUDE_ORDER_DEPRECATED_IN_5_4
 #include "IMeshPainter.h"
 #include "InputCoreTypes.h"
 #include "Internationalization/Text.h"
diff --git a/Engine/Source/Runtime/Engine/Classes/Components/InterpToMovementComponent.h b/Engine/Source/Runtime/Engine/Classes/Components/InterpToMovementComponent.h
--- a/Engine/Source/Runtime/Engine/Classes/Components/InterpToMovementComponent.h
+++ b/Engine/Source/Runtime/Engine/Classes/Components/InterpToMovementComponent.h
@@ -2,10 +2,10 @@
 
 #pragma once
 
-#if UE_ENABLE_INCLUDE_ORDER_DEPRECATED_IN_5_5
+#if UE_ENABLE_INCLUDE_ORDER_DEPRECATED_IN_5_4
 #include "CoreMinimal.h"
 #include "Engine/HitResult.h"
-#endif // UE_ENABLE_INCLUDE_ORDER_DEPRECATED_IN_5_5
+#endif // UE_ENABLE_INCLUDE_ORDER_DEPRECATED_IN_5_4
 #include "GameFramework/MovementComponent.h"
 #include "UObject/ObjectMacros.h"
 #include "InterpToMovementComponent.generated.h"
diff --git a/Engine/Source/Runtime/Engine/Classes/GameFramework/ProjectileMovementComponent.h b/Engine/Source/Runtime/Engine/Classes/GameFramework/ProjectileMovementComponent.h
--- a/Engine/Source/Runtime/Engine/Classes/GameFramework/ProjectileMovementComponent.h
+++ b/Engine/Source/Runtime/Engine/Classes/GameFramework/ProjectileMovementComponent.h
@@ -2,10 +2,10 @@
 
 #pragma once
 
-#if UE_ENABLE_INCLUDE_ORDER_DEPRECATED_IN_5_5
+#if UE_ENABLE_INCLUDE_ORDER_DEPRECATED_IN_5_4
 #include "CoreMinimal.h"
 #include "Engine/HitResult.h"
-#endif // UE_ENABLE_INCLUDE_ORDER_DEPRECATED_IN_5_5
+#endif // UE_ENABLE_INCLUDE_ORDER_DEPRECATED_IN_5_4
 #include "GameFramework/MovementComponent.h"
 #include "UObject/ObjectMacros.h"
 #include "ProjectileMovementComponent.generated.h"
diff --git a/Engine/Source/Runtime/Engine/Classes/Kismet/KismetSystemLibrary.h b/Engine/Source/Runtime/Engine/Classes/Kismet/KismetSystemLibrary.h
--- a/Engine/Source/Runtime/Engine/Classes/Kismet/KismetSystemLibrary.h
+++ b/Engine/Source/Runtime/Engine/Classes/Kismet/KismetSystemLibrary.h
@@ -2,15 +2,15 @@
 
 #pragma once
 
-#if UE_ENABLE_INCLUDE_ORDER_DEPRECATED_IN_5_5
+#if UE_ENABLE_INCLUDE_ORDER_DEPRECATED_IN_5_4
 #include "CoreMinimal.h"
-#endif // UE_ENABLE_INCLUDE_ORDER_DEPRECATED_IN_5_5
+#endif // UE_ENABLE_INCLUDE_ORDER_DEPRECATED_IN_5_4
 #include "UObject/ObjectMacros.h"
 #include "UObject/Object.h"
 #include "Templates/SubclassOf.h"
-#if UE_ENABLE_INCLUDE_ORDER_DEPRECATED_IN_5_5
+#if UE_ENABLE_INCLUDE_ORDER_DEPRECATED_IN_5_4
 #include "Engine/HitResult.h"
-#endif // UE_ENABLE_INCLUDE_ORDER_DEPRECATED_IN_5_5
+#endif // UE_ENABLE_INCLUDE_ORDER_DEPRECATED_IN_5_4
 #include "UObject/UnrealType.h"
 #include "UObject/ScriptMacros.h"
 #include "UObject/Interface.h"
diff --git a/Engine/Source/Runtime/Engine/Public/Collision.h b/Engine/Source/Runtime/Engine/Public/Collision.h
--- a/Engine/Source/Runtime/Engine/Public/Collision.h
+++ b/Engine/Source/Runtime/Engine/Public/Collision.h
@@ -6,14 +6,14 @@
 
 #pragma once
 
-#if UE_ENABLE_INCLUDE_ORDER_DEPRECATED_IN_5_5
+#if UE_ENABLE_INCLUDE_ORDER_DEPRECATED_IN_5_4
 #include "CoreMinimal.h"
-#endif // UE_ENABLE_INCLUDE_ORDER_DEPRECATED_IN_5_5
+#endif // UE_ENABLE_INCLUDE_ORDER_DEPRECATED_IN_5_4
 #include "Math/Vector.h"
 #include "Stats/Stats.h"
-#if UE_ENABLE_INCLUDE_ORDER_DEPRECATED_IN_5_5
+#if UE_ENABLE_INCLUDE_ORDER_DEPRECATED_IN_5_4
 #include "Engine/HitResult.h"
-#endif // UE_ENABLE_INCLUDE_ORDER_DEPRECATED_IN_5_5
+#endif // UE_ENABLE_INCLUDE_ORDER_DEPRECATED_IN_5_4
 #include "EngineDefines.h"
 
 struct FHitResult;
diff --git a/Engine/Source/Runtime/Engine/Public/Engine/ScopedMovementUpdate.h b/Engine/Source/Runtime/Engine/Public/Engine/ScopedMovementUpdate.h
--- a/Engine/Source/Runtime/Engine/Public/Engine/ScopedMovementUpdate.h
+++ b/Engine/Source/Runtime/Engine/Public/Engine/ScopedMovementUpdate.h
@@ -2,10 +2,10 @@
 
 #pragma once
 
-#if UE_ENABLE_INCLUDE_ORDER_DEPRECATED_IN_5_5
+#if UE_ENABLE_INCLUDE_ORDER_DEPRECATED_IN_5_4
 #include "CoreMinimal.h"
 #include "Engine/HitResult.h"
-#endif // UE_ENABLE_INCLUDE_ORDER_DEPRECATED_IN_5_5
+#endif // UE_ENABLE_INCLUDE_ORDER_DEPRECATED_IN_5_4
 #include "Engine/OverlapInfo.h"
 
 struct FHitResult;
diff --git a/Engine/Source/Runtime/Engine/Public/WorldCollision.h b/Engine/Source/Runtime/Engine/Public/WorldCollision.h
--- a/Engine/Source/Runtime/Engine/Public/WorldCollision.h
+++ b/Engine/Source/Runtime/Engine/Public/WorldCollision.h
@@ -6,15 +6,15 @@
 #pragma once 
 
 #include "Async/TaskGraphFwd.h"
-#if UE_ENABLE_INCLUDE_ORDER_DEPRECATED_IN_5_5
+#if UE_ENABLE_INCLUDE_ORDER_DEPRECATED_IN_5_4
 #include "CoreMinimal.h"
-#endif // UE_ENABLE_INCLUDE_ORDER_DEPRECATED_IN_5_5
+#endif // UE_ENABLE_INCLUDE_ORDER_DEPRECATED_IN_5_4
 #include "CollisionQueryParams.h"
 #include "CollisionShape.h"
-#if UE_ENABLE_INCLUDE_ORDER_DEPRECATED_IN_5_5
+#if UE_ENABLE_INCLUDE_ORDER_DEPRECATED_IN_5_4
 #include "Engine/OverlapResult.h"
 #include "Engine/HitResult.h"
-#endif // UE_ENABLE_INCLUDE_ORDER_DEPRECATED_IN_5_5
+#endif // UE_ENABLE_INCLUDE_ORDER_DEPRECATED_IN_5_4
 
 #if UE_ENABLE_INCLUDE_ORDER_DEPRECATED_IN_5_2
 #include "Async/TaskGraphInterfaces.h"
diff --git a/Engine/Source/Runtime/Experimental/ChaosSolverEngine/Public/Chaos/ChaosNotifyHandlerInterface.h b/Engine/Source/Runtime/Experimental/ChaosSolverEngine/Public/Chaos/ChaosNotifyHandlerInterface.h
--- a/Engine/Source/Runtime/Experimental/ChaosSolverEngine/Public/Chaos/ChaosNotifyHandlerInterface.h
+++ b/Engine/Source/Runtime/Experimental/ChaosSolverEngine/Public/Chaos/ChaosNotifyHandlerInterface.h
@@ -4,9 +4,9 @@
 
 #include "Kismet/BlueprintFunctionLibrary.h"
 #include "UObject/Interface.h"
-#if UE_ENABLE_INCLUDE_ORDER_DEPRECATED_IN_5_5
+#if UE_ENABLE_INCLUDE_ORDER_DEPRECATED_IN_5_4
 #include "Engine/HitResult.h"
-#endif // UE_ENABLE_INCLUDE_ORDER_DEPRECATED_IN_5_5
+#endif // UE_ENABLE_INCLUDE_ORDER_DEPRECATED_IN_5_4
 #include "ChaosNotifyHandlerInterface.generated.h"
 
 struct FHitResult;
