Mover - Fix movement velocity when landing by always projecting prior velocity onto movement plane in walking mode

diff --git a/Engine/Plugins/Experimental/Mover/Source/Mover/Private/DefaultMovementSet/Modes/WalkingMode.cpp b/Engine/Plugins/Experimental/Mover/Source/Mover/Private/DefaultMovementSet/Modes/WalkingMode.cpp
--- a/Engine/Plugins/Experimental/Mover/Source/Mover/Private/DefaultMovementSet/Modes/WalkingMode.cpp
+++ b/Engine/Plugins/Experimental/Mover/Source/Mover/Private/DefaultMovementSet/Modes/WalkingMode.cpp
@@ -73,7 +73,7 @@ void UWalkingMode::OnGenerateMove(const FMoverTickStartData& StartState, const F
 	}
 
 	Params.OrientationIntent = IntendedOrientation_WorldSpace;
-	Params.PriorVelocity = StartingSyncState->GetVelocity_WorldSpace();
+	Params.PriorVelocity = FVector::VectorPlaneProject(StartingSyncState->GetVelocity_WorldSpace(), MovementNormal);
 	Params.PriorOrientation = StartingSyncState->GetOrientation_WorldSpace();
 	Params.GroundNormal = MovementNormal;
 	Params.TurningRate = CommonLegacySettings->TurningRate;
