Mover: adding API for directly reading/writing to the simulation sync state that is yet to be committed/sent. This is to help support advanced mechanics that may break out of Network Prediction's rigid simulation window in order to interoperate with other gameplay or physics-simulated objects.

diff --git a/Engine/Plugins/Experimental/Mover/Source/Mover/Private/Backends/MoverNetworkPredictionLiaison.cpp b/Engine/Plugins/Experimental/Mover/Source/Mover/Private/Backends/MoverNetworkPredictionLiaison.cpp
--- a/Engine/Plugins/Experimental/Mover/Source/Mover/Private/Backends/MoverNetworkPredictionLiaison.cpp
+++ b/Engine/Plugins/Experimental/Mover/Source/Mover/Private/Backends/MoverNetworkPredictionLiaison.cpp
@@ -3,6 +3,7 @@
 #include "Backends/MoverNetworkPredictionLiaison.h"
 #include "NetworkPredictionModelDefRegistry.h"
 #include "NetworkPredictionProxyInit.h"
+#include "NetworkPredictionProxyWrite.h"
 #include "GameFramework/Actor.h"
 #include "MoverComponent.h"
 
@@ -91,6 +92,28 @@ int32 UMoverNetworkPredictionLiaisonComponent::GetCurrentSimFrame()
 	return NetworkPredictionProxy.GetPendingFrame();
 }
 
+
+bool UMoverNetworkPredictionLiaisonComponent::ReadPendingSyncState(OUT FMoverSyncState& OutSyncState)
+{
+	if (const FMoverSyncState* PendingSyncState = NetworkPredictionProxy.ReadSyncState<FMoverSyncState>())
+	{
+		OutSyncState = *PendingSyncState;
+		return true;
+	}
+
+	return false;
+}
+
+bool UMoverNetworkPredictionLiaisonComponent::WritePendingSyncState(const FMoverSyncState& SyncStateToWrite)
+{
+	NetworkPredictionProxy.WriteSyncState<FMoverSyncState>([&SyncStateToWrite](FMoverSyncState& PendingSyncStateRef)
+		{
+			PendingSyncStateRef = SyncStateToWrite;
+		});
+
+	return true;
+}
+
 void UMoverNetworkPredictionLiaisonComponent::BeginPlay()
 {
 	Super::BeginPlay();
diff --git a/Engine/Plugins/Experimental/Mover/Source/Mover/Public/Backends/MoverBackendLiaison.h b/Engine/Plugins/Experimental/Mover/Source/Mover/Public/Backends/MoverBackendLiaison.h
--- a/Engine/Plugins/Experimental/Mover/Source/Mover/Public/Backends/MoverBackendLiaison.h
+++ b/Engine/Plugins/Experimental/Mover/Source/Mover/Public/Backends/MoverBackendLiaison.h
@@ -3,13 +3,15 @@
 #pragma once
 
 #include "UObject/Interface.h"
+#include "MoverSimulationTypes.h"
 #include "MoverBackendLiaison.generated.h"
 
+
 /**
- * MoverBackendLiaisonInterface: any object or system wanting to be the driver of Mover actors must implement this. The intent is to act as a 
+ * MoverBackendLiaisonInterface: any object or system wanting to be the driver of Mover actors must implement this. The intent is to act as a
  * middleman between the Mover actor and the system that drives it, such as the Network Prediction plugin.
  * In practice, objects implementing this interface should be some kind of UActorComponent. The Mover actor instantiates its backend liaison
- * when initialized, then relies on the liaison to call various functions as the simulation progresses. See @KinematicMoverComponent.
+ * when initialized, then relies on the liaison to call various functions as the simulation progresses. See @MoverComponent.
  */
 UINTERFACE(MinimalAPI)
 class UMoverBackendLiaisonInterface : public UInterface
@@ -24,4 +26,7 @@ class IMoverBackendLiaisonInterface
 public:
 	virtual float GetCurrentSimTimeMs() = 0;
 	virtual int32 GetCurrentSimFrame() = 0;
+
+	virtual bool ReadPendingSyncState(OUT FMoverSyncState& OutSyncState) { return false; }
+	virtual bool WritePendingSyncState(const FMoverSyncState& SyncStateToWrite) { return false; }
 };
\ No newline at end of file
diff --git a/Engine/Plugins/Experimental/Mover/Source/Mover/Public/Backends/MoverNetworkPredictionLiaison.h b/Engine/Plugins/Experimental/Mover/Source/Mover/Public/Backends/MoverNetworkPredictionLiaison.h
--- a/Engine/Plugins/Experimental/Mover/Source/Mover/Public/Backends/MoverNetworkPredictionLiaison.h
+++ b/Engine/Plugins/Experimental/Mover/Source/Mover/Public/Backends/MoverNetworkPredictionLiaison.h
@@ -47,6 +47,8 @@ public:
 	// IMoverBackendLiaisonInterface
 	virtual float GetCurrentSimTimeMs() override;
 	virtual int32 GetCurrentSimFrame() override;
+	virtual bool ReadPendingSyncState(OUT FMoverSyncState& OutSyncState) override;
+	virtual bool WritePendingSyncState(const FMoverSyncState& SyncStateToWrite) override;
 	// End IMoverBackendLiaisonInterface
 
 	virtual void BeginPlay() override;
