Physics Mover - Move registration of physics mover updates to BeginPlay to avoid running when the mover component is not set up. Add checks to ensure update can run

diff --git a/Engine/Plugins/Experimental/Mover/Source/Mover/Private/Backends/MoverNetworkPhysicsLiaison.cpp b/Engine/Plugins/Experimental/Mover/Source/Mover/Private/Backends/MoverNetworkPhysicsLiaison.cpp
--- a/Engine/Plugins/Experimental/Mover/Source/Mover/Private/Backends/MoverNetworkPhysicsLiaison.cpp
+++ b/Engine/Plugins/Experimental/Mover/Source/Mover/Private/Backends/MoverNetworkPhysicsLiaison.cpp
@@ -392,10 +392,21 @@ void UMoverNetworkPhysicsLiaisonComponent::InitializeComponent()
 		MoverComp->InitMoverSimulation();
 		MoverComp->ModeFSM->SetModeImmediately(MoverComp->StartingMovementMode);
 	}
+
+	// Register network data for recording and rewind/resim
+	if (NetworkPhysicsComponent)
+	{
+		NetworkPhysicsComponent->CreateDataHistory<FNetworkPhysicsMoverTraits>(this);
+	}
 }
 
 void UMoverNetworkPhysicsLiaisonComponent::UninitializeComponent()
 {
+	if (NetworkPhysicsComponent)
+	{
+		NetworkPhysicsComponent->RemoveDataHistory();
+	}
+
 	Super::UninitializeComponent();
 }
 
@@ -428,7 +439,8 @@ bool UMoverNetworkPhysicsLiaisonComponent::HasValidPhysicsState() const
 bool UMoverNetworkPhysicsLiaisonComponent::HasValidState() const
 {
 	return HasValidPhysicsState() && MoverComp && MoverComp->UpdatedCompAsPrimitive && MoverComp->UpdatedComponent
-		&& MoverComp->ModeFSM->IsValidLowLevel() && MoverComp->SimBlackboard->IsValidLowLevel();
+		&& MoverComp->ModeFSM->IsValidLowLevelFast() && MoverComp->SimBlackboard->IsValidLowLevelFast()
+		&& MoverComp->InputProducer && MoverComp->MovementMixer;
 }
 
 void UMoverNetworkPhysicsLiaisonComponent::OnCreatePhysicsState()
@@ -436,56 +448,11 @@ void UMoverNetworkPhysicsLiaisonComponent::OnCreatePhysicsState()
 	Super::OnCreatePhysicsState();
 
 	SetupConstraint();
-
-	UWorld* World = GetWorld();
-	if (World && World->IsGameWorld())
-	{
-		if (FPhysScene_Chaos* PhysScene = World->GetPhysicsScene())
-		{
-			if (Chaos::FPhysicsSolver* Solver = PhysScene->GetSolver())
-			{
-				if (MoverComp && MoverComp->UpdatedCompAsPrimitive)
-				{
-					if (FBodyInstance* BI = MoverComp->UpdatedCompAsPrimitive->GetBodyInstance())
-					{
-						if (Chaos::FSingleParticlePhysicsProxy* CharacterProxy = BI->ActorHandle)
-						{
-							// Register network data for recording and rewind/resim
-							if (NetworkPhysicsComponent)
-							{
-								NetworkPhysicsComponent->CreateDataHistory<FNetworkPhysicsMoverTraits>(this);
-							}
-
-							// Register with the physics mover manager
-							if (UPhysicsMoverManager* Manager = World->GetSubsystem<UPhysicsMoverManager>())
-							{
-								Manager->RegisterPhysicsMoverComponent(this);
-							}
-						}
-					}
-				}
-			}
-		}
-	}
 }
 
 void UMoverNetworkPhysicsLiaisonComponent::OnDestroyPhysicsState()
 {
-	UWorld* World = GetWorld();
-	if (World && World->IsGameWorld() && HasValidPhysicsState())
-	{
-		if (NetworkPhysicsComponent)
-		{
-			NetworkPhysicsComponent->RemoveDataHistory();
-		}
-
-		DestroyConstraint();
-
-		if (UPhysicsMoverManager* Manager = World->GetSubsystem<UPhysicsMoverManager>())
-		{
-			Manager->UnregisterPhysicsMoverComponent(this);
-		}
-	}
+	DestroyConstraint();
 	
 	Super::OnDestroyPhysicsState();
 }
@@ -510,6 +477,34 @@ bool UMoverNetworkPhysicsLiaisonComponent::CanCreatePhysics() const
 	return true;
 }
 
+void UMoverNetworkPhysicsLiaisonComponent::BeginPlay()
+{
+	Super::BeginPlay();
+
+	// Register with the physics mover manager
+	if (UWorld* World = GetWorld())
+	{
+		if (UPhysicsMoverManager* Manager = World->GetSubsystem<UPhysicsMoverManager>())
+		{
+			Manager->RegisterPhysicsMoverComponent(this);
+		}
+	}
+}
+
+void UMoverNetworkPhysicsLiaisonComponent::EndPlay(const EEndPlayReason::Type EndPlayReason)
+{
+	// Unregister with the physics mover manager
+	if (UWorld* World = GetWorld())
+	{
+		if (UPhysicsMoverManager* Manager = World->GetSubsystem<UPhysicsMoverManager>())
+		{
+			Manager->UnregisterPhysicsMoverComponent(this);
+		}
+	}
+
+	Super::EndPlay(EndPlayReason);
+}
+
 //////////////////////////////////////////////////////////////////////////
 
 Chaos::FUniqueIdx UMoverNetworkPhysicsLiaisonComponent::GetUniqueIdx() const
diff --git a/Engine/Plugins/Experimental/Mover/Source/Mover/Public/Backends/MoverNetworkPhysicsLiaison.h b/Engine/Plugins/Experimental/Mover/Source/Mover/Public/Backends/MoverNetworkPhysicsLiaison.h
--- a/Engine/Plugins/Experimental/Mover/Source/Mover/Public/Backends/MoverNetworkPhysicsLiaison.h
+++ b/Engine/Plugins/Experimental/Mover/Source/Mover/Public/Backends/MoverNetworkPhysicsLiaison.h
@@ -117,6 +117,8 @@ public:
 	virtual void OnCreatePhysicsState() override;
 	virtual void OnDestroyPhysicsState() override;
 	virtual bool CanCreatePhysics() const;
+	virtual void BeginPlay() override;
+	virtual void EndPlay(const EEndPlayReason::Type EndPlayReason) override;
 
 	// Required for Network Physics Rewind/Resim data
 	void GetCurrentInputData(OUT FMoverInputCmdContext& InputCmd) const;
