Mover 2.0 adding some data validation to physics movement modes to make sure that a physics backend class is present to avoid unwanted interactions and crashes

diff --git a/Engine/Plugins/Experimental/Mover/Source/Mover/Private/PhysicsMover/Modes/PhysicsDrivenFallingMode.cpp b/Engine/Plugins/Experimental/Mover/Source/Mover/Private/PhysicsMover/Modes/PhysicsDrivenFallingMode.cpp
--- a/Engine/Plugins/Experimental/Mover/Source/Mover/Private/PhysicsMover/Modes/PhysicsDrivenFallingMode.cpp
+++ b/Engine/Plugins/Experimental/Mover/Source/Mover/Private/PhysicsMover/Modes/PhysicsDrivenFallingMode.cpp
@@ -10,9 +10,16 @@
 #include "MoveLibrary/FloorQueryUtils.h"
 #include "MoveLibrary/WaterMovementUtils.h"
 #include "PhysicsMover/PhysicsMovementUtils.h"
+#if WITH_EDITOR
+#include "Backends/MoverNetworkPhysicsLiaison.h"
+#include "Internationalization/Text.h"
+#include "Misc/DataValidation.h"
+#endif // WITH_EDITOR
 
 #include UE_INLINE_GENERATED_CPP_BY_NAME(PhysicsDrivenFallingMode)
 
+#define LOCTEXT_NAMESPACE "PhysicsDrivenFallingMode"
+
 UPhysicsDrivenFallingMode::UPhysicsDrivenFallingMode(const FObjectInitializer& ObjectInitializer)
 	: Super(ObjectInitializer)
 {
@@ -26,6 +33,21 @@ void UPhysicsDrivenFallingMode::UpdateConstraintSettings(Chaos::FCharacterGround
 	Constraint.SetTargetHeight(TargetHeight);
 }
 
+#if WITH_EDITOR
+EDataValidationResult UPhysicsDrivenFallingMode::IsDataValid(FDataValidationContext& Context) const
+{
+	EDataValidationResult Result = Super::IsDataValid(Context);
+	const UClass* BackendClass = GetMoverComponent()->BackendClass;
+	if (BackendClass && !BackendClass->IsChildOf<UMoverNetworkPhysicsLiaisonComponent>())
+	{
+		Context.AddError(LOCTEXT("PhysicsMovementModeHasValidPhysicsLiaison", "Physics movement modes need to have a backend class that supports physics (UMoverNetworkPhysicsLiaisonComponent)."));
+		Result = EDataValidationResult::Invalid;
+	}
+		
+	return Result;
+}
+#endif // WITH_EDITOR
+
 void UPhysicsDrivenFallingMode::OnSimulationTick(const FSimulationTickParams& Params, FMoverTickEndData& OutputState)
 {
 	const FMoverTickStartData& StartState = Params.StartState;
@@ -126,4 +148,6 @@ void UPhysicsDrivenFallingMode::OnSimulationTick(const FSimulationTickParams& Pa
 		TargetPos,
 		TargetOrient,
 		TargetVel);
-}
\ No newline at end of file
+}
+
+#undef LOCTEXT_NAMESPACE
\ No newline at end of file
diff --git a/Engine/Plugins/Experimental/Mover/Source/Mover/Private/PhysicsMover/Modes/PhysicsDrivenFlyingMode.cpp b/Engine/Plugins/Experimental/Mover/Source/Mover/Private/PhysicsMover/Modes/PhysicsDrivenFlyingMode.cpp
--- a/Engine/Plugins/Experimental/Mover/Source/Mover/Private/PhysicsMover/Modes/PhysicsDrivenFlyingMode.cpp
+++ b/Engine/Plugins/Experimental/Mover/Source/Mover/Private/PhysicsMover/Modes/PhysicsDrivenFlyingMode.cpp
@@ -9,9 +9,16 @@
 #include "Math/UnitConversion.h"
 #include "MoveLibrary/MovementUtils.h"
 #include "PhysicsMover/PhysicsMovementUtils.h"
+#if WITH_EDITOR
+#include "Backends/MoverNetworkPhysicsLiaison.h"
+#include "Internationalization/Text.h"
+#include "Misc/DataValidation.h"
+#endif // WITH_EDITOR
 
 #include UE_INLINE_GENERATED_CPP_BY_NAME(PhysicsDrivenFlyingMode)
 
+#define LOCTEXT_NAMESPACE "PhysicsDrivenFlyingMode"
+
 UPhysicsDrivenFlyingMode::UPhysicsDrivenFlyingMode(const FObjectInitializer& ObjectInitializer)
 	: Super(ObjectInitializer)
 {
@@ -25,6 +32,21 @@ void UPhysicsDrivenFlyingMode::UpdateConstraintSettings(Chaos::FCharacterGroundC
 	Constraint.SetTargetHeight(0.0f);
 }
 
+#if WITH_EDITOR
+EDataValidationResult UPhysicsDrivenFlyingMode::IsDataValid(FDataValidationContext& Context) const
+{
+	EDataValidationResult Result = Super::IsDataValid(Context);
+	const UClass* BackendClass = GetMoverComponent()->BackendClass;
+	if (BackendClass && !BackendClass->IsChildOf<UMoverNetworkPhysicsLiaisonComponent>())
+	{
+		Context.AddError(LOCTEXT("PhysicsMovementModeHasValidPhysicsLiaison", "Physics movement modes need to have a backend class that supports physics (UMoverNetworkPhysicsLiaisonComponent)."));
+		Result = EDataValidationResult::Invalid;
+	}
+		
+	return Result;
+}
+#endif // WITH_EDITOR
+
 void UPhysicsDrivenFlyingMode::OnSimulationTick(const FSimulationTickParams& Params, FMoverTickEndData& OutputState)
 {
 	const FMoverTickStartData& StartState = Params.StartState;
@@ -80,4 +102,6 @@ void UPhysicsDrivenFlyingMode::OnSimulationTick(const FSimulationTickParams& Par
 		TargetPos,
 		TargetOrient,
 		TargetVel);
-}
\ No newline at end of file
+}
+
+#undef LOCTEXT_NAMESPACE
\ No newline at end of file
diff --git a/Engine/Plugins/Experimental/Mover/Source/Mover/Private/PhysicsMover/Modes/PhysicsDrivenSwimmingMode.cpp b/Engine/Plugins/Experimental/Mover/Source/Mover/Private/PhysicsMover/Modes/PhysicsDrivenSwimmingMode.cpp
--- a/Engine/Plugins/Experimental/Mover/Source/Mover/Private/PhysicsMover/Modes/PhysicsDrivenSwimmingMode.cpp
+++ b/Engine/Plugins/Experimental/Mover/Source/Mover/Private/PhysicsMover/Modes/PhysicsDrivenSwimmingMode.cpp
@@ -11,9 +11,16 @@
 #include "DefaultMovementSet/Settings/CommonLegacyMovementSettings.h"
 #include "GameFramework/PhysicsVolume.h"
 #include "PhysicsMover/PhysicsMovementUtils.h"
+#if WITH_EDITOR
+#include "Backends/MoverNetworkPhysicsLiaison.h"
+#include "Internationalization/Text.h"
+#include "Misc/DataValidation.h"
+#endif // WITH_EDITOR
 
 #include UE_INLINE_GENERATED_CPP_BY_NAME(PhysicsDrivenSwimmingMode)
 
+#define LOCTEXT_NAMESPACE "PhysicsDrivenSwimmingMode"
+
 UPhysicsDrivenSwimmingMode::UPhysicsDrivenSwimmingMode(const FObjectInitializer& ObjectInitializer)
 	: Super(ObjectInitializer)
 {
@@ -27,6 +34,21 @@ void UPhysicsDrivenSwimmingMode::UpdateConstraintSettings(Chaos::FCharacterGroun
 	Constraint.SetTwistTorqueLimit(0.0f);
 }
 
+#if WITH_EDITOR
+EDataValidationResult UPhysicsDrivenSwimmingMode::IsDataValid(FDataValidationContext& Context) const
+{
+	EDataValidationResult Result = Super::IsDataValid(Context);
+	const UClass* BackendClass = GetMoverComponent()->BackendClass;
+	if (BackendClass && !BackendClass->IsChildOf<UMoverNetworkPhysicsLiaisonComponent>())
+	{
+		Context.AddError(LOCTEXT("PhysicsMovementModeHasValidPhysicsLiaison", "Physics movement modes need to have a backend class that supports physics (UMoverNetworkPhysicsLiaisonComponent)."));
+		Result = EDataValidationResult::Invalid;
+	}
+		
+	return Result;
+}
+#endif // WITH_EDITOR
+
 void UPhysicsDrivenSwimmingMode::OnSimulationTick(const FSimulationTickParams& Params, FMoverTickEndData& OutputState)
 {
 	const UMoverComponent* MoverComp = GetMoverComponent();
@@ -154,4 +176,6 @@ bool UPhysicsDrivenSwimmingMode::AttemptTeleport(USceneComponent* UpdatedCompone
 	GetBlackboard_Mutable()->Invalidate(CommonBlackboard::LastMovementBase);
 
 	return true;
-}
\ No newline at end of file
+}
+
+#undef LOCTEXT_NAMESPACE
\ No newline at end of file
diff --git a/Engine/Plugins/Experimental/Mover/Source/Mover/Private/PhysicsMover/Modes/PhysicsDrivenWalkingMode.cpp b/Engine/Plugins/Experimental/Mover/Source/Mover/Private/PhysicsMover/Modes/PhysicsDrivenWalkingMode.cpp
--- a/Engine/Plugins/Experimental/Mover/Source/Mover/Private/PhysicsMover/Modes/PhysicsDrivenWalkingMode.cpp
+++ b/Engine/Plugins/Experimental/Mover/Source/Mover/Private/PhysicsMover/Modes/PhysicsDrivenWalkingMode.cpp
@@ -13,11 +13,18 @@
 #include "MoveLibrary/WaterMovementUtils.h"
 #include "PhysicsMover/PhysicsMovementUtils.h"
 #include "PhysicsMover/PhysicsMoverSimulationTypes.h"
+#if WITH_EDITOR
+#include "Backends/MoverNetworkPhysicsLiaison.h"
+#include "Internationalization/Text.h"
+#include "Misc/DataValidation.h"
+#endif // WITH_EDITOR
 
 #include UE_INLINE_GENERATED_CPP_BY_NAME(PhysicsDrivenWalkingMode)
 
 extern FPhysicsDrivenMotionDebugParams GPhysicsDrivenMotionDebugParams;
 
+#define LOCTEXT_NAMESPACE "PhysicsDrivenWalkingMode"
+
 UPhysicsDrivenWalkingMode::UPhysicsDrivenWalkingMode(const FObjectInitializer& ObjectInitializer)
 	: Super(ObjectInitializer)
 {
@@ -32,6 +39,21 @@ void UPhysicsDrivenWalkingMode::UpdateConstraintSettings(Chaos::FCharacterGround
 	Constraint.SetTargetHeight(TargetHeight);
 }
 
+#if WITH_EDITOR
+EDataValidationResult UPhysicsDrivenWalkingMode::IsDataValid(FDataValidationContext& Context) const
+{
+	EDataValidationResult Result = Super::IsDataValid(Context);
+	const UClass* BackendClass = GetMoverComponent()->BackendClass;
+	if (BackendClass && !BackendClass->IsChildOf<UMoverNetworkPhysicsLiaisonComponent>())
+	{
+		Context.AddError(LOCTEXT("PhysicsMovementModeHasValidPhysicsLiaison", "Physics movement modes need to have a backend class that supports physics (UMoverNetworkPhysicsLiaisonComponent)."));
+		Result = EDataValidationResult::Invalid;
+	}
+		
+	return Result;
+}
+#endif // WITH_EDITOR
+
 void UPhysicsDrivenWalkingMode::OnSimulationTick(const FSimulationTickParams& Params, FMoverTickEndData& OutputState)
 {
 	const UMoverComponent* MoverComp = GetMoverComponent();
@@ -224,4 +246,6 @@ void UPhysicsDrivenWalkingMode::SwitchToState(const FName& StateName, const FSim
 		StartingSyncState->GetOrientation_WorldSpace(),
 		StartingSyncState->GetVelocity_WorldSpace(),
 		nullptr);
-}
\ No newline at end of file
+}
+
+#undef LOCTEXT_NAMESPACE
\ No newline at end of file
diff --git a/Engine/Plugins/Experimental/Mover/Source/Mover/Public/PhysicsMover/Modes/PhysicsDrivenFallingMode.h b/Engine/Plugins/Experimental/Mover/Source/Mover/Public/PhysicsMover/Modes/PhysicsDrivenFallingMode.h
--- a/Engine/Plugins/Experimental/Mover/Source/Mover/Public/PhysicsMover/Modes/PhysicsDrivenFallingMode.h
+++ b/Engine/Plugins/Experimental/Mover/Source/Mover/Public/PhysicsMover/Modes/PhysicsDrivenFallingMode.h
@@ -25,6 +25,10 @@ public:
 
 	virtual void UpdateConstraintSettings(Chaos::FCharacterGroundConstraint& Constraint) const override;
 
+#if WITH_EDITOR
+	virtual EDataValidationResult IsDataValid(class FDataValidationContext& Context) const override;
+#endif // WITH_EDITOR
+
 	// Maximum torque the character can apply to rotate in air about the vertical axis
 	UPROPERTY(EditAnywhere, Category = "Physics Mover", meta = (ClampMin = "0", UIMin = "0", ForceUnits = "NewtonMeters"))
 	float TwistTorqueLimit = 0.0f;
diff --git a/Engine/Plugins/Experimental/Mover/Source/Mover/Public/PhysicsMover/Modes/PhysicsDrivenFlyingMode.h b/Engine/Plugins/Experimental/Mover/Source/Mover/Public/PhysicsMover/Modes/PhysicsDrivenFlyingMode.h
--- a/Engine/Plugins/Experimental/Mover/Source/Mover/Public/PhysicsMover/Modes/PhysicsDrivenFlyingMode.h
+++ b/Engine/Plugins/Experimental/Mover/Source/Mover/Public/PhysicsMover/Modes/PhysicsDrivenFlyingMode.h
@@ -25,6 +25,10 @@ public:
 
 	virtual void UpdateConstraintSettings(Chaos::FCharacterGroundConstraint& Constraint) const override;
 
+#if WITH_EDITOR
+	virtual EDataValidationResult IsDataValid(class FDataValidationContext& Context) const override;
+#endif // WITH_EDITOR
+
 	// Maximum torque the character can apply to rotate in air about the vertical axis
 	UPROPERTY(EditAnywhere, Category = "Physics Mover", meta = (ClampMin = "0", UIMin = "0", ForceUnits = "NewtonMeters"))
 	float TwistTorqueLimit = 0.0f;
diff --git a/Engine/Plugins/Experimental/Mover/Source/Mover/Public/PhysicsMover/Modes/PhysicsDrivenSwimmingMode.h b/Engine/Plugins/Experimental/Mover/Source/Mover/Public/PhysicsMover/Modes/PhysicsDrivenSwimmingMode.h
--- a/Engine/Plugins/Experimental/Mover/Source/Mover/Public/PhysicsMover/Modes/PhysicsDrivenSwimmingMode.h
+++ b/Engine/Plugins/Experimental/Mover/Source/Mover/Public/PhysicsMover/Modes/PhysicsDrivenSwimmingMode.h
@@ -27,6 +27,10 @@ public:
 
 	virtual void UpdateConstraintSettings(Chaos::FCharacterGroundConstraint& Constraint) const override;
 
+#if WITH_EDITOR
+	virtual EDataValidationResult IsDataValid(class FDataValidationContext& Context) const override;
+#endif // WITH_EDITOR
+
 	// Target height for the character. This is the desired distance from the center of the capsule to the floor
 	UPROPERTY(EditAnywhere, Category = "Physics Mover", meta = (ClampMin = "0", UIMin = "0", ForceUnits = "cm"))
 	float TargetHeight = 54.0f;
diff --git a/Engine/Plugins/Experimental/Mover/Source/Mover/Public/PhysicsMover/Modes/PhysicsDrivenWalkingMode.h b/Engine/Plugins/Experimental/Mover/Source/Mover/Public/PhysicsMover/Modes/PhysicsDrivenWalkingMode.h
--- a/Engine/Plugins/Experimental/Mover/Source/Mover/Public/PhysicsMover/Modes/PhysicsDrivenWalkingMode.h
+++ b/Engine/Plugins/Experimental/Mover/Source/Mover/Public/PhysicsMover/Modes/PhysicsDrivenWalkingMode.h
@@ -26,6 +26,10 @@ public:
 
 	virtual void UpdateConstraintSettings(Chaos::FCharacterGroundConstraint& Constraint) const override;
 
+#if WITH_EDITOR
+	virtual EDataValidationResult IsDataValid(class FDataValidationContext& Context) const override;
+#endif // WITH_EDITOR
+
 	// Maximum force the character can apply to reach the motion target
 	UPROPERTY(EditAnywhere, Category = "Physics Mover", meta = (ClampMin = "0", UIMin = "0", ForceUnits = "Newtons"))
 	float RadialForceLimit = 1500.0f;
