Mover: fixes issue where a new BP actor with a Mover component would fail to compile until the asset was reloaded. Now you just need to save it. Improved data validation messaging.

diff --git a/Engine/Plugins/Experimental/Mover/Source/Mover/Private/MoverComponent.cpp b/Engine/Plugins/Experimental/Mover/Source/Mover/Private/MoverComponent.cpp
--- a/Engine/Plugins/Experimental/Mover/Source/Mover/Private/MoverComponent.cpp
+++ b/Engine/Plugins/Experimental/Mover/Source/Mover/Private/MoverComponent.cpp
@@ -20,6 +20,7 @@
 #include "Misc/AssertionMacros.h"
 #include "Misc/TransactionObjectEvent.h"
 #include "Blueprint/BlueprintExceptionInfo.h"
+#include "UObject/ObjectSaveContext.h"
 
 #if WITH_EDITOR
 #include "Misc/DataValidation.h"
@@ -536,8 +537,16 @@ void UMoverComponent::ProcessFirstSimTickAfterRollback(const FMoverTimeStep& Tim
 	bHasRolledBack = false;
 }
 
+
 #if WITH_EDITOR
 
+void UMoverComponent::PreSave(FObjectPreSaveContext ObjectSaveContext)
+{
+	Super::PreSave(ObjectSaveContext);
+
+	RefreshSharedSettings();
+}
+
 void UMoverComponent::PostCDOCompiled(const FPostCDOCompiledContext& Context)
 {
 	Super::PostCDOCompiled(Context);
@@ -614,7 +623,7 @@ bool UMoverComponent::ValidateSetup(FDataValidationContext& Context) const
 		if (Element.Value && Element.Value->SharedSettingsClass &&
 			FindSharedSettings(Element.Value->SharedSettingsClass) == nullptr)
 		{
-			Context.AddError(FText::Format(LOCTEXT("MissingModeSettingsError", "Movement mode on {0}, mapped as {1}, is missing its desired SharedSettingsClass {2} and may not function properly."),
+			Context.AddError(FText::Format(LOCTEXT("MissingModeSettingsError", "Movement mode on {0}, mapped as {1}, is missing its desired SharedSettingsClass {2}. You may need to save the asset and/or recompile."),
 				FText::FromString(GetNameSafe(GetOwner())),
 				FText::FromName(Element.Key),
 				FText::FromString(Element.Value->SharedSettingsClass->GetName())));
diff --git a/Engine/Plugins/Experimental/Mover/Source/Mover/Public/MoverComponent.h b/Engine/Plugins/Experimental/Mover/Source/Mover/Public/MoverComponent.h
--- a/Engine/Plugins/Experimental/Mover/Source/Mover/Public/MoverComponent.h
+++ b/Engine/Plugins/Experimental/Mover/Source/Mover/Public/MoverComponent.h
@@ -331,6 +331,7 @@ protected:
 	void ProcessFirstSimTickAfterRollback(const FMoverTimeStep& TimeStep);
 
 	#if WITH_EDITOR
+	virtual void PreSave(FObjectPreSaveContext ObjectSaveContext) override;
 	virtual void PostCDOCompiled(const FPostCDOCompiledContext& Context) override;
 	virtual void PostEditChangeProperty(FPropertyChangedEvent& PropertyChangedEvent) override;
 	virtual void PostTransacted(const FTransactionObjectEvent& TransactionEvent) override;
