Mover - Footsteps FX Fix

#changelist validated
#virtualized

diff --git a/Engine/Plugins/Experimental/Mover/Source/Mover/Private/PhysicsMover/Modes/PhysicsDrivenFallingMode.cpp b/Engine/Plugins/Experimental/Mover/Source/Mover/Private/PhysicsMover/Modes/PhysicsDrivenFallingMode.cpp
--- a/Engine/Plugins/Experimental/Mover/Source/Mover/Private/PhysicsMover/Modes/PhysicsDrivenFallingMode.cpp
+++ b/Engine/Plugins/Experimental/Mover/Source/Mover/Private/PhysicsMover/Modes/PhysicsDrivenFallingMode.cpp
@@ -99,7 +99,8 @@ void UPhysicsDrivenFallingMode::OnSimulationTick(const FSimulationTickParams& Pa
 	SimBlackboard->Set(CommonBlackboard::LastWaterResult, WaterResult);
 
 	const bool bIsMovingUp = CharacterInputs->bIsJumpJustPressed || (UpDir.Dot(ProposedMove.LinearVelocity) > 0.0f);
-	const bool bStartSwimming = WaterResult.WaterSplineData.ImmersionDepth > CommonLegacySettings->SwimmingStartImmersionDepth;
+	const float ProjectedImmersionDepth = WaterResult.WaterSplineData.ImmersionDepth - (UpDir.Dot(ProposedMove.LinearVelocity) * DeltaSeconds);
+	const bool bStartSwimming = ProjectedImmersionDepth > CommonLegacySettings->SwimmingStartImmersionDepth;
 
 	if (WaterResult.IsSwimmableVolume() && bStartSwimming && !bIsMovingUp)
 	{
diff --git a/Engine/Plugins/Experimental/Mover/Source/Mover/Private/PhysicsMover/Modes/PhysicsDrivenWalkingMode.cpp b/Engine/Plugins/Experimental/Mover/Source/Mover/Private/PhysicsMover/Modes/PhysicsDrivenWalkingMode.cpp
--- a/Engine/Plugins/Experimental/Mover/Source/Mover/Private/PhysicsMover/Modes/PhysicsDrivenWalkingMode.cpp
+++ b/Engine/Plugins/Experimental/Mover/Source/Mover/Private/PhysicsMover/Modes/PhysicsDrivenWalkingMode.cpp
@@ -107,7 +107,7 @@ void UPhysicsDrivenWalkingMode::OnSimulationTick(const FSimulationTickParams& Pa
 	SimBlackboard->Set(CommonBlackboard::LastFloorResult, FloorResult);
 	SimBlackboard->Set(CommonBlackboard::LastWaterResult, WaterResult);
 
-	const bool bStartSwimming = (WaterResult.WaterSplineData.ImmersionDepth+TargetHeight) > CommonLegacySettings->SwimmingIdealImmersionDepth;
+	const bool bStartSwimming = WaterResult.WaterSplineData.ImmersionDepth > CommonLegacySettings->SwimmingStartImmersionDepth;
 	
 	if (WaterResult.IsSwimmableVolume() && bStartSwimming)
 	{
diff --git a/Engine/Plugins/Experimental/Mover/Source/Mover/Private/PhysicsMover/PhysicsMovementUtils.cpp b/Engine/Plugins/Experimental/Mover/Source/Mover/Private/PhysicsMover/PhysicsMovementUtils.cpp
--- a/Engine/Plugins/Experimental/Mover/Source/Mover/Private/PhysicsMover/PhysicsMovementUtils.cpp
+++ b/Engine/Plugins/Experimental/Mover/Source/Mover/Private/PhysicsMover/PhysicsMovementUtils.cpp
@@ -45,6 +45,8 @@ void UPhysicsMovementUtils::FindFloor(
 		const FVector DeltaPosVert = DeltaPos.ProjectOnTo(UpDir);
 		const FVector DeltaPosHoriz = DeltaPos - DeltaPosVert;
 
+		const FVector WaterQueryLocation = Location - (TargetHeight * UpDir);
+
 		const float StartOffset = TargetHeight - MaxStepHeight - QueryRadius;
 		const float SweepDistance = MaxStepHeight + QueryRadius + FMath::Max(MaxStepHeight, -DeltaPosVert.Dot(UpDir));
 
@@ -84,7 +86,7 @@ void UPhysicsMovementUtils::FindFloor(
 			}
 #endif
 
-			GetWaterResultFromHitResults(Hits, Location, OutWaterResult);
+			GetWaterResultFromHitResults(Hits, WaterQueryLocation, TargetHeight, OutWaterResult);
 
 			if (bWalkable)
 			{
@@ -127,7 +129,7 @@ void UPhysicsMovementUtils::FindFloor(
 				{
 					Start = StartPoints[Idx];
 					End = Start - UpDir * SweepDistance;
-					World->LineTraceSingleByChannel(LineQueryHit[Idx], Start, End, CollisionChannel, QueryParams);
+					World->LineTraceSingleByChannel(LineQueryHit[Idx], Start, End, CollisionChannel, QueryParams, ResponseParams);
 
 					if (LineQueryHit[Idx].IsValidBlockingHit())
 					{
@@ -190,7 +192,7 @@ void UPhysicsMovementUtils::FindFloor(
 				// so run again with no blocking results
 				ResponseParams.CollisionResponse.SetAllChannels(ECR_Overlap);
 				World->SweepMultiByChannel(Hits, Start, End, FQuat::Identity, CollisionChannel, FCollisionShape::MakeSphere(QueryRadius), QueryParams, ResponseParams);
-				GetWaterResultFromHitResults(Hits, Location, OutWaterResult);
+				GetWaterResultFromHitResults(Hits, WaterQueryLocation, TargetHeight, OutWaterResult);
 			}
 		}
 		else
@@ -200,7 +202,7 @@ void UPhysicsMovementUtils::FindFloor(
 			OutFloorResult.FloorDist = 1.0e10f;
 
 			// In deep waters i.e., WaterDepth > SweepDistance, we may hit only water and no blocking surfaces
-			GetWaterResultFromHitResults(Hits, Location, OutWaterResult);
+			GetWaterResultFromHitResults(Hits, WaterQueryLocation, TargetHeight, OutWaterResult);
 		}
 	}
 }
@@ -299,7 +301,7 @@ FVector UPhysicsMovementUtils::ComputeIntegratedGroundVelocityFromHitResult(cons
 	return GroundVelocity;
 }
 
-bool UPhysicsMovementUtils::GetWaterResultFromHitResults(const TArray<FHitResult>& Hits, const FVector& Location, FWaterCheckResult& OutWaterResult)
+bool UPhysicsMovementUtils::GetWaterResultFromHitResults(const TArray<FHitResult>& Hits, const FVector& Location, const float TargetHeight, FWaterCheckResult& OutWaterResult)
 {
 	// Find the closet hit that is a water body
 	// Note: Relies on ordering of hit results
@@ -309,20 +311,40 @@ bool UPhysicsMovementUtils::GetWaterResultFromHitResults(const TArray<FHitResult
 
 		if (Hit.Component.IsValid())
 		{
-			if (const AActor* Actor = Hit.Component->GetOwner())
+			if (AActor* Actor = Hit.Component->GetOwner())
 			{
 				if (Actor->IsA(AWaterBody::StaticClass()))
 				{
+					AWaterBody* WaterBody = Cast<AWaterBody>(Actor);
+
 					OutWaterResult.HitResult = Hit;
 					OutWaterResult.bSwimmableVolume = true;
 
-					FWaterBodyQueryResult QueryResult = Cast<AWaterBody>(Actor)->GetWaterBodyComponent()->QueryWaterInfoClosestToWorldLocation(
-						Location, EWaterBodyQueryFlags::ComputeDepth | EWaterBodyQueryFlags::ComputeVelocity | EWaterBodyQueryFlags::ComputeImmersionDepth);
+					OutWaterResult.WaterSplineData.SplineInputKey = WaterBody->GetWaterBodyComponent()->FindInputKeyClosestToWorldLocation(Location);
+
+					OutWaterResult.WaterSplineData.WaterBody = WaterBody;
+
+					FWaterBodyQueryResult QueryResult = WaterBody->GetWaterBodyComponent()->QueryWaterInfoClosestToWorldLocation(
+							Location, 
+							EWaterBodyQueryFlags::ComputeLocation | EWaterBodyQueryFlags::ComputeNormal | EWaterBodyQueryFlags::ComputeImmersionDepth,
+							OutWaterResult.WaterSplineData.SplineInputKey
+						);
 
-					OutWaterResult.WaterSplineData.WaterDepth = QueryResult.GetWaterSurfaceDepth();
-					OutWaterResult.WaterSplineData.RawWaterVelocity = QueryResult.GetVelocity();
 					OutWaterResult.WaterSplineData.ImmersionDepth = QueryResult.GetImmersionDepth();
 
+					OutWaterResult.WaterSplineData.WaterPlaneLocation = QueryResult.GetWaterPlaneLocation();
+					OutWaterResult.WaterSplineData.WaterPlaneNormal = QueryResult.GetWaterPlaneNormal();
+
+					const float CapsuleBottom = Location.Z;
+					const float CapsuleTop = Location.Z + (TargetHeight * 2);
+					OutWaterResult.WaterSplineData.WaterSurfaceLocation = QueryResult.GetWaterSurfaceLocation();
+			
+					OutWaterResult.WaterSplineData.WaterSurfaceOffset = OutWaterResult.WaterSplineData.WaterSurfaceLocation - Location;
+			
+					OutWaterResult.WaterSplineData.ImmersionPercent = FMath::Clamp((OutWaterResult.WaterSplineData.WaterSurfaceLocation.Z - CapsuleBottom) / (CapsuleTop - CapsuleBottom), 0.f, 1.f);
+			
+					OutWaterResult.WaterSplineData.WaterSurfaceNormal = QueryResult.GetWaterSurfaceNormal();
+
 #if PHYSICSDRIVENMOTION_DEBUG_DRAW
 					Chaos::FDebugDrawQueue::GetInstance().DrawDebugLine(Location, Location - FVector::UpVector * OutWaterResult.WaterSplineData.ImmersionDepth, FColor::Blue, false, -1.f, 10, 1.0f);
 					Chaos::FDebugDrawQueue::GetInstance().DrawDebugPoint(OutWaterResult.HitResult.Location, FColor::Blue, false, -1.f, 10, 1.0f);
diff --git a/Engine/Plugins/Experimental/Mover/Source/Mover/Public/PhysicsMover/PhysicsMovementUtils.h b/Engine/Plugins/Experimental/Mover/Source/Mover/Public/PhysicsMover/PhysicsMovementUtils.h
--- a/Engine/Plugins/Experimental/Mover/Source/Mover/Public/PhysicsMover/PhysicsMovementUtils.h
+++ b/Engine/Plugins/Experimental/Mover/Source/Mover/Public/PhysicsMover/PhysicsMovementUtils.h
@@ -26,7 +26,7 @@ public:
 	static bool IsHitSurfaceWalkableWithStepUpCheck(const FHitResult& Hit, float StepHeight, float MaxStepHeight, float MinStepUpHeight, float MaxWalkSlopeCosine);
 
 	// Checks if any hit is with water and, if so, fills in the OutWaterResult
-	static bool GetWaterResultFromHitResults(const TArray<FHitResult>& Hits, const FVector& Location, FWaterCheckResult& OutWaterResult);
+	static bool GetWaterResultFromHitResults(const TArray<FHitResult>& Hits, const FVector& Location, const float TargetHeight, FWaterCheckResult& OutWaterResult);
 
 	// Returns the current ground velocity at the character position
 	static FVector ComputeGroundVelocityFromHitResult(const FVector& CharacterPosition, const FHitResult& FloorHit, const float DeltaSeconds);
