Mover 2.0 initializing mover component start state on initialize and updating starting sync state on begin play if location has changed to avoid a warning about desynced positions. Also, downgrading an ensure to a warning to avoid pausing the editor for an issue that could potentially fix itself

diff --git a/Engine/Plugins/Experimental/Mover/Source/Mover/Private/Backends/MoverNetworkPredictionLiaison.cpp b/Engine/Plugins/Experimental/Mover/Source/Mover/Private/Backends/MoverNetworkPredictionLiaison.cpp
--- a/Engine/Plugins/Experimental/Mover/Source/Mover/Private/Backends/MoverNetworkPredictionLiaison.cpp
+++ b/Engine/Plugins/Experimental/Mover/Source/Mover/Private/Backends/MoverNetworkPredictionLiaison.cpp
@@ -62,6 +62,7 @@ void UMoverNetworkPredictionLiaisonComponent::InitializeSimulationState(FMoverSy
 	check(MoverComp);
 	StartingOutSync = OutSync;
 	StartingOutAux = OutAux;
+	MoverComp->InitializeSimulationState(StartingOutSync, StartingOutAux);
 }
 
 void UMoverNetworkPredictionLiaisonComponent::SimulationTick(const FNetSimTimeStep& TimeStep, const TNetSimInput<KinematicMoverStateTypes>& SimInput, const TNetSimOutput<KinematicMoverStateTypes>& SimOutput)
@@ -119,9 +120,17 @@ void UMoverNetworkPredictionLiaisonComponent::BeginPlay()
 	Super::BeginPlay();
 	if (StartingOutSync && StartingOutAux)
 	{
-		MoverComp->InitializeSimulationState(StartingOutSync, StartingOutAux);
-		StartingOutSync = nullptr;
-		StartingOutAux = nullptr;
+		if (FMoverDefaultSyncState* StartingSyncState = StartingOutSync->SyncStateCollection.FindMutableDataByType<FMoverDefaultSyncState>())
+		{
+			const FTransform UpdatedComponentTransform = MoverComp->GetUpdatedComponentTransform();
+			// if our location has changed between initialization and begin play (ex: Actors sharing an exact start location and one gets "pushed" to make them fit) lets write the new location to avoid any disagreements
+			if (!UpdatedComponentTransform.GetLocation().Equals(StartingSyncState->GetLocation_WorldSpace()))
+			{
+				StartingSyncState->SetTransforms_WorldSpace(UpdatedComponentTransform.GetLocation(),
+													 UpdatedComponentTransform.GetRotation().Rotator(),
+													 FVector::ZeroVector);	// no initial velocity
+			}
+		}
 	}
 }
 
diff --git a/Engine/Plugins/Experimental/Mover/Source/Mover/Private/MoverComponent.cpp b/Engine/Plugins/Experimental/Mover/Source/Mover/Private/MoverComponent.cpp
--- a/Engine/Plugins/Experimental/Mover/Source/Mover/Private/MoverComponent.cpp
+++ b/Engine/Plugins/Experimental/Mover/Source/Mover/Private/MoverComponent.cpp
@@ -301,10 +301,14 @@ void UMoverComponent::SimulationTick(const FMoverTimeStep& InTimeStep, const FMo
 
 	if (const FMoverDefaultSyncState* StartingSyncState = SimInput.SyncState.SyncStateCollection.FindDataByType<FMoverDefaultSyncState>())
 	{
-		ensureMsgf(UpdatedComponent->GetComponentLocation().Equals(StartingSyncState->GetLocation_WorldSpace()) || StartingSyncState->GetMovementBase(),
-			TEXT("Simulation start location (%s) disagrees with actual mover component location (%s). This indicates movement of the component out-of-band with the simulation, and will lead to poor quality motion."),
+		if (!(UpdatedComponent->GetComponentLocation().Equals(StartingSyncState->GetLocation_WorldSpace()) || StartingSyncState->GetMovementBase()))
+		{
+			UE_LOG(LogMover, Warning, TEXT("%s %s: Simulation start location (%s) disagrees with actual mover component location (%s). This indicates movement of the component out-of-band with the simulation, and if happens often will lead to poor quality motion."),
+			*GetNameSafe(GetOwner()),
+			*StaticEnum<ENetRole>()->GetValueAsString(GetOwnerRole()),
 			*StartingSyncState->GetLocation_WorldSpace().ToCompactString(),
 			*UpdatedComponent->GetComponentLocation().ToCompactString());
+		}
 	}
 
 	// Sync state data should carry over between frames
