Mover: shuffling minimal required properties out of the sync state collection and onto the core sync state, since they are never optional.

diff --git a/Engine/Plugins/Experimental/Mover/Source/Mover/Private/Backends/MoverNetworkPhysicsLiaison.cpp b/Engine/Plugins/Experimental/Mover/Source/Mover/Private/Backends/MoverNetworkPhysicsLiaison.cpp
--- a/Engine/Plugins/Experimental/Mover/Source/Mover/Private/Backends/MoverNetworkPhysicsLiaison.cpp
+++ b/Engine/Plugins/Experimental/Mover/Source/Mover/Private/Backends/MoverNetworkPhysicsLiaison.cpp
@@ -643,8 +643,7 @@ void UMoverNetworkPhysicsLiaisonComponent::ProduceInput_External(float DeltaSeco
 		}
 		else
 		{
-			FMoverDefaultSyncState& Default = Input.SyncState.SyncStateCollection.FindOrAddMutableDataByType<FMoverDefaultSyncState>();
-			Default.MovementMode = MoverComp->StartingMovementMode;
+			Input.SyncState.MovementMode = MoverComp->StartingMovementMode;
 		}
 
 		if (bCachedInputIsValid)
@@ -699,19 +698,13 @@ void UMoverNetworkPhysicsLiaisonComponent::ConsumeOutput_External(const FPhysics
 					// TODO: Generalize handling of events generated on physics thread
 					if (MoverComp->HasValidCachedState())
 					{
-						if (const FMoverDefaultSyncState* CachedSyncState = MoverComp->GetSyncState().SyncStateCollection.FindDataByType<FMoverDefaultSyncState>())
+						if ((MoverComp->GetSyncState().MovementMode == DefaultModeNames::Falling) && (InterpolatedSyncState.MovementMode == DefaultModeNames::Walking))
 						{
-							if (const FMoverDefaultSyncState* InterpolatedState = InterpolatedSyncState.SyncStateCollection.FindDataByType<FMoverDefaultSyncState>())
+							if (UFallingMode* FallingMode = MoverComp->FindMode_Mutable<UFallingMode>())
 							{
-								if ((CachedSyncState->MovementMode == DefaultModeNames::Falling) && (InterpolatedState->MovementMode == DefaultModeNames::Walking))
-								{
-									if (UFallingMode* FallingMode = MoverComp->FindMode_Mutable<UFallingMode>())
-									{
-										FHitResult HitResult;
-										MoverComp->TryGetFloorCheckHitResult(HitResult);
-										FallingMode->OnLanded.Broadcast(InterpolatedState->MovementMode, HitResult);
-									}
-								}
+								FHitResult HitResult;
+								MoverComp->TryGetFloorCheckHitResult(HitResult);
+								FallingMode->OnLanded.Broadcast(InterpolatedSyncState.MovementMode, HitResult);
 							}
 						}
 					}
@@ -825,18 +818,18 @@ void UMoverNetworkPhysicsLiaisonComponent::OnPreSimulate_Internal(const FPhysics
 		return;
 	}
 
-	FMoverDefaultSyncState& SyncState = Input.SyncState.SyncStateCollection.FindOrAddMutableDataByType<FMoverDefaultSyncState>();
 
-	const IPhysicsCharacterMovementModeInterface* PhysicsMode = Cast<const IPhysicsCharacterMovementModeInterface>(MoverComp->ModeFSM->FindMovementMode(SyncState.MovementMode));
+
+	const IPhysicsCharacterMovementModeInterface* PhysicsMode = Cast<const IPhysicsCharacterMovementModeInterface>(MoverComp->ModeFSM->FindMovementMode(Input.SyncState.MovementMode));
 	if (!PhysicsMode)
 	{
-		UE_LOG(LogMover, Verbose, TEXT("Attempting to run non-physics movement mode %s in physics mover update."), *SyncState.MovementMode.ToString());
+		UE_LOG(LogMover, Verbose, TEXT("Attempting to run non-physics movement mode %s in physics mover update."), *Input.SyncState.MovementMode.ToString());
 		return;
 	}
 
 	// Make the sync state velocity relative to the ground if walking
 	FVector LocalGroundVelocity = FVector::ZeroVector;
-	if (SyncState.MovementMode == DefaultModeNames::Walking)
+	if (Input.SyncState.MovementMode == DefaultModeNames::Walking)
 	{
 		if (const UMoverBlackboard* Blackboard = MoverComp->GetSimBlackboard())
 		{
@@ -848,6 +841,8 @@ void UMoverNetworkPhysicsLiaisonComponent::OnPreSimulate_Internal(const FPhysics
 			}
 		}
 	}
+
+	FMoverDefaultSyncState& SyncState = Input.SyncState.SyncStateCollection.FindOrAddMutableDataByType<FMoverDefaultSyncState>();
 	SyncState.SetTransforms_WorldSpace(CharacterParticle->GetX(), FRotator(CharacterParticle->GetR()), CharacterParticle->GetV() - LocalGroundVelocity);
 
 	///////////////////////////////////////////////////////////////////////////////////////////////////////////////////////////
@@ -887,7 +882,7 @@ void UMoverNetworkPhysicsLiaisonComponent::OnPreSimulate_Internal(const FPhysics
 	FMoverDefaultSyncState& OutputSyncState = Output.SyncState.SyncStateCollection.FindOrAddMutableDataByType<FMoverDefaultSyncState>();
 
 	const FName MovementModeAfterTick = MoverComp->ModeFSM->GetCurrentModeName();
-	OutputSyncState.MovementMode = MovementModeAfterTick;
+	Output.SyncState.MovementMode = MovementModeAfterTick;
 
 	MoverComp->SimBlackboard->TryGet(CommonBlackboard::LastFloorResult, Output.FloorResult);
 
@@ -905,7 +900,7 @@ void UMoverNetworkPhysicsLiaisonComponent::OnPreSimulate_Internal(const FPhysics
 	FVector TargetVelocity = OutputSyncState.GetVelocity_WorldSpace() + LocalGroundVelocity;
 
 	// Landed so add the new ground velocity
-	if ((OutputSyncState.MovementMode == DefaultModeNames::Walking) && (SyncState.MovementMode != DefaultModeNames::Walking))
+	if ((Output.SyncState.MovementMode == DefaultModeNames::Walking) && (Input.SyncState.MovementMode != DefaultModeNames::Walking))
 	{
 		if (const UPhysicsDrivenWalkingMode* WalkingMode = Cast<UPhysicsDrivenWalkingMode>(MoverComp->FindMovementMode(UPhysicsDrivenWalkingMode::StaticClass())))
 		{
@@ -1060,59 +1055,56 @@ void UMoverNetworkPhysicsLiaisonComponent::OnContactModification_Internal(const
 		return;
 	}
 
-	if (const FMoverDefaultSyncState* SyncState = Input.SyncState.SyncStateCollection.FindDataByType<FMoverDefaultSyncState>())
-	{
-		float PawnHalfHeight;
-		float PawnRadius;
-		MoverComp->UpdatedComponent->CalcBoundingCylinder(PawnRadius, PawnHalfHeight);
+	float PawnHalfHeight;
+	float PawnRadius;
+	MoverComp->UpdatedComponent->CalcBoundingCylinder(PawnRadius, PawnHalfHeight);
 
-		const float CharacterHeight = CharacterParticle->GetX().Z;
-		const float EndCapHeight = CharacterHeight - PawnHalfHeight + PawnRadius;
+	const float CharacterHeight = CharacterParticle->GetX().Z;
+	const float EndCapHeight = CharacterHeight - PawnHalfHeight + PawnRadius;
 
-		const float CosThetaMax = 0.707f;
+	const float CosThetaMax = 0.707f;
 
-		float MinContactHeightStepUps = CharacterHeight - 1.0e10f;
-		if (SyncState->MovementMode == DefaultModeNames::Walking)
+	float MinContactHeightStepUps = CharacterHeight - 1.0e10f;
+	if (Input.SyncState.MovementMode == DefaultModeNames::Walking)
+	{
+		if (const UPhysicsDrivenWalkingMode* WalkingMode = Cast<UPhysicsDrivenWalkingMode>(MoverComp->FindMode_Mutable<UPhysicsDrivenWalkingMode>()))
 		{
-			if (const UPhysicsDrivenWalkingMode* WalkingMode = Cast<UPhysicsDrivenWalkingMode>(MoverComp->FindMode_Mutable<UPhysicsDrivenWalkingMode>()))
-			{
-				// Contacts on the character capsule below the MaxStepHeight tend to snag the character when stepping
-				// up or down, so disable them
+			// Contacts on the character capsule below the MaxStepHeight tend to snag the character when stepping
+			// up or down, so disable them
 
-				if (const UCommonLegacyMovementSettings* Settings = MoverComp->FindSharedSettings<UCommonLegacyMovementSettings>())
+			if (const UCommonLegacyMovementSettings* Settings = MoverComp->FindSharedSettings<UCommonLegacyMovementSettings>())
+			{
+				const float StepDistance = FMath::Abs(WalkingMode->TargetHeight - ConstraintHandle->GetData().GroundDistance);
+				if (StepDistance >= GPhysicsDrivenMotionDebugParams.MinStepUpDistance)
 				{
-					const float StepDistance = FMath::Abs(WalkingMode->TargetHeight - ConstraintHandle->GetData().GroundDistance);
-					if (StepDistance >= GPhysicsDrivenMotionDebugParams.MinStepUpDistance)
-					{
-						MinContactHeightStepUps = CharacterHeight - WalkingMode->TargetHeight + Settings->MaxStepHeight;
-					}
+					MinContactHeightStepUps = CharacterHeight - WalkingMode->TargetHeight + Settings->MaxStepHeight;
 				}
 			}
 		}
+	}
+
+	for (Chaos::FContactPairModifier& PairModifier : Modifier.GetContacts(CharacterParticle))
+	{
+		const int32 CharacterIdx = CharacterParticle == PairModifier.GetParticlePair()[0] ? 0 : 1;
+		const int32 OtherIdx = CharacterIdx == 0 ? 1 : 0;
 
-		for (Chaos::FContactPairModifier& PairModifier : Modifier.GetContacts(CharacterParticle))
+		for (int32 Idx = 0; Idx < PairModifier.GetNumContacts(); ++Idx)
 		{
-			const int32 CharacterIdx = CharacterParticle == PairModifier.GetParticlePair()[0] ? 0 : 1;
-			const int32 OtherIdx = CharacterIdx == 0 ? 1 : 0;
+			Chaos::FVec3 Point0, Point1;
+			PairModifier.GetWorldContactLocations(Idx, Point0, Point1);
+			Chaos::FVec3 CharacterPoint = CharacterIdx == 0 ? Point0 : Point1;
 
-			for (int32 Idx = 0; Idx < PairModifier.GetNumContacts(); ++Idx)
+			Chaos::FVec3 ContactNormal = PairModifier.GetWorldNormal(Idx);
+			if ((ContactNormal.Z > CosThetaMax) && CharacterPoint.Z < EndCapHeight)
 			{
-				Chaos::FVec3 Point0, Point1;
-				PairModifier.GetWorldContactLocations(Idx, Point0, Point1);
-				Chaos::FVec3 CharacterPoint = CharacterIdx == 0 ? Point0 : Point1;
-
-				Chaos::FVec3 ContactNormal = PairModifier.GetWorldNormal(Idx);
-				if ((ContactNormal.Z > CosThetaMax) && CharacterPoint.Z < EndCapHeight)
-				{
-					// Disable any nearly vertical contact with the end cap of the capsule
-					// This will be handled by the character ground constraint
-					PairModifier.SetContactPointDisabled(Idx);
-				}
-				else if ((CharacterPoint.Z < MinContactHeightStepUps) && (GroundParticle == PairModifier.GetParticlePair()[OtherIdx]))
-				{
-					// In the case of steps ups disable all contacts below the max step height
-					PairModifier.SetContactPointDisabled(Idx);
-				}
+				// Disable any nearly vertical contact with the end cap of the capsule
+				// This will be handled by the character ground constraint
+				PairModifier.SetContactPointDisabled(Idx);
+			}
+			else if ((CharacterPoint.Z < MinContactHeightStepUps) && (GroundParticle == PairModifier.GetParticlePair()[OtherIdx]))
+			{
+				// In the case of steps ups disable all contacts below the max step height
+				PairModifier.SetContactPointDisabled(Idx);
 			}
 		}
 	}
diff --git a/Engine/Plugins/Experimental/Mover/Source/Mover/Private/Debug/GameplayDebuggerCategory_Mover.cpp b/Engine/Plugins/Experimental/Mover/Source/Mover/Private/Debug/GameplayDebuggerCategory_Mover.cpp
--- a/Engine/Plugins/Experimental/Mover/Source/Mover/Private/Debug/GameplayDebuggerCategory_Mover.cpp
+++ b/Engine/Plugins/Experimental/Mover/Source/Mover/Private/Debug/GameplayDebuggerCategory_Mover.cpp
@@ -130,12 +130,9 @@ void FGameplayDebuggerCategory_Mover::CollectData(APlayerController* OwnerPC, AA
 		{
 			const FMoverSyncState& SyncState = MyMoverComponent->GetSyncState();
 
-			if (const FMoverDefaultSyncState* MoverState = SyncState.SyncStateCollection.FindDataByType<FMoverDefaultSyncState>())
+			for (auto it = SyncState.LayeredMoves.GetActiveMovesIterator(); it; ++it)
 			{
-				for (auto it = MoverState->LayeredMoves.GetActiveMovesIterator(); it; ++it)
-				{
-					DataPack.ActiveLayeredMoves.Add(*it->Get()->ToSimpleString());
-				}
+				DataPack.ActiveLayeredMoves.Add(*it->Get()->ToSimpleString());
 			}
 		}
 	}
diff --git a/Engine/Plugins/Experimental/Mover/Source/Mover/Private/DefaultMovementSet/Modes/SwimmingMode.cpp b/Engine/Plugins/Experimental/Mover/Source/Mover/Private/DefaultMovementSet/Modes/SwimmingMode.cpp
--- a/Engine/Plugins/Experimental/Mover/Source/Mover/Private/DefaultMovementSet/Modes/SwimmingMode.cpp
+++ b/Engine/Plugins/Experimental/Mover/Source/Mover/Private/DefaultMovementSet/Modes/SwimmingMode.cpp
@@ -183,12 +183,10 @@ void USwimmingMode::OnUnregistered()
 
 bool USwimmingMode::AttemptJump(float UpwardsSpeed, FMoverTickEndData& OutputState)
 {
-	FMoverDefaultSyncState& OutputSyncState = OutputState.SyncState.SyncStateCollection.FindOrAddMutableDataByType<FMoverDefaultSyncState>();
-
 	// TODO: This should check if a jump is even allowed
 	TSharedPtr<FLayeredMove_JumpImpulse> JumpMove = MakeShared<FLayeredMove_JumpImpulse>();
 	JumpMove->UpwardsSpeed = UpwardsSpeed;
-	OutputSyncState.LayeredMoves.QueueLayeredMove(JumpMove);
+	OutputState.SyncState.LayeredMoves.QueueLayeredMove(JumpMove);
 	OutputState.MovementEndState.NextModeName = DefaultModeNames::Falling;
 	return true;
 }
diff --git a/Engine/Plugins/Experimental/Mover/Source/Mover/Private/DefaultMovementSet/Modes/WalkingMode.cpp b/Engine/Plugins/Experimental/Mover/Source/Mover/Private/DefaultMovementSet/Modes/WalkingMode.cpp
--- a/Engine/Plugins/Experimental/Mover/Source/Mover/Private/DefaultMovementSet/Modes/WalkingMode.cpp
+++ b/Engine/Plugins/Experimental/Mover/Source/Mover/Private/DefaultMovementSet/Modes/WalkingMode.cpp
@@ -330,12 +330,10 @@ void UWalkingMode::OnUnregistered()
 
 bool UWalkingMode::AttemptJump(float UpwardsSpeed, FMoverTickEndData& OutputState)
 {
-	FMoverDefaultSyncState& OutputSyncState = OutputState.SyncState.SyncStateCollection.FindOrAddMutableDataByType<FMoverDefaultSyncState>();
-
 	// TODO: This should check if a jump is even allowed
 	TSharedPtr<FLayeredMove_JumpImpulse> JumpMove = MakeShared<FLayeredMove_JumpImpulse>();
 	JumpMove->UpwardsSpeed = UpwardsSpeed;
-	OutputSyncState.LayeredMoves.QueueLayeredMove(JumpMove);
+	OutputState.SyncState.LayeredMoves.QueueLayeredMove(JumpMove);
 	OutputState.MovementEndState.NextModeName = CommonLegacySettings->AirMovementModeName;
 	return true;
 }
diff --git a/Engine/Plugins/Experimental/Mover/Source/Mover/Private/MovementModeStateMachine.cpp b/Engine/Plugins/Experimental/Mover/Source/Mover/Private/MovementModeStateMachine.cpp
--- a/Engine/Plugins/Experimental/Mover/Source/Mover/Private/MovementModeStateMachine.cpp
+++ b/Engine/Plugins/Experimental/Mover/Source/Mover/Private/MovementModeStateMachine.cpp
@@ -141,7 +141,7 @@ void UMovementModeStateMachine::OnSimulationTick(USceneComponent* UpdatedCompone
 
 	if (!QueuedModeTransition->IsSet())
 	{
-		QueueNextMode(SubstepStartSyncState->MovementMode);
+		QueueNextMode(SubstepStartData.SyncState.MovementMode);
 	}
 
 	AdvanceToNextMode();
@@ -152,16 +152,16 @@ void UMovementModeStateMachine::OnSimulationTick(USceneComponent* UpdatedCompone
 	{
 		SubstepStartSyncState = SubstepStartData.SyncState.SyncStateCollection.FindMutableDataByType<FMoverDefaultSyncState>();
 		FMoverDefaultSyncState* OutputSyncState = &OutputState.SyncState.SyncStateCollection.FindOrAddMutableDataByType<FMoverDefaultSyncState>();
-		OutputSyncState->MovementMode = CurrentModeName;
+		OutputState.SyncState.MovementMode = CurrentModeName;
 
 		OutputState.MovementEndState.ResetToDefaults();
 
 		SubTimeStep.StepMs = EndingSimTimeMs - SubTimeStep.BaseSimTimeMs;		// TODO: convert this to an overridable function that can support MaxStepTime, MaxIterations, etc.
 
 		// Transfer any queued moves into the starting state. They'll be started during the move generation.
-		FlushQueuedMovesToGroup(SubstepStartSyncState->LayeredMoves);
-		OutputSyncState->LayeredMoves = SubstepStartSyncState->LayeredMoves;
-		FLayeredMoveGroup& CurrentLayeredMoves = OutputSyncState->LayeredMoves;
+		FlushQueuedMovesToGroup(SubstepStartData.SyncState.LayeredMoves);
+		OutputState.SyncState.LayeredMoves = SubstepStartData.SyncState.LayeredMoves;
+		FLayeredMoveGroup& CurrentLayeredMoves = OutputState.SyncState.LayeredMoves;
 
 		// Gather any layered move contributions
 		FProposedMove CombinedLayeredMove;
@@ -187,7 +187,7 @@ void UMovementModeStateMachine::OnSimulationTick(USceneComponent* UpdatedCompone
 		if (bHasLayeredMoveContributions && !CombinedLayeredMove.PreferredMode.IsNone())
 		{
 			SetModeImmediately(CombinedLayeredMove.PreferredMode);
-			OutputSyncState->MovementMode = CurrentModeName;
+			OutputState.SyncState.MovementMode = CurrentModeName;
 		}
 
 		// Merge proposed movement from the current mode with movement from layered moves
@@ -288,7 +288,7 @@ void UMovementModeStateMachine::OnSimulationTick(USceneComponent* UpdatedCompone
 
 		// Switch modes if necessary (note that this will allow exit/enter on the same state)
 		AdvanceToNextMode();
-		OutputSyncState->MovementMode = CurrentModeName;
+		OutputState.SyncState.MovementMode = CurrentModeName;
 
 		const float RemainingMs = FMath::Clamp(OutputState.MovementEndState.RemainingMs, 0.0f, SubTimeStep.StepMs);
 		SubTimeStep.BaseSimTimeMs += (SubTimeStep.StepMs - RemainingMs);
@@ -417,12 +417,9 @@ FTransitionEvalResult UImmediateMovementModeTransition::OnEvaluate(const FSimula
 		{
 			return FTransitionEvalResult(NextMode);
 		}
-		else if (const FMoverDefaultSyncState* SyncState = Params.StartState.SyncState.SyncStateCollection.FindDataByType<FMoverDefaultSyncState>())
+		else if (NextMode != Params.StartState.SyncState.MovementMode)
 		{
-			if (NextMode != SyncState->MovementMode)
-			{
-				return FTransitionEvalResult(NextMode);
-			}
+			return FTransitionEvalResult(NextMode);
 		}
 	}
 
diff --git a/Engine/Plugins/Experimental/Mover/Source/Mover/Private/MoverComponent.cpp b/Engine/Plugins/Experimental/Mover/Source/Mover/Private/MoverComponent.cpp
--- a/Engine/Plugins/Experimental/Mover/Source/Mover/Private/MoverComponent.cpp
+++ b/Engine/Plugins/Experimental/Mover/Source/Mover/Private/MoverComponent.cpp
@@ -267,10 +267,10 @@ void UMoverComponent::InitializeSimulationState(FMoverSyncState* OutSync, FMover
 		MoverState->SetTransforms_WorldSpace(UpdatedComponent->GetComponentLocation(),
 		                                     UpdatedComponent->GetComponentRotation(),
 		                                     FVector::ZeroVector);	// no initial velocity
-
-		MoverState->MovementMode = StartingMovementMode;
 	}
 
+	OutSync->MovementMode = StartingMovementMode;
+
 	*OutAux = FMoverAuxStateContext();
 
 	CachedLastSyncState = *OutSync;
@@ -350,7 +350,7 @@ void UMoverComponent::SimulationTick(const FMoverTimeStep& InTimeStep, const FMo
 	if (FMoverDefaultSyncState* OutputSyncState = SimOutput.SyncState.SyncStateCollection.FindMutableDataByType<FMoverDefaultSyncState>())
 	{
 		const FName MovementModeAfterTick = ModeFSM->GetCurrentModeName();
-		OutputSyncState->MovementMode = MovementModeAfterTick;
+		SimOutput.SyncState.MovementMode = MovementModeAfterTick;
 
 		if (UpdatedComponent->GetComponentLocation().Equals(OutputSyncState->GetLocation_WorldSpace()) == false ||
 			UpdatedComponent->GetComponentQuat().Rotator().Equals(OutputSyncState->GetOrientation_WorldSpace(), ROTATOR_TOLERANCE) == false)
@@ -828,20 +828,17 @@ const FLayeredMoveBase* UMoverComponent::FindActiveLayeredMoveByType(const UScri
 {
 	if (bHasValidCachedState)
 	{
-		if (const FMoverDefaultSyncState* MoverState = CachedLastSyncState.SyncStateCollection.FindDataByType<FMoverDefaultSyncState>())
+		for (auto it = CachedLastSyncState.LayeredMoves.GetActiveMovesIterator(); it; ++it)
 		{
-			for (auto it = MoverState->LayeredMoves.GetActiveMovesIterator(); it; ++it)
+			UStruct* CandidateStruct = it->Get()->GetScriptStruct();
+			while (CandidateStruct)
 			{
-				UStruct* CandidateStruct = it->Get()->GetScriptStruct();
-				while (CandidateStruct)
+				if (DataStructType == CandidateStruct)
 				{
-					if (DataStructType == CandidateStruct)
-					{
-						return it->Get();
-					}
-
-					CandidateStruct = CandidateStruct->GetSuperStruct();
+					return it->Get();
 				}
+
+				CandidateStruct = CandidateStruct->GetSuperStruct();
 			}
 		}
 	}
@@ -1035,10 +1032,7 @@ bool UMoverComponent::IsFalling() const
 { 
 	if (bHasValidCachedState)
 	{
-		if (const FMoverDefaultSyncState* SyncState = CachedLastSyncState.SyncStateCollection.FindDataByType<FMoverDefaultSyncState>())
-		{
-			return SyncState->MovementMode == DefaultModeNames::Falling;
-		}
+		return CachedLastSyncState.MovementMode == DefaultModeNames::Falling;
 	}
 
 	return false;
@@ -1048,10 +1042,7 @@ bool UMoverComponent::IsAirborne() const
 {
 	if (bHasValidCachedState)
 	{
-		if (const FMoverDefaultSyncState* SyncState = CachedLastSyncState.SyncStateCollection.FindDataByType<FMoverDefaultSyncState>())
-		{
-			return SyncState->MovementMode == DefaultModeNames::Flying || SyncState->MovementMode == DefaultModeNames::Falling;
-		}
+		return CachedLastSyncState.MovementMode == DefaultModeNames::Flying || CachedLastSyncState.MovementMode == DefaultModeNames::Falling;
 	}
 
 	return false;
@@ -1061,10 +1052,7 @@ bool UMoverComponent::IsOnGround() const
 {
 	if (bHasValidCachedState)
 	{
-		if (const FMoverDefaultSyncState* SyncState = CachedLastSyncState.SyncStateCollection.FindDataByType<FMoverDefaultSyncState>())
-		{
-			return SyncState->MovementMode == DefaultModeNames::Walking;
-		}
+		return CachedLastSyncState.MovementMode == DefaultModeNames::Walking;
 	}
 
 	return false;
@@ -1158,10 +1146,7 @@ FName UMoverComponent::GetMovementModeName() const
 { 
 	if (bHasValidCachedState)
 	{
-		if (const FMoverDefaultSyncState* SyncState = CachedLastSyncState.SyncStateCollection.FindDataByType<FMoverDefaultSyncState>())
-		{
-			return SyncState->MovementMode;
-		}
+		return CachedLastSyncState.MovementMode;
 	}
 
 	return NAME_None;
diff --git a/Engine/Plugins/Experimental/Mover/Source/Mover/Private/MoverDataModelTypes.cpp b/Engine/Plugins/Experimental/Mover/Source/Mover/Private/MoverDataModelTypes.cpp
--- a/Engine/Plugins/Experimental/Mover/Source/Mover/Private/MoverDataModelTypes.cpp
+++ b/Engine/Plugins/Experimental/Mover/Source/Mover/Private/MoverDataModelTypes.cpp
@@ -131,8 +131,6 @@ bool FMoverDefaultSyncState::NetSerialize(FArchive& Ar, UPackageMap* Map, bool&
 	SerializePackedVector<10, 16>(Velocity, Ar);
 	Orientation.SerializeCompressedShort(Ar);
 
-	Ar << MovementMode;
-
 	// Optional movement base
 	bool bIsUsingMovementBase = (Ar.IsSaving() ? (MovementBase != nullptr) : false);
 	Ar.SerializeBits(&bIsUsingMovementBase, 1);
@@ -150,8 +148,6 @@ bool FMoverDefaultSyncState::NetSerialize(FArchive& Ar, UPackageMap* Map, bool&
 		MovementBase = nullptr;
 	}
 
-	LayeredMoves.NetSerialize(Ar);
-
 	bOutSuccess = true;
 	return true;
 }
@@ -164,7 +160,6 @@ void FMoverDefaultSyncState::ToString(FAnsiStringBuilderBase& Out) const
 	Out.Appendf("Intent: X=%.2f Y=%.2f Z=%.2f\n", MoveDirectionIntent.X, MoveDirectionIntent.Y, MoveDirectionIntent.Z);
 	Out.Appendf("Vel: X=%.2f Y=%.2f Z=%.2f\n", Velocity.X, Velocity.Y, Velocity.Z);
 	Out.Appendf("Orient: P=%.2f Y=%.2f R=%.2f\n", Orientation.Pitch, Orientation.Yaw, Orientation.Roll);
-	Out.Appendf("MovementMode: %s\n", TCHAR_TO_ANSI(*MovementMode.ToString()));
 
 	if (MovementBase)
 	{
@@ -176,7 +171,6 @@ void FMoverDefaultSyncState::ToString(FAnsiStringBuilderBase& Out) const
 		Out.Appendf("MovementBase: none\n");
 	}
 
-	Out.Appendf("Layered Moves: %s\n", TCHAR_TO_ANSI(*LayeredMoves.ToSimpleString()));
 }
 
 
@@ -219,8 +213,6 @@ void FMoverDefaultSyncState::Interpolate(const FMoverDataStructBase& From, const
 	else
 	{
 
-		MovementMode = ToState->MovementMode;
-
 		// No matter what base we started from, we always interpolate into the "To" movement base's space
 		MovementBase         = ToState->MovementBase;
 		MovementBaseBoneName = ToState->MovementBaseBoneName;
@@ -265,7 +257,6 @@ void FMoverDefaultSyncState::Interpolate(const FMoverDataStructBase& From, const
 		Velocity			= FMath::Lerp(FromVelocity_ToSpace,		ToState->Velocity, Pct);
 		Orientation			= FMath::Lerp(FromOrientation_ToSpace,	ToState->Orientation, Pct);
 
-		LayeredMoves = ToState->LayeredMoves;
 	}
 }
 
diff --git a/Engine/Plugins/Experimental/Mover/Source/Mover/Private/PhysicsMover/Modes/PhysicsDrivenSwimmingMode.cpp b/Engine/Plugins/Experimental/Mover/Source/Mover/Private/PhysicsMover/Modes/PhysicsDrivenSwimmingMode.cpp
--- a/Engine/Plugins/Experimental/Mover/Source/Mover/Private/PhysicsMover/Modes/PhysicsDrivenSwimmingMode.cpp
+++ b/Engine/Plugins/Experimental/Mover/Source/Mover/Private/PhysicsMover/Modes/PhysicsDrivenSwimmingMode.cpp
@@ -152,12 +152,10 @@ void UPhysicsDrivenSwimmingMode::OnSimulationTick(const FSimulationTickParams& P
 
 bool UPhysicsDrivenSwimmingMode::AttemptJump(float UpwardsSpeed, FMoverTickEndData& OutputState)
 {
-	FMoverDefaultSyncState& OutputSyncState = OutputState.SyncState.SyncStateCollection.FindOrAddMutableDataByType<FMoverDefaultSyncState>();
-
 	// TODO: This should check if a jump is even allowed
  	TSharedPtr<FLayeredMove_JumpImpulse> JumpMove = MakeShared<FLayeredMove_JumpImpulse>();
 	JumpMove->UpwardsSpeed = UpwardsSpeed;
-	OutputSyncState.LayeredMoves.QueueLayeredMove(JumpMove);
+	OutputState.SyncState.LayeredMoves.QueueLayeredMove(JumpMove);
 	OutputState.MovementEndState.NextModeName = DefaultModeNames::Falling;
 	return true;
 }
diff --git a/Engine/Plugins/Experimental/Mover/Source/Mover/Public/MoverDataModelTypes.h b/Engine/Plugins/Experimental/Mover/Source/Mover/Public/MoverDataModelTypes.h
--- a/Engine/Plugins/Experimental/Mover/Source/Mover/Public/MoverDataModelTypes.h
+++ b/Engine/Plugins/Experimental/Mover/Source/Mover/Public/MoverDataModelTypes.h
@@ -144,10 +144,6 @@ public:
 	UPROPERTY(BlueprintReadWrite, EditAnywhere, Category = Mover)
 	FVector MoveDirectionIntent;
 
-	// The mode we ended up in from the prior frame, and which we'll start in during the next frame
-	UPROPERTY(BlueprintReadWrite, EditAnywhere, Category = Mover)
-	FName MovementMode;
-
 protected:
 	// Optional: when moving on a base, input may be relative to this object
 	UPROPERTY(BlueprintReadOnly, Category = Mover)
@@ -163,11 +159,6 @@ protected:
 	UPROPERTY(BlueprintReadOnly, Category = Mover)
 	FQuat MovementBaseQuat;
 
-public:
-	// Additional moves influencing our proposed motion
-	UPROPERTY(BlueprintReadWrite, EditAnywhere, Category = Mover)
-	FLayeredMoveGroup LayeredMoves;
-
 public:
 
 	FMoverDefaultSyncState()
@@ -175,7 +166,6 @@ public:
 		, Orientation(ForceInitToZero)
 		, Velocity(ForceInitToZero)
 		, MoveDirectionIntent(ForceInitToZero)
-		, MovementMode(NAME_None)
 		, MovementBase(nullptr)
 		, MovementBaseBoneName(NAME_None)
 		, MovementBasePos(ForceInitToZero)
diff --git a/Engine/Plugins/Experimental/Mover/Source/Mover/Public/MoverSimulationTypes.h b/Engine/Plugins/Experimental/Mover/Source/Mover/Public/MoverSimulationTypes.h
--- a/Engine/Plugins/Experimental/Mover/Source/Mover/Public/MoverSimulationTypes.h
+++ b/Engine/Plugins/Experimental/Mover/Source/Mover/Public/MoverSimulationTypes.h
@@ -104,11 +104,21 @@ struct MOVER_API FMoverSyncState
 	GENERATED_USTRUCT_BODY()
 
 public:
+
+	// The mode we ended up in from the prior frame, and which we'll start in during the next frame
+	UPROPERTY(BlueprintReadWrite, EditAnywhere, Category = Mover)
+	FName MovementMode;
+
+	// Additional moves influencing our proposed motion
+	UPROPERTY(BlueprintReadWrite, EditAnywhere, Category = Mover)
+	FLayeredMoveGroup LayeredMoves;
+
 	UPROPERTY(BlueprintReadWrite, Category = Mover)
 	FMoverDataCollection SyncStateCollection;
 
 	FMoverSyncState()
 	{
+		MovementMode = NAME_None;
 		SyncStateCollection.FindOrAddDataByType<FMoverDefaultSyncState>();
 	}
 
@@ -116,22 +126,31 @@ public:
 
 	void NetSerialize(const FNetSerializeParams& P)
 	{
+		P.Ar << MovementMode;
+		LayeredMoves.NetSerialize(P.Ar);
+
 		bool bIgnoredResult(false);
 		SyncStateCollection.NetSerialize(P.Ar, nullptr, bIgnoredResult);
 	}
 
 	void ToString(FAnsiStringBuilderBase& Out) const
 	{
+		Out.Appendf("MovementMode: %s\n", TCHAR_TO_ANSI(*MovementMode.ToString()));
+		Out.Appendf("Layered Moves: %s\n", TCHAR_TO_ANSI(*LayeredMoves.ToSimpleString()));
 		SyncStateCollection.ToString(Out);
 	}
 
 	bool ShouldReconcile(const FMoverSyncState& AuthorityState) const
 	{
-		return SyncStateCollection.ShouldReconcile(AuthorityState.SyncStateCollection);		
+		return (MovementMode != AuthorityState.MovementMode) || 
+		       SyncStateCollection.ShouldReconcile(AuthorityState.SyncStateCollection);
 	}
 
 	void Interpolate(const FMoverSyncState* From, const FMoverSyncState* To, float Pct)
 	{
+		MovementMode = To->MovementMode;
+		LayeredMoves = To->LayeredMoves;
+
 		SyncStateCollection.Interpolate(From->SyncStateCollection, To->SyncStateCollection, Pct);
 	}
 
diff --git a/Engine/Plugins/Experimental/MoverExamples/Source/MoverExamples/Private/CharacterVariants/Ziplining/ZipliningTransitions.cpp b/Engine/Plugins/Experimental/MoverExamples/Source/MoverExamples/Private/CharacterVariants/Ziplining/ZipliningTransitions.cpp
--- a/Engine/Plugins/Experimental/MoverExamples/Source/MoverExamples/Private/CharacterVariants/Ziplining/ZipliningTransitions.cpp
+++ b/Engine/Plugins/Experimental/MoverExamples/Source/MoverExamples/Private/CharacterVariants/Ziplining/ZipliningTransitions.cpp
@@ -21,9 +21,9 @@ FTransitionEvalResult UZiplineStartTransition::OnEvaluate(const FSimulationTickP
 
 	UMoverComponent* MoverComp = Params.MoverComponent;
 
-	const FMoverDefaultSyncState* SyncState = Params.StartState.SyncState.SyncStateCollection.FindDataByType<FMoverDefaultSyncState>();
+	const FMoverSyncState& SyncState = Params.StartState.SyncState;
 
-	if (MoverComp && MoverComp->IsAirborne() && SyncState && SyncState->MovementMode != ZipliningModeName)
+	if (MoverComp && MoverComp->IsAirborne() && SyncState.MovementMode != ZipliningModeName)
 	{
 		if (const FMoverExampleAbilityInputs* AbilityInputs = Params.StartState.InputCmd.InputCollection.FindDataByType<FMoverExampleAbilityInputs>())
 		{
