Mover 2.0: Fixing issue where falling mode wouldn't respect terminal velocity and apply gravity after clamping to terminal velocity

diff --git a/Engine/Plugins/Experimental/Mover/Source/Mover/Private/DefaultMovementSet/Modes/FallingMode.cpp b/Engine/Plugins/Experimental/Mover/Source/Mover/Private/DefaultMovementSet/Modes/FallingMode.cpp
--- a/Engine/Plugins/Experimental/Mover/Source/Mover/Private/DefaultMovementSet/Modes/FallingMode.cpp
+++ b/Engine/Plugins/Experimental/Mover/Source/Mover/Private/DefaultMovementSet/Modes/FallingMode.cpp
@@ -104,29 +104,27 @@ void UFallingMode::OnGenerateMove(const FMoverTickStartData& StartState, const F
 	}
 	
 	OutProposedMove = UAirMovementUtils::ComputeControlledFreeMove(Params);
+	const FVector VelocityWithGravity = StartVelocity + UMovementUtils::ComputeVelocityFromGravity(GetMoverComponent()->GetGravityAcceleration(), DeltaSeconds);
 
 	//  If we are going faster than TerminalVerticalVelocity apply VerticalFallingDeceleration otherwise reset Z velocity to before we applied deceleration 
-	if (StartVelocity.GetAbs().Z > TerminalVerticalSpeed)
+	if (VelocityWithGravity.GetAbs().Z > TerminalVerticalSpeed)
 	{
 		if (bShouldClampTerminalVerticalSpeed)
 		{
-			OutProposedMove.LinearVelocity.Z = FMath::Sign(StartVelocity.Z) * TerminalVerticalSpeed;
+			OutProposedMove.LinearVelocity.Z = FMath::Sign(VelocityWithGravity.Z) * TerminalVerticalSpeed;
 		}
 		else
 		{
-			float DesiredDeceleration = FMath::Abs(TerminalVerticalSpeed - FMath::Abs(StartVelocity.Z)) / DeltaSeconds;
+			float DesiredDeceleration = FMath::Abs(TerminalVerticalSpeed - FMath::Abs(VelocityWithGravity.Z)) / DeltaSeconds;
 			float DecelerationToApply = FMath::Min(DesiredDeceleration, VerticalFallingDeceleration);
-			DecelerationToApply = FMath::Sign(StartVelocity.Z) * DecelerationToApply * DeltaSeconds;
-			OutProposedMove.LinearVelocity.Z = StartVelocity.Z - DecelerationToApply;
+			DecelerationToApply = FMath::Sign(VelocityWithGravity.Z) * DecelerationToApply * DeltaSeconds;
+			OutProposedMove.LinearVelocity.Z = VelocityWithGravity.Z - DecelerationToApply;
 		}
 	}
 	else
 	{
-		OutProposedMove.LinearVelocity.Z = StartVelocity.Z;	// restore original vertical velocity
+		OutProposedMove.LinearVelocity.Z = VelocityWithGravity.Z;
 	}
-
-	// Add velocity change due to gravity
-	OutProposedMove.LinearVelocity += UMovementUtils::ComputeVelocityFromGravity(GetMoverComponent()->GetGravityAcceleration(), DeltaSeconds);
 }
 
 void UFallingMode::OnSimulationTick(const FSimulationTickParams& Params, FMoverTickEndData& OutputState)
