Mover 2.0 adding velocity to the mover debugger

diff --git a/Engine/Plugins/Experimental/Mover/Source/Mover/Private/Debug/GameplayDebuggerCategory_Mover.cpp b/Engine/Plugins/Experimental/Mover/Source/Mover/Private/Debug/GameplayDebuggerCategory_Mover.cpp
--- a/Engine/Plugins/Experimental/Mover/Source/Mover/Private/Debug/GameplayDebuggerCategory_Mover.cpp
+++ b/Engine/Plugins/Experimental/Mover/Source/Mover/Private/Debug/GameplayDebuggerCategory_Mover.cpp
@@ -94,6 +94,7 @@ void FGameplayDebuggerCategory_Mover::CollectData(APlayerController* OwnerPC, AA
 	// Set defaults for info that may not be available
 	DataPack.MovementModeName = FString("invalid");
 	DataPack.MovementBaseInfo = FString("invalid");
+	DataPack.Velocity = FVector::ZeroVector;
 	DataPack.MoveIntent = FVector::ZeroVector;
 	DataPack.ActiveLayeredMoves.Empty();
 	DataPack.ModeMap.Empty();
@@ -106,6 +107,7 @@ void FGameplayDebuggerCategory_Mover::CollectData(APlayerController* OwnerPC, AA
 		DataPack.MovementModeName = MyMoverComponent->GetMovementModeName().ToString();
 		DataPack.MovementBaseInfo = MovementBaseComp ? FString::Printf(TEXT("%s.%s"), *GetNameSafe(MovementBaseComp->GetOwner()), *MovementBaseComp->GetName()) : FString();
 		DataPack.MoveIntent = MyMoverComponent->GetMovementIntent();
+		DataPack.Velocity = MyMoverComponent->GetVelocity();
 
 		for (auto ModeIter = MyMoverComponent->MovementModes.begin(); ModeIter; ++ModeIter)
 		{
@@ -151,10 +153,11 @@ void FGameplayDebuggerCategory_Mover::DrawData(APlayerController* OwnerPC, FGame
 		DrawInWorldInfo(*FocusedActor, CanvasContext);
 	}
 	
-	CanvasContext.Printf(TEXT("{yellow}%s\n{grey}Local Role: {white}%s\n{grey}Mode: {white}%s\n{yellow}Active Moves: {white}\n%s\n{yellow}Mode Map: \n{white}%s\n{yellow}Active Transitions: {white}\n%s"),
+	CanvasContext.Printf(TEXT("{yellow}%s\n{grey}Local Role: {white}%s\n{grey}Mode: {white}%s\n{grey}Velocity: {white}%s\n{yellow}Active Moves: {white}\n%s\n{yellow}Mode Map: \n{white}%s\n{yellow}Active Transitions: {white}\n%s"),
 		*DataPack.PawnName,
 		*DataPack.LocalRole,
 		*DataPack.MovementModeName,
+		*DataPack.Velocity.ToString(),
 		*FString::JoinBy(DataPack.ActiveLayeredMoves, TEXT("\n"), [](FString MoveAsString) { return MoveAsString; }),
 		*FString::JoinBy(DataPack.ModeMap, TEXT("\n"), [](FString ModeMappingAsString) { return ModeMappingAsString; }),
 		*FString::JoinBy(DataPack.ActiveTransitions, TEXT("\n"), [](FString TransitionAsString) { return TransitionAsString; })
@@ -262,6 +265,7 @@ void FGameplayDebuggerCategory_Mover::FRepData::Serialize(FArchive& Ar)
 	Ar << LocalRole;
 	Ar << MovementModeName;
 	Ar << MovementBaseInfo;
+	Ar << Velocity;
 	Ar << MoveIntent;
 	Ar << ActiveLayeredMoves;
 	Ar << ModeMap;
diff --git a/Engine/Plugins/Experimental/Mover/Source/Mover/Private/Debug/GameplayDebuggerCategory_Mover.h b/Engine/Plugins/Experimental/Mover/Source/Mover/Private/Debug/GameplayDebuggerCategory_Mover.h
--- a/Engine/Plugins/Experimental/Mover/Source/Mover/Private/Debug/GameplayDebuggerCategory_Mover.h
+++ b/Engine/Plugins/Experimental/Mover/Source/Mover/Private/Debug/GameplayDebuggerCategory_Mover.h
@@ -29,6 +29,7 @@ protected:
 		FString LocalRole;
 		FString MovementModeName;
 		FString MovementBaseInfo;
+		FVector Velocity;
 		FVector MoveIntent;
 		TArray<FString> ActiveLayeredMoves;
 		TArray<FString> ModeMap;
