Mover 2.0 adding some safety measures to AddMovementMode to avoid issues with passing in abstract classes

diff --git a/Engine/Plugins/Experimental/Mover/Source/Mover/Private/MoverComponent.cpp b/Engine/Plugins/Experimental/Mover/Source/Mover/Private/MoverComponent.cpp
--- a/Engine/Plugins/Experimental/Mover/Source/Mover/Private/MoverComponent.cpp
+++ b/Engine/Plugins/Experimental/Mover/Source/Mover/Private/MoverComponent.cpp
@@ -858,17 +858,26 @@ UBaseMovementMode* UMoverComponent::AddMovementModeFromClass(FName ModeName, TSu
 		UE_LOG(LogMover, Warning, TEXT("Attempted to add a movement mode that wasn't valid. AddMovementModeFromClass did not add anything."));
 		return nullptr;
 	}
+	if (MovementMode->HasAnyClassFlags(CLASS_Abstract))
+	{
+		UE_LOG(LogMover, Warning, TEXT("The Movement Mode class (%s) is abstract and is not a valid class to instantiate. AddMovementModeFromClass will not do anything."), *GetNameSafe(MovementMode));
+		return nullptr;
+	}
 
 	TObjectPtr<UBaseMovementMode> AddedMovementMode =  NewObject<UBaseMovementMode>(this, MovementMode);
-	AddMovementModeFromObject(ModeName, AddedMovementMode);
-	
-	return AddedMovementMode; 
+	return AddMovementModeFromObject(ModeName, AddedMovementMode) ? AddedMovementMode : nullptr;
 }
 
 bool UMoverComponent::AddMovementModeFromObject(FName ModeName, UBaseMovementMode* MovementMode)
 {
 	if (MovementMode)
 	{
+		if (MovementMode->GetClass()->HasAnyClassFlags(CLASS_Abstract))
+		{
+			UE_LOG(LogMover, Warning, TEXT("The Movement Mode class (%s) is abstract and is not a valid class to instantiate. AddMovementModeFromObject will not do anything."), *GetNameSafe(MovementMode));
+			return false;
+		}
+		
 		if (TObjectPtr<UBaseMovementMode>* FoundMovementMode = MovementModes.Find(ModeName))
 		{
 			if (FoundMovementMode->Get()->GetClass() == MovementMode->GetClass())
diff --git a/Engine/Plugins/Experimental/Mover/Source/Mover/Public/MoverComponent.h b/Engine/Plugins/Experimental/Mover/Source/Mover/Public/MoverComponent.h
--- a/Engine/Plugins/Experimental/Mover/Source/Mover/Public/MoverComponent.h
+++ b/Engine/Plugins/Experimental/Mover/Source/Mover/Public/MoverComponent.h
@@ -150,7 +150,7 @@ public:
 
 	// Add a movement mode to available movement modes. Returns true if the movement mode was added successfully. Returns the mode that was made.
 	UFUNCTION(BlueprintCallable, Category = Mover, meta=(DeterminesOutputType="MovementMode"))
-	UBaseMovementMode* AddMovementModeFromClass(FName ModeName, TSubclassOf<UBaseMovementMode> MovementMode);
+	UBaseMovementMode* AddMovementModeFromClass(FName ModeName, UPARAM(meta = (AllowAbstract = "false"))TSubclassOf<UBaseMovementMode> MovementMode);
 
 	// Add a movement mode to available movement modes. Returns true if the movement mode was added successfully
 	UFUNCTION(BlueprintCallable, Category = Mover)
