Mover plugin: added data validation and runtime checks for valid transitions

diff --git a/Engine/Plugins/Experimental/Mover/Source/Mover/Private/MovementMode.cpp b/Engine/Plugins/Experimental/Mover/Source/Mover/Private/MovementMode.cpp
--- a/Engine/Plugins/Experimental/Mover/Source/Mover/Private/MovementMode.cpp
+++ b/Engine/Plugins/Experimental/Mover/Source/Mover/Private/MovementMode.cpp
@@ -3,9 +3,12 @@
 #include "MovementMode.h"
 #include "MoverComponent.h"
 #include "Engine/BlueprintGeneratedClass.h"
+#include "Misc/DataValidation.h"
 
 #include UE_INLINE_GENERATED_CPP_BY_NAME(MovementMode)
 
+#define LOCTEXT_NAMESPACE "Mover"
+
 UBaseMovementMode::UBaseMovementMode(const FObjectInitializer& ObjectInitializer)
 	: Super(ObjectInitializer)
 {
@@ -92,7 +95,14 @@ EDataValidationResult UBaseMovementMode::IsDataValid(FDataValidationContext& Con
 	EDataValidationResult Result = EDataValidationResult::Valid;
 	for (UBaseMovementModeTransition* Transition : Transitions)
 	{
-		if (Transition->IsDataValid(Context) == EDataValidationResult::Invalid)
+		if (!IsValid(Transition))
+		{
+			Context.AddError(FText::Format(LOCTEXT("InvalidTransitionError", "Invalid or missing transition object on mode of type {0}. Clean up the Transitions array."),
+				FText::FromString(GetClass()->GetName())));
+
+			Result = EDataValidationResult::Invalid;
+		}
+		else if (Transition->IsDataValid(Context) == EDataValidationResult::Invalid)
 		{
 			Result = EDataValidationResult::Invalid;
 		}
@@ -119,3 +129,4 @@ void UBaseMovementMode::OnSimulationTick(const FSimulationTickParams& Params, FM
 {
 }
 
+#undef LOCTEXT_NAMESPACE
diff --git a/Engine/Plugins/Experimental/Mover/Source/Mover/Private/MovementModeStateMachine.cpp b/Engine/Plugins/Experimental/Mover/Source/Mover/Private/MovementModeStateMachine.cpp
--- a/Engine/Plugins/Experimental/Mover/Source/Mover/Private/MovementModeStateMachine.cpp
+++ b/Engine/Plugins/Experimental/Mover/Source/Mover/Private/MovementModeStateMachine.cpp
@@ -230,12 +230,15 @@ void UMovementModeStateMachine::OnSimulationTick(USceneComponent* UpdatedCompone
 
 				for (UBaseMovementModeTransition* Transition : CurrentMode->Transitions)
 				{
-					EvalResult = Transition->DoEvaluate(SimTickParams);
-
-					if (!EvalResult.NextMode.IsNone())
+					if (IsValid(Transition))
 					{
-						TransitionToTrigger = Transition;
-						break;
+						EvalResult = Transition->DoEvaluate(SimTickParams);
+
+						if (!EvalResult.NextMode.IsNone())
+						{
+							TransitionToTrigger = Transition;
+							break;
+						}
 					}
 				}
 
@@ -243,12 +246,15 @@ void UMovementModeStateMachine::OnSimulationTick(USceneComponent* UpdatedCompone
 				{
 					for (UBaseMovementModeTransition* Transition : MoverComp->Transitions)
 					{
-						EvalResult = Transition->DoEvaluate(SimTickParams);
-
-						if (!EvalResult.NextMode.IsNone())
+						if (IsValid(Transition))
 						{
-							TransitionToTrigger = Transition;
-							break;
+							EvalResult = Transition->DoEvaluate(SimTickParams);
+
+							if (!EvalResult.NextMode.IsNone())
+							{
+								TransitionToTrigger = Transition;
+								break;
+							}
 						}
 					}
 				}
diff --git a/Engine/Plugins/Experimental/Mover/Source/Mover/Private/MoverComponent.cpp b/Engine/Plugins/Experimental/Mover/Source/Mover/Private/MoverComponent.cpp
--- a/Engine/Plugins/Experimental/Mover/Source/Mover/Private/MoverComponent.cpp
+++ b/Engine/Plugins/Experimental/Mover/Source/Mover/Private/MoverComponent.cpp
@@ -646,6 +646,18 @@ bool UMoverComponent::ValidateSetup(FDataValidationContext& Context) const
 		bDidFindAnyProblems = true;
 	}
 
+	// Verify transitions
+	for (const UBaseMovementModeTransition* Transition : Transitions)
+	{
+		if (!IsValid(Transition))
+		{
+			Context.AddError(FText::Format(LOCTEXT("InvalidTransitionError", "Invalid or missing transition object on {0}. Clean up the Transitions array."),
+				FText::FromString(GetNameSafe(GetOwner()))));
+
+			bDidFindAnyProblems = true;
+		}
+	}
+
 	// Verify backend liaison
 	if (!BackendClass)
 	{
