Fixing conflicting localization names. No functional changes.

#rnx

diff --git a/Engine/Plugins/Experimental/Mover/Source/Mover/Private/MovementMode.cpp b/Engine/Plugins/Experimental/Mover/Source/Mover/Private/MovementMode.cpp
--- a/Engine/Plugins/Experimental/Mover/Source/Mover/Private/MovementMode.cpp
+++ b/Engine/Plugins/Experimental/Mover/Source/Mover/Private/MovementMode.cpp
@@ -97,7 +97,7 @@ EDataValidationResult UBaseMovementMode::IsDataValid(FDataValidationContext& Con
 	{
 		if (!IsValid(Transition))
 		{
-			Context.AddError(FText::Format(LOCTEXT("InvalidTransitionError", "Invalid or missing transition object on mode of type {0}. Clean up the Transitions array."),
+			Context.AddError(FText::Format(LOCTEXT("InvalidTransitionOnModeError", "Invalid or missing transition object on mode of type {0}. Clean up the Transitions array."),
 				FText::FromString(GetClass()->GetName())));
 
 			Result = EDataValidationResult::Invalid;
