Physics Mover - Support non-async physics update

diff --git a/Engine/Plugins/Experimental/Mover/Source/Mover/Private/Backends/MoverNetworkPhysicsLiaison.cpp b/Engine/Plugins/Experimental/Mover/Source/Mover/Private/Backends/MoverNetworkPhysicsLiaison.cpp
--- a/Engine/Plugins/Experimental/Mover/Source/Mover/Private/Backends/MoverNetworkPhysicsLiaison.cpp
+++ b/Engine/Plugins/Experimental/Mover/Source/Mover/Private/Backends/MoverNetworkPhysicsLiaison.cpp
@@ -231,6 +231,17 @@ UMoverNetworkPhysicsLiaisonComponent::UMoverNetworkPhysicsLiaisonComponent()
 	bWantsInitializeComponent = true;
 	bAutoActivate = true;
 
+	if (UWorld* World = GetWorld())
+	{
+		if (FPhysScene_Chaos* PhysScene = GetWorld()->GetPhysicsScene())
+		{
+			if (Chaos::FPhysicsSolver* Solver = PhysScene->GetSolver())
+			{
+				bUsingAsyncPhysics = Solver->IsUsingAsyncResults();
+			}
+		}
+	}
+
 	if (Chaos::FPhysicsSolverBase::IsNetworkPhysicsPredictionEnabled())
 	{
 		// Network physics relies on movement being replicated
@@ -259,7 +270,7 @@ float UMoverNetworkPhysicsLiaisonComponent::GetCurrentSimTimeMs()
 	{
 		if (Chaos::FPhysicsSolver* Solver = Scene->GetSolver())
 		{
-			return Solver->GetAsyncDeltaTime() * GetCurrentSimFrame() * 1000.0f;
+			return bUsingAsyncPhysics ? Solver->GetAsyncDeltaTime() * GetCurrentSimFrame() * 1000.0f : Solver->GetSolverTime() * 1000.0f;
 		}
 	}
 
@@ -558,9 +569,10 @@ int32 UMoverNetworkPhysicsLiaisonComponent::GetNetworkPhysicsTickOffset() const
 	return 0;
 }
 
-FMoverTimeStep UMoverNetworkPhysicsLiaisonComponent::GetCurrentMoverTimeStep_Internal() const
+FMoverTimeStep UMoverNetworkPhysicsLiaisonComponent::GetCurrentAsyncMoverTimeStep_Internal() const
 {
 	Chaos::EnsureIsInPhysicsThreadContext();
+	ensure(bUsingAsyncPhysics);
 
 	FMoverTimeStep TimeStep;
 
@@ -578,9 +590,10 @@ FMoverTimeStep UMoverNetworkPhysicsLiaisonComponent::GetCurrentMoverTimeStep_Int
 	return TimeStep;
 }
 
-FMoverTimeStep UMoverNetworkPhysicsLiaisonComponent::GetCurrentMoverTimeStep_External() const
+FMoverTimeStep UMoverNetworkPhysicsLiaisonComponent::GetCurrentAsyncMoverTimeStep_External() const
 {
 	Chaos::EnsureIsInGameThreadContext();
+	ensure(bUsingAsyncPhysics);
 
 	FMoverTimeStep TimeStep;
 
@@ -599,6 +612,26 @@ FMoverTimeStep UMoverNetworkPhysicsLiaisonComponent::GetCurrentMoverTimeStep_Ext
 	return TimeStep;
 }
 
+FMoverTimeStep UMoverNetworkPhysicsLiaisonComponent::GetCurrentMoverTimeStep(float DeltaSeconds) const
+{
+	ensure(!bUsingAsyncPhysics);
+
+	FMoverTimeStep TimeStep;
+
+	if (FPhysScene* Scene = GetWorld()->GetPhysicsScene())
+	{
+		if (Chaos::FPhysicsSolver* Solver = Scene->GetSolver())
+		{
+			TimeStep.ServerFrame = Solver->GetCurrentFrame();
+			TimeStep.StepMs = DeltaSeconds * 1000.0f;
+			TimeStep.BaseSimTimeMs = Solver->GetSolverTime() * 1000.0f;
+			TimeStep.bIsResimulating = Solver->GetEvolution()->IsResimming();
+		}
+	}
+
+	return TimeStep;
+}
+
 void UMoverNetworkPhysicsLiaisonComponent::ProduceInput_External(float DeltaSeconds, OUT FPhysicsMoverAsyncInput& Input)
 {
 	Chaos::EnsureIsInGameThreadContext();
@@ -646,13 +679,16 @@ void UMoverNetworkPhysicsLiaisonComponent::ProduceInput_External(float DeltaSeco
 			Input.SyncState.MovementMode = MoverComp->StartingMovementMode;
 		}
 
+		const FMoverTimeStep MoverTimeStep = bUsingAsyncPhysics ? GetCurrentAsyncMoverTimeStep_External() : GetCurrentMoverTimeStep(DeltaSeconds);
+		MoverComp->CachedLastSimTickTimeStep = MoverTimeStep;
+
 		if (bCachedInputIsValid)
 		{
-			MoverComp->OnPreSimulationTick.Broadcast(GetCurrentMoverTimeStep_External(), Input.InputCmd);
+			MoverComp->OnPreSimulationTick.Broadcast(MoverTimeStep, Input.InputCmd);
 		}
 		else
 		{
-			MoverComp->OnPreSimulationTick.Broadcast(GetCurrentMoverTimeStep_External(), NetInputCmd);
+			MoverComp->OnPreSimulationTick.Broadcast(MoverTimeStep, NetInputCmd);
 		}
 
 		// This is required so that the physics thread can have a copy of the data to access
@@ -667,66 +703,80 @@ void UMoverNetworkPhysicsLiaisonComponent::ConsumeOutput_External(const FPhysics
 
 	if (Output.bIsValid && MoverComp)
 	{
-		if (bCachedLastPhysicsSyncStateIsValid)
-		{
-			const double PhysicsResultsTime = GetWorld()->GetPhysicsScene()->GetSolver()->GetPhysicsResultsTime_External();
+		const float DeltaSeconds = MoverComp->CachedLastSimTickTimeStep.StepMs * 0.001f;
+		const FMoverTimeStep MoverTimeStep = bUsingAsyncPhysics ? GetCurrentAsyncMoverTimeStep_External() : GetCurrentMoverTimeStep(DeltaSeconds);
 
-			if (PhysicsResultsTime < OutputTimeInSeconds)
+		if (bUsingAsyncPhysics)
+		{
+			if (bCachedLastPhysicsSyncStateIsValid)
 			{
-				const FMoverTimeStep TimeStep = GetCurrentMoverTimeStep_External();
+				const double PhysicsResultsTime = GetWorld()->GetPhysicsScene()->GetSolver()->GetPhysicsResultsTime_External();
 
-				// Output is in the future. Interpolate from the cached previous state to the next state
-				const float Denom = (OutputTimeInSeconds - CachedLastPhysicsSyncStateOutputTime);
-				if (Denom > 0)
+				if (PhysicsResultsTime < OutputTimeInSeconds)
 				{
-					const float Alpha = FMath::Clamp((PhysicsResultsTime - CachedLastPhysicsSyncStateOutputTime) / Denom, 0.0f, 1.0f);
-					FMoverSyncState InterpolatedSyncState;
-					InterpolatedSyncState.Interpolate(&CachedLastPhysicsSyncState, &Output.SyncState, Alpha);
 
-					// The physics mover does not set the sync state transform, it uses the physics transform,
-					// so get the world transform from the physics particle.
-					if (const Chaos::FSingleParticlePhysicsProxy* CharacterProxy = Constraint->GetCharacterParticleProxy())
+					// Output is in the future. Interpolate from the cached previous state to the next state
+					const float Denom = (OutputTimeInSeconds - CachedLastPhysicsSyncStateOutputTime);
+					if (Denom > 0)
 					{
-						const Chaos::FRigidBodyHandle_External& Body = CharacterProxy->GetGameThreadAPI();
-						if (FMoverDefaultSyncState* SyncState = InterpolatedSyncState.SyncStateCollection.FindMutableDataByType<FMoverDefaultSyncState>())
+						const float Alpha = FMath::Clamp((PhysicsResultsTime - CachedLastPhysicsSyncStateOutputTime) / Denom, 0.0f, 1.0f);
+						FMoverSyncState InterpolatedSyncState;
+						InterpolatedSyncState.Interpolate(&CachedLastPhysicsSyncState, &Output.SyncState, Alpha);
+
+						// The physics mover does not set the sync state transform, it uses the physics transform,
+						// so get the world transform from the physics particle.
+						if (const Chaos::FSingleParticlePhysicsProxy* CharacterProxy = Constraint->GetCharacterParticleProxy())
 						{
-							SyncState->SetTransforms_WorldSpace(Body.X(), FRotator(Body.R()), Body.V());
+							const Chaos::FRigidBodyHandle_External& Body = CharacterProxy->GetGameThreadAPI();
+							if (FMoverDefaultSyncState* SyncState = InterpolatedSyncState.SyncStateCollection.FindMutableDataByType<FMoverDefaultSyncState>())
+							{
+								SyncState->SetTransforms_WorldSpace(Body.X(), FRotator(Body.R()), Body.V());
+							}
 						}
-					}
 
-					// If landed, broadcast OnLanded
-					// TODO: Generalize handling of events generated on physics thread
-					if (MoverComp->HasValidCachedState())
-					{
-						if ((MoverComp->GetSyncState().MovementMode == DefaultModeNames::Falling) && (InterpolatedSyncState.MovementMode == DefaultModeNames::Walking))
+						// If landed, broadcast OnLanded
+						// TODO: Generalize handling of events generated on physics thread
+						if (MoverComp->HasValidCachedState())
 						{
-							if (UFallingMode* FallingMode = MoverComp->FindMode_Mutable<UFallingMode>())
+							if ((MoverComp->GetSyncState().MovementMode == DefaultModeNames::Falling) && (InterpolatedSyncState.MovementMode == DefaultModeNames::Walking))
 							{
-								FHitResult HitResult;
-								MoverComp->TryGetFloorCheckHitResult(HitResult);
-								FallingMode->OnLanded.Broadcast(InterpolatedSyncState.MovementMode, HitResult);
+								if (UFallingMode* FallingMode = MoverComp->FindMode_Mutable<UFallingMode>())
+								{
+									FHitResult HitResult;
+									MoverComp->TryGetFloorCheckHitResult(HitResult);
+									FallingMode->OnLanded.Broadcast(InterpolatedSyncState.MovementMode, HitResult);
+								}
 							}
 						}
+
+						// TODO: Consider moving to a util function
+						MoverComp->CachedLastSyncState = InterpolatedSyncState;
+						MoverComp->CachedLastSimTickTimeStep = MoverTimeStep;
+						MoverComp->bHasValidCachedState = true;
 					}
 
-					// TODO: Consider moving to a util function
-					MoverComp->CachedLastSyncState = InterpolatedSyncState;
-					MoverComp->CachedLastSimTickTimeStep = TimeStep;
-					MoverComp->bHasValidCachedState = true;
+					// Call on the last output for the frame
+					UpdateConstraintSettings();
+					MoverComp->OnPostSimulationTick.Broadcast(MoverTimeStep);
+					bCachedInputIsValid = false;
 				}
-
-				// Call on the last output for the frame
-				UpdateConstraintSettings();
-
-				MoverComp->OnPostSimulationTick.Broadcast(TimeStep);
-
-				bCachedInputIsValid = false;
 			}
-		}
 
-		CachedLastPhysicsSyncState = Output.SyncState;
-		CachedLastPhysicsSyncStateOutputTime = OutputTimeInSeconds;
-		bCachedLastPhysicsSyncStateIsValid = true;
+			CachedLastPhysicsSyncState = Output.SyncState;
+			CachedLastPhysicsSyncStateOutputTime = OutputTimeInSeconds;
+			bCachedLastPhysicsSyncStateIsValid = true;
+		}
+		else
+		{
+			MoverComp->CachedLastSyncState = Output.SyncState;
+			MoverComp->CachedLastSimTickTimeStep = MoverTimeStep;
+			MoverComp->bHasValidCachedState = true;
+
+			// Call on the last output for the frame
+			UpdateConstraintSettings();
+			MoverComp->OnPostSimulationTick.Broadcast(MoverTimeStep);
+			bCachedInputIsValid = false;
+		}
 	}
 }
 
@@ -765,6 +815,7 @@ void UMoverNetworkPhysicsLiaisonComponent::ProcessInputs_Internal(int32 PhysicsS
 
 				if (!bLocalPlayer || bIsSolverResim)
 				{
+					ensureMsgf(bUsingAsyncPhysics, TEXT("Physics mover requires async physics to work in multiplayer."));
 					GetCurrentInputData(Input.InputCmd);
 				}
 
@@ -818,8 +869,6 @@ void UMoverNetworkPhysicsLiaisonComponent::OnPreSimulate_Internal(const FPhysics
 		return;
 	}
 
-
-
 	const IPhysicsCharacterMovementModeInterface* PhysicsMode = Cast<const IPhysicsCharacterMovementModeInterface>(MoverComp->ModeFSM->FindMovementMode(Input.SyncState.MovementMode));
 	if (!PhysicsMode)
 	{
@@ -860,7 +909,7 @@ void UMoverNetworkPhysicsLiaisonComponent::OnPreSimulate_Internal(const FPhysics
 	FCharacterDefaultInputs* InputCmd = TickStartData.InputCmd.InputCollection.FindMutableDataByType<FCharacterDefaultInputs>();
 	check(InputCmd);
 
-	FMoverTimeStep TimeStep = GetCurrentMoverTimeStep_Internal();
+	const FMoverTimeStep TimeStep = bUsingAsyncPhysics ? GetCurrentAsyncMoverTimeStep_Internal() : GetCurrentMoverTimeStep(TickParams.DeltaTimeSeconds);
 
 	// Update movement state machine
 
diff --git a/Engine/Plugins/Experimental/Mover/Source/Mover/Public/Backends/MoverNetworkPhysicsLiaison.h b/Engine/Plugins/Experimental/Mover/Source/Mover/Public/Backends/MoverNetworkPhysicsLiaison.h
--- a/Engine/Plugins/Experimental/Mover/Source/Mover/Public/Backends/MoverNetworkPhysicsLiaison.h
+++ b/Engine/Plugins/Experimental/Mover/Source/Mover/Public/Backends/MoverNetworkPhysicsLiaison.h
@@ -152,11 +152,14 @@ protected:
 
 	int32 GetNetworkPhysicsTickOffset() const;
 
-	// Time step on the physics thread
-	FMoverTimeStep GetCurrentMoverTimeStep_Internal() const;
+	// Time step on the physics thread when using async physics
+	FMoverTimeStep GetCurrentAsyncMoverTimeStep_Internal() const;
 
-	// Time step on the game thread. Uses physics results time
-	FMoverTimeStep GetCurrentMoverTimeStep_External() const;
+	// Time step on the game thread when using async physics. Uses physics results time
+	FMoverTimeStep GetCurrentAsyncMoverTimeStep_External() const;
+
+	// Time step on the game thread when using non-async physics
+	FMoverTimeStep GetCurrentMoverTimeStep(float DeltaSeconds) const;
 
 	// These are written to by the network input and state data
 	FMoverInputCmdContext NetInputCmd;
@@ -177,4 +180,6 @@ protected:
 	bool bCachedLastPhysicsSyncStateIsValid = false;
 
 	bool bCachedInputIsValid = false;
+
+	bool bUsingAsyncPhysics = false;
 };
\ No newline at end of file
