Mover: based movement improvements. Fixes issues with movement sometimes being applied 1 sim frame late. Fixes subtle drifting that could cause an actor to fall off its base over time. Allows based movement to occur even if not standing on an object, like jumping on a boat.
Also added "scrape off" case to test map.

#virtualized

diff --git a/Engine/Plugins/Experimental/Mover/Source/Mover/Private/DefaultMovementSet/Modes/FallingMode.cpp b/Engine/Plugins/Experimental/Mover/Source/Mover/Private/DefaultMovementSet/Modes/FallingMode.cpp
--- a/Engine/Plugins/Experimental/Mover/Source/Mover/Private/DefaultMovementSet/Modes/FallingMode.cpp
+++ b/Engine/Plugins/Experimental/Mover/Source/Mover/Private/DefaultMovementSet/Modes/FallingMode.cpp
@@ -157,7 +157,7 @@ void UFallingMode::OnSimulationTick(const FSimulationTickParams& Params, FMoverT
 	UMoverBlackboard* SimBlackboard = GetBlackboard_Mutable();
 
 	SimBlackboard->Invalidate(CommonBlackboard::LastFloorResult);	// falling = no valid floor
-	SimBlackboard->Invalidate(CommonBlackboard::LastMovementBase);
+	SimBlackboard->Invalidate(CommonBlackboard::LastFoundDynamicMovementBase);
 
 	OutputSyncState.MoveDirectionIntent = (ProposedMove.bHasDirIntent ? ProposedMove.DirectionIntent : FVector::ZeroVector);
 
@@ -327,12 +327,12 @@ void UFallingMode::CaptureFinalState(USceneComponent* UpdatedComponent, const FM
 
 	if (MovementBaseInfo.HasRelativeInfo())
 	{
-		SimBlackboard->Set(CommonBlackboard::LastMovementBase, MovementBaseInfo);
+		SimBlackboard->Set(CommonBlackboard::LastFoundDynamicMovementBase, MovementBaseInfo);
 
 		OutputSyncState.SetTransforms_WorldSpace( FinalLocation,
 												  UpdatedComponent->GetComponentRotation(),
 												  EffectiveVelocity,
-												  MovementBaseInfo.MovementBase, MovementBaseInfo.BoneName);
+												  MovementBaseInfo.MovementBase.Get(), MovementBaseInfo.BoneName);
 	}
 	else
 	{
diff --git a/Engine/Plugins/Experimental/Mover/Source/Mover/Private/DefaultMovementSet/Modes/FlyingMode.cpp b/Engine/Plugins/Experimental/Mover/Source/Mover/Private/DefaultMovementSet/Modes/FlyingMode.cpp
--- a/Engine/Plugins/Experimental/Mover/Source/Mover/Private/DefaultMovementSet/Modes/FlyingMode.cpp
+++ b/Engine/Plugins/Experimental/Mover/Source/Mover/Private/DefaultMovementSet/Modes/FlyingMode.cpp
@@ -86,7 +86,7 @@ void UFlyingMode::OnSimulationTick(const FSimulationTickParams& Params, FMoverTi
 	UMoverBlackboard* SimBlackboard = GetBlackboard_Mutable();
 
 	SimBlackboard->Invalidate(CommonBlackboard::LastFloorResult);	// flying = no valid floor
-	SimBlackboard->Invalidate(CommonBlackboard::LastMovementBase);
+	SimBlackboard->Invalidate(CommonBlackboard::LastFoundDynamicMovementBase);
 
 	OutputSyncState.MoveDirectionIntent = (ProposedMove.bHasDirIntent ? ProposedMove.DirectionIntent : FVector::ZeroVector);
 
diff --git a/Engine/Plugins/Experimental/Mover/Source/Mover/Private/DefaultMovementSet/Modes/SwimmingMode.cpp b/Engine/Plugins/Experimental/Mover/Source/Mover/Private/DefaultMovementSet/Modes/SwimmingMode.cpp
--- a/Engine/Plugins/Experimental/Mover/Source/Mover/Private/DefaultMovementSet/Modes/SwimmingMode.cpp
+++ b/Engine/Plugins/Experimental/Mover/Source/Mover/Private/DefaultMovementSet/Modes/SwimmingMode.cpp
@@ -203,9 +203,12 @@ bool USwimmingMode::AttemptTeleport(USceneComponent* UpdatedComponent, const FVe
 												  nullptr ); // no movement base
 		
 		// TODO: instead of invalidating it, consider checking for a floor. Possibly a dynamic base?
-		GetBlackboard_Mutable()->Invalidate(CommonBlackboard::LastFloorResult);
-		GetBlackboard_Mutable()->Invalidate(CommonBlackboard::LastWaterResult);
-		GetBlackboard_Mutable()->Invalidate(CommonBlackboard::LastMovementBase);
+		if (UMoverBlackboard* SimBlackboard = GetBlackboard_Mutable())
+		{
+			SimBlackboard->Invalidate(CommonBlackboard::LastFloorResult);
+			SimBlackboard->Invalidate(CommonBlackboard::LastWaterResult);
+			SimBlackboard->Invalidate(CommonBlackboard::LastFoundDynamicMovementBase);
+		}
 
 		return true;
 	}
diff --git a/Engine/Plugins/Experimental/Mover/Source/Mover/Private/DefaultMovementSet/Modes/WalkingMode.cpp b/Engine/Plugins/Experimental/Mover/Source/Mover/Private/DefaultMovementSet/Modes/WalkingMode.cpp
--- a/Engine/Plugins/Experimental/Mover/Source/Mover/Private/DefaultMovementSet/Modes/WalkingMode.cpp
+++ b/Engine/Plugins/Experimental/Mover/Source/Mover/Private/DefaultMovementSet/Modes/WalkingMode.cpp
@@ -142,8 +142,6 @@ void UWalkingMode::OnSimulationTick(const FSimulationTickParams& Params, FMoverT
 	MoveRecord.SetDeltaSeconds(DeltaSeconds);
 
 	FFloorCheckResult CurrentFloor;
-	FRelativeBaseInfo OldRelativeBase;	// last known movement base (typically from a prior frame or prior substep)
-
 	UMoverBlackboard* SimBlackboard = GetBlackboard_Mutable();
 
 	// If we don't have cached floor information, we need to search for it again
@@ -154,20 +152,6 @@ void UWalkingMode::OnSimulationTick(const FSimulationTickParams& Params, FMoverT
 			UpdatedPrimitive->GetComponentLocation(), CurrentFloor);
 	}
 
-	if (!SimBlackboard->TryGet(CommonBlackboard::LastMovementBase, OldRelativeBase))
-	{
-		OldRelativeBase = UpdateFloorAndBaseInfo(CurrentFloor);
-	}
-	
-	
-	// If we're on a dynamic movement base, attempt to move along with whatever motion is has performed since we last ticked
-	if (OldRelativeBase.UsesSameBase(StartingSyncState->GetMovementBase(), StartingSyncState->GetMovementBaseBoneName()))
-	{
-		FHitResult BasedMoveHitResult;
-		bool bDidMoveAlongWithBase = UBasedMovementUtils::TryMoveToStayWithBase(UpdatedComponent, UpdatedPrimitive, OldRelativeBase, MoveRecord, bIgnoreBaseRotation);
-
-	}
-
 	OutputSyncState.MoveDirectionIntent = (ProposedMove.bHasDirIntent ? ProposedMove.DirectionIntent : FVector::ZeroVector);
 
 	const FRotator StartingOrient = StartingSyncState->GetOrientation_WorldSpace();
@@ -189,11 +173,15 @@ void UWalkingMode::OnSimulationTick(const FSimulationTickParams& Params, FMoverT
 	
 	FOptionalFloorCheckResult StepUpFloorResult;	// passed to sub-operations, so we can use their final floor results if they did a test
 
+	bool bDidAttemptMovement = false;
+
 	float PercentTimeAppliedSoFar = MoveHitResult.Time;
 	bool bWasFirstMoveBlocked = false;
+
 	if (!CurMoveDelta.IsNearlyZero() || bIsOrientationChanging)
 	{
 		// Attempt to move the full amount first
+		bDidAttemptMovement = true;
 		bool bMoved = UMovementUtils::TrySafeMoveUpdatedComponent(UpdatedComponent, UpdatedPrimitive, CurMoveDelta, TargetOrientQuat, true, MoveHitResult, ETeleportType::None, MoveRecord);
 		float LastMoveSeconds = DeltaSeconds;
 
@@ -272,7 +260,7 @@ void UWalkingMode::OnSimulationTick(const FSimulationTickParams& Params, FMoverT
 			OutputState.MovementEndState.NextModeName = CommonLegacySettings->AirMovementModeName;
 			OutputState.MovementEndState.RemainingMs = Params.TimeStep.StepMs - (Params.TimeStep.StepMs * PercentTimeAppliedSoFar);
 			MoveRecord.SetDeltaSeconds((Params.TimeStep.StepMs - OutputState.MovementEndState.RemainingMs) * 0.001f);
-			CaptureFinalState(UpdatedComponent, CurrentFloor, MoveRecord, OutputSyncState);
+			CaptureFinalState(UpdatedComponent, bDidAttemptMovement, CurrentFloor, MoveRecord, OutputSyncState);
 			return;
 		}
 	}
@@ -303,12 +291,12 @@ void UWalkingMode::OnSimulationTick(const FSimulationTickParams& Params, FMoverT
 			OutputState.MovementEndState.NextModeName = CommonLegacySettings->AirMovementModeName;
 			OutputState.MovementEndState.RemainingMs = Params.TimeStep.StepMs;
 			MoveRecord.SetDeltaSeconds((Params.TimeStep.StepMs - OutputState.MovementEndState.RemainingMs) * 0.001f);
-			CaptureFinalState(UpdatedComponent, CurrentFloor, MoveRecord, OutputSyncState);
+			CaptureFinalState(UpdatedComponent, bDidAttemptMovement, CurrentFloor, MoveRecord, OutputSyncState);
 			return;
 		}
 	}
 
-	CaptureFinalState(UpdatedComponent, CurrentFloor, MoveRecord, OutputSyncState);
+	CaptureFinalState(UpdatedComponent, bDidAttemptMovement, CurrentFloor, MoveRecord, OutputSyncState);
 
 }
 
@@ -350,8 +338,11 @@ bool UWalkingMode::AttemptTeleport(USceneComponent* UpdatedComponent, const FVec
 												  nullptr ); // no movement base
 		
 		// TODO: instead of invalidating it, consider checking for a floor. Possibly a dynamic base?
-		GetBlackboard_Mutable()->Invalidate(CommonBlackboard::LastFloorResult);
-		GetBlackboard_Mutable()->Invalidate(CommonBlackboard::LastMovementBase);
+		if (UMoverBlackboard* SimBlackboard = GetBlackboard_Mutable())
+		{
+			SimBlackboard->Invalidate(CommonBlackboard::LastFloorResult);
+			SimBlackboard->Invalidate(CommonBlackboard::LastFoundDynamicMovementBase);
+		}
 
 		return true;
 	}
@@ -359,22 +350,40 @@ bool UWalkingMode::AttemptTeleport(USceneComponent* UpdatedComponent, const FVec
 	return false;
 }
 
-// TODO: replace this function with simply looking at/collapsing the MovementRecord
-void UWalkingMode::CaptureFinalState(USceneComponent* UpdatedComponent, const FFloorCheckResult& FloorResult, const FMovementRecord& Record, FMoverDefaultSyncState& OutputSyncState) const
+void UWalkingMode::CaptureFinalState(USceneComponent* UpdatedComponent, bool bDidAttemptMovement, const FFloorCheckResult& FloorResult, const FMovementRecord& Record, FMoverDefaultSyncState& OutputSyncState) const
 {
-	FRelativeBaseInfo BaseInfo = UpdateFloorAndBaseInfo(FloorResult);
+	FRelativeBaseInfo PriorBaseInfo;
+
+	UMoverBlackboard* SimBlackboard = GetBlackboard_Mutable();
+
+	const bool bHasPriorBaseInfo = SimBlackboard->TryGet(CommonBlackboard::LastFoundDynamicMovementBase, PriorBaseInfo);
+
+	FRelativeBaseInfo CurrentBaseInfo = UpdateFloorAndBaseInfo(FloorResult);
+
+	// If we're on a dynamic base and we're not trying to move, keep using the same relative actor location. This prevents slow relative 
+	//  drifting that can occur from repeated floor sampling as the base moves through the world.
+	if (CurrentBaseInfo.HasRelativeInfo() 
+		&& bHasPriorBaseInfo && !bDidAttemptMovement 
+		&& PriorBaseInfo.UsesSameBase(CurrentBaseInfo))
+	{
+		CurrentBaseInfo.ContactLocalPosition = PriorBaseInfo.ContactLocalPosition;
+	}
 
 	// TODO: Update Main/large movement record with substeps from our local record
 	
-	if (BaseInfo.HasRelativeInfo())
+	if (CurrentBaseInfo.HasRelativeInfo())
 	{
+		SimBlackboard->Set(CommonBlackboard::LastFoundDynamicMovementBase, CurrentBaseInfo);
+
 		OutputSyncState.SetTransforms_WorldSpace( UpdatedComponent->GetComponentLocation(),
 												  UpdatedComponent->GetComponentRotation(),
 												  Record.GetRelevantVelocity(),
-												  BaseInfo.MovementBase, BaseInfo.BoneName);
+												  CurrentBaseInfo.MovementBase.Get(), CurrentBaseInfo.BoneName);
 	}
 	else
 	{
+		SimBlackboard->Invalidate(CommonBlackboard::LastFoundDynamicMovementBase);
+
 		OutputSyncState.SetTransforms_WorldSpace( UpdatedComponent->GetComponentLocation(),
 												  UpdatedComponent->GetComponentRotation(),
 												  Record.GetRelevantVelocity(),
@@ -396,13 +405,6 @@ FRelativeBaseInfo UWalkingMode::UpdateFloorAndBaseInfo(const FFloorCheckResult&
 	if (FloorResult.IsWalkableFloor() && UBasedMovementUtils::IsADynamicBase(FloorResult.HitResult.GetComponent()))
 	{
 		ReturnBaseInfo.SetFromFloorResult(FloorResult);
-
-		SimBlackboard->Set(CommonBlackboard::LastMovementBase, ReturnBaseInfo);
-	}
-	else
-	{
-
-		SimBlackboard->Invalidate(CommonBlackboard::LastMovementBase);
 	}
 
 	return ReturnBaseInfo;
diff --git a/Engine/Plugins/Experimental/Mover/Source/Mover/Private/MoveLibrary/BasedMovementUtils.cpp b/Engine/Plugins/Experimental/Mover/Source/Mover/Private/MoveLibrary/BasedMovementUtils.cpp
--- a/Engine/Plugins/Experimental/Mover/Source/Mover/Private/MoveLibrary/BasedMovementUtils.cpp
+++ b/Engine/Plugins/Experimental/Mover/Source/Mover/Private/MoveLibrary/BasedMovementUtils.cpp
@@ -4,8 +4,10 @@
 #include "MoveLibrary/MovementUtils.h"
 #include "MoveLibrary/FloorQueryUtils.h"
 #include "Components/PrimitiveComponent.h"
+#include "MoverComponent.h"
 #include "MoverLog.h"
 #include "Kismet/KismetMathLibrary.h"
+#include "DefaultMovementSet/Settings/CommonLegacyMovementSettings.h"
 
 #include UE_INLINE_GENERATED_CPP_BY_NAME(BasedMovementUtils)
 
@@ -15,6 +17,7 @@ void FRelativeBaseInfo::Clear()
 	BoneName = NAME_None;
 	Location = FVector::ZeroVector;
 	Rotation = FQuat::Identity;
+	ContactLocalPosition = FVector::ZeroVector;
 }
 
 bool FRelativeBaseInfo::HasRelativeInfo() const
@@ -24,7 +27,7 @@ bool FRelativeBaseInfo::HasRelativeInfo() const
 
 bool FRelativeBaseInfo::UsesSameBase(const FRelativeBaseInfo& Other) const
 {
-	return UsesSameBase(Other.MovementBase, Other.BoneName);
+	return UsesSameBase(Other.MovementBase.Get(), Other.BoneName);
 }
 
 bool FRelativeBaseInfo::UsesSameBase(const UPrimitiveComponent* OtherComp, FName OtherBoneName) const
@@ -42,11 +45,15 @@ void FRelativeBaseInfo::SetFromFloorResult(const FFloorCheckResult& FloorTestRes
 	{
 		MovementBase = FloorTestResult.HitResult.GetComponent();
 
-		if (MovementBase)
+		if (MovementBase.IsValid())
 		{
 			BoneName = FloorTestResult.HitResult.BoneName;
-			bDidSucceed  = UBasedMovementUtils::GetMovementBaseTransform(MovementBase, BoneName, /*out*/Location, /*out*/Rotation);
-			bDidSucceed &= UBasedMovementUtils::TransformWorldLocationToBased(MovementBase, BoneName, FloorTestResult.HitResult.ImpactPoint, /*out*/ContactLocalPosition);
+
+			if (UBasedMovementUtils::GetMovementBaseTransform(MovementBase.Get(), BoneName, OUT Location, OUT Rotation) &&
+				UBasedMovementUtils::TransformWorldLocationToBased(MovementBase.Get(), BoneName, FloorTestResult.HitResult.ImpactPoint, OUT ContactLocalPosition))
+			{
+				bDidSucceed = true;
+			}
 		}
 	}
 
@@ -62,10 +69,10 @@ void FRelativeBaseInfo::SetFromComponent(UPrimitiveComponent* InRelativeComp, FN
 
 	MovementBase = InRelativeComp;
 
-	if (MovementBase)
+	if (MovementBase.IsValid())
 	{
 		BoneName = InBoneName;
-		bDidSucceed = UBasedMovementUtils::GetMovementBaseTransform(MovementBase, BoneName, /*out*/Location, /*out*/Rotation);
+		bDidSucceed = UBasedMovementUtils::GetMovementBaseTransform(MovementBase.Get(), BoneName, /*out*/Location, /*out*/Rotation);
 	}
 
 	if (!bDidSucceed)
@@ -74,70 +81,39 @@ void FRelativeBaseInfo::SetFromComponent(UPrimitiveComponent* InRelativeComp, FN
 	}
 }
 
+
+FString FRelativeBaseInfo::ToString() const
+{
+	if (MovementBase.IsValid())
+	{
+		return FString::Printf(TEXT("Base: %s, Loc: %s, Rot: %s, LocalContact: %s"),
+			*GetNameSafe(MovementBase->GetOwner()),
+			*Location.ToCompactString(),
+			*Rotation.Rotator().ToCompactString(),
+			*ContactLocalPosition.ToCompactString());
+	}
+
+	return FString(TEXT("Base: NULL"));
+}
+
 bool UBasedMovementUtils::IsADynamicBase(const UPrimitiveComponent* MovementBase)
 {
 	return (MovementBase && MovementBase->Mobility == EComponentMobility::Movable);
 }
 
-bool UBasedMovementUtils::TryMoveToStayWithBase(USceneComponent* UpdatedComponent, UPrimitiveComponent* UpdatedPrimitive, const FRelativeBaseInfo& OldBaseInfo, FMovementRecord& MoveRecord, const bool& bIgnoreBaseRotation)
+bool UBasedMovementUtils::IsBaseSimulatingPhysics(const UPrimitiveComponent* MovementBase)
 {
-	FVector NewBaseLocation;
-	FQuat NewBaseQuat;
-
-
-	if (UBasedMovementUtils::GetMovementBaseTransform(OldBaseInfo.MovementBase, OldBaseInfo.BoneName, OUT NewBaseLocation, OUT NewBaseQuat))
+	bool bBaseIsSimulatingPhysics = false;
+	const USceneComponent* AttachParent = MovementBase;
+	while (!bBaseIsSimulatingPhysics && AttachParent)
 	{
-		const bool bDidBaseRotationChange = !OldBaseInfo.Rotation.Equals(NewBaseQuat, UE_SMALL_NUMBER);
-		const bool bDidBaseLocationChange = (OldBaseInfo.Location != NewBaseLocation);
-
-		// Find change in rotation
-		FQuat DeltaQuat = FQuat::Identity;
-		FVector DeltaLoc = FVector::ZeroVector;
-		FQuat TargetQuat = UpdatedComponent->GetComponentQuat();
-
-		if (bDidBaseRotationChange && !bIgnoreBaseRotation)
-		{
-			DeltaQuat = NewBaseQuat * OldBaseInfo.Rotation.Inverse();
-			TargetQuat = DeltaQuat * TargetQuat;
-
-			// TODO: make this respect the "up" direction, rather than assuming +Z = up
-			FVector TargetForwVector = TargetQuat.GetForwardVector();
-			TargetForwVector.Z = 0.f;
-			TargetForwVector.Normalize();
-
-			TargetQuat = UKismetMathLibrary::MakeRotFromX(TargetForwVector).Quaternion();
-		}
-
-		if (bDidBaseRotationChange || bDidBaseLocationChange)
-		{
-			// Calculate new transform matrix of base actor (ignoring scale).
-			const FQuatRotationTranslationMatrix OldLocalToWorld(OldBaseInfo.Rotation, OldBaseInfo.Location);
-			const FQuatRotationTranslationMatrix NewLocalToWorld(NewBaseQuat, NewBaseLocation);
-
-			// Find change in location
-			// NOTE that we need to use the floor hit location, not the actor's root position which may be floating above the base
-			const FVector NewWorldBaseContactPos = NewLocalToWorld.TransformPosition(OldBaseInfo.ContactLocalPosition);
-
-			const FVector OldWorldBaseContactPos = OldLocalToWorld.TransformPosition(OldBaseInfo.ContactLocalPosition);
-
-			DeltaLoc = NewWorldBaseContactPos - OldWorldBaseContactPos;
-
-			EMoveComponentFlags MoveComponentFlags = MOVECOMP_IgnoreBases;
-
-			const bool bSweep = true;
-			FHitResult MoveHitResult;
-			
-			bool bDidMove = UMovementUtils::TryMoveUpdatedComponent_Internal(UpdatedComponent, DeltaLoc, TargetQuat, bSweep, MoveComponentFlags, &MoveHitResult, ETeleportType::None);
-			
-			return bDidMove;
-		}
-
-		// TODO: rework this to ensure MoveRecord contains the appropriate thing
+		bBaseIsSimulatingPhysics = AttachParent->IsSimulatingPhysics();
+		AttachParent = AttachParent->GetAttachParent();
 	}
-
-	return false;
+	return bBaseIsSimulatingPhysics;
 }
 
+
 bool UBasedMovementUtils::GetMovementBaseTransform(const UPrimitiveComponent* MovementBase, const FName BoneName, FVector& OutLocation, FQuat& OutQuat)
 {
 	if (MovementBase)
@@ -285,3 +261,226 @@ void UBasedMovementUtils::TransformRotatorToLocal(FQuat BaseQuat, FRotator World
 	FQuat WorldQuat(WorldSpaceRotator);
 	OutLocalRotator = (BaseQuat.Inverse() * WorldQuat).Rotator();
 }
+
+void UBasedMovementUtils::AddTickDependency(FTickFunction& BasedObjectTick, UPrimitiveComponent* NewBase)
+{
+	if (NewBase && IsADynamicBase(NewBase))
+	{
+		if (NewBase->PrimaryComponentTick.bCanEverTick)
+		{
+			BasedObjectTick.AddPrerequisite(NewBase, NewBase->PrimaryComponentTick);
+		}
+
+		AActor* NewBaseOwner = NewBase->GetOwner();
+		if (NewBaseOwner)
+		{
+			if (NewBaseOwner->PrimaryActorTick.bCanEverTick)
+			{
+				BasedObjectTick.AddPrerequisite(NewBaseOwner, NewBaseOwner->PrimaryActorTick);
+			}
+
+			// @TODO: We need to find a more efficient way of finding all ticking components in an actor.
+			for (UActorComponent* Component : NewBaseOwner->GetComponents())
+			{
+				// Dont allow a based component (e.g. a particle system) to push us into a different tick group
+				if (Component && Component->PrimaryComponentTick.bCanEverTick && Component->PrimaryComponentTick.TickGroup <= BasedObjectTick.TickGroup)
+				{
+					BasedObjectTick.AddPrerequisite(Component, Component->PrimaryComponentTick);
+				}
+			}
+		}
+	}
+	else
+	{
+		UE_LOG(LogMover, Warning, TEXT("Attempted to AddTickDependency on an invalid or non-dynamic base: %s"), *GetNameSafe(NewBase));
+	}
+}
+
+void UBasedMovementUtils::RemoveTickDependency(FTickFunction& BasedObjectTick, UPrimitiveComponent* OldBase)
+{
+	if (OldBase)
+	{
+		BasedObjectTick.RemovePrerequisite(OldBase, OldBase->PrimaryComponentTick);
+		
+		if (AActor* OldBaseOwner = OldBase->GetOwner())
+		{
+			BasedObjectTick.RemovePrerequisite(OldBaseOwner, OldBaseOwner->PrimaryActorTick);
+
+			// @TODO: We need to find a more efficient way of finding all ticking components in an actor.
+			for (UActorComponent* Component : OldBaseOwner->GetComponents())
+			{
+				if (Component && Component->PrimaryComponentTick.bCanEverTick)
+				{
+					BasedObjectTick.RemovePrerequisite(Component, Component->PrimaryComponentTick);
+				}
+			}
+		}
+	}
+}
+
+
+void UBasedMovementUtils::UpdateSimpleBasedMovement(UMoverComponent* TargetMoverComp)
+{
+	if (!TargetMoverComp)
+	{
+		return;
+	}
+
+	UMoverBlackboard* SimBlackboard = TargetMoverComp->GetSimBlackboard_Mutable();
+	USceneComponent* UpdatedComponent = TargetMoverComp->UpdatedComponent;
+
+	bool bIgnoreBaseRotation = false;
+
+	if (const UCommonLegacyMovementSettings* CommonSettings = TargetMoverComp->FindSharedSettings<UCommonLegacyMovementSettings>())
+	{
+		bIgnoreBaseRotation = CommonSettings->bIgnoreBaseRotation;
+	}
+
+	bool bDidGetUpToDate = false;
+	if (TargetMoverComp->HasValidCachedState())
+	{
+		FRelativeBaseInfo LastFoundBaseInfo;	// Last-found is the most recent capture during movement, likely set this sim frame
+		FRelativeBaseInfo LastAppliedBaseInfo;	// Last-applied is the one that our based movement is up to date with, likely set in the last sim frame
+		FRelativeBaseInfo CurrentBaseInfo;		// Current info is the current snapshot of the current base, with up-to-date transform that may be different than last-found.
+
+		const bool bHasLastFoundInfo = SimBlackboard->TryGet(CommonBlackboard::LastFoundDynamicMovementBase, LastFoundBaseInfo);
+		const bool bHasLastAppliedInfo = SimBlackboard->TryGet(CommonBlackboard::LastAppliedDynamicMovementBase, LastAppliedBaseInfo);
+		if (bHasLastFoundInfo)
+		{
+			if (!bHasLastAppliedInfo || !LastFoundBaseInfo.UsesSameBase(LastAppliedBaseInfo))
+			{
+				LastAppliedBaseInfo = LastFoundBaseInfo;	// This is the first time we've checked this base, so start with the last-found capture
+			}
+
+			if (!ensureMsgf(LastFoundBaseInfo.HasRelativeInfo() && LastFoundBaseInfo.UsesSameBase(LastAppliedBaseInfo),
+					TEXT("Attempting to update based movement with a missing or mismatched base. This may indicate a logic problem with detecting bases.")))
+			{ 
+				SimBlackboard->Invalidate(CommonBlackboard::LastFoundDynamicMovementBase);
+				SimBlackboard->Invalidate(CommonBlackboard::LastAppliedDynamicMovementBase);
+				return;
+			}
+
+			CurrentBaseInfo.SetFromComponent(LastFoundBaseInfo.MovementBase.Get(), LastFoundBaseInfo.BoneName);
+			CurrentBaseInfo.ContactLocalPosition = LastFoundBaseInfo.ContactLocalPosition;
+
+			FVector CurrentBaseLocation;
+			FQuat CurrentBaseQuat;
+			
+
+			if (UBasedMovementUtils::GetMovementBaseTransform(CurrentBaseInfo.MovementBase.Get(), CurrentBaseInfo.BoneName, OUT CurrentBaseLocation, OUT CurrentBaseQuat))
+			{
+				const bool bDidBaseRotationChange = !LastAppliedBaseInfo.Rotation.Equals(CurrentBaseQuat, UE_SMALL_NUMBER);
+				const bool bDidBaseLocationChange = (LastAppliedBaseInfo.Location != CurrentBaseLocation);
+
+				FQuat DeltaQuat = FQuat::Identity;
+				FVector WorldDeltaLocation = FVector::ZeroVector;
+				FQuat WorldTargetQuat = UpdatedComponent->GetComponentQuat();
+
+				// Find change in rotation
+
+				if (bDidBaseRotationChange && !bIgnoreBaseRotation)
+				{
+					DeltaQuat = CurrentBaseQuat * LastAppliedBaseInfo.Rotation.Inverse();
+					WorldTargetQuat = DeltaQuat * WorldTargetQuat;
+
+					// TODO: make this respect the "up" direction, rather than assuming +Z = up
+					FVector TargetForwVector = WorldTargetQuat.GetForwardVector();
+					TargetForwVector.Z = 0.f;
+					TargetForwVector.Normalize();
+
+					WorldTargetQuat = UKismetMathLibrary::MakeRotFromX(TargetForwVector).Quaternion();
+				}
+
+				if (bDidBaseLocationChange || bDidBaseRotationChange)
+				{
+					// Calculate new transform matrix of base actor (ignoring scale).
+					const FQuatRotationTranslationMatrix OldLocalToWorld(LastAppliedBaseInfo.Rotation, LastAppliedBaseInfo.Location);
+					const FQuatRotationTranslationMatrix NewLocalToWorld(CurrentBaseQuat, CurrentBaseLocation);
+
+					// Find change in location
+					// NOTE that we are using the floor hit location, not the actor's root position which may be floating above the base
+					const FVector NewWorldBaseContactPos = NewLocalToWorld.TransformPosition(CurrentBaseInfo.ContactLocalPosition);
+					const FVector OldWorldBaseContactPos = OldLocalToWorld.TransformPosition(CurrentBaseInfo.ContactLocalPosition);
+					WorldDeltaLocation = NewWorldBaseContactPos - OldWorldBaseContactPos;
+
+					const FVector OldWorldLocation = UpdatedComponent->GetComponentLocation();
+					EMoveComponentFlags MoveComponentFlags = MOVECOMP_IgnoreBases;
+					const bool bSweep = true;
+					FHitResult MoveHitResult;
+
+					bool bDidMove = UMovementUtils::TryMoveUpdatedComponent_Internal(UpdatedComponent, WorldDeltaLocation, WorldTargetQuat, bSweep, MoveComponentFlags, &MoveHitResult, ETeleportType::None);
+					
+					const FVector NewWorldLocation = UpdatedComponent->GetComponentLocation();
+
+					if ((NewWorldLocation - (OldWorldLocation + WorldDeltaLocation)).IsNearlyZero() == false)
+					{
+						// Find the remaining delta that wasn't achieved
+						const FVector UnachievedWorldDelta = (OldWorldLocation + WorldDeltaLocation) - NewWorldLocation;
+
+						// Convert the remaining delta to current base space
+						FVector UnachievedLocalDelta;
+						UBasedMovementUtils::TransformLocationToLocal(CurrentBaseLocation, CurrentBaseQuat, UnachievedWorldDelta, OUT UnachievedLocalDelta);
+						
+						// Subtract the remaining delta to reflect the change in the contact position
+						CurrentBaseInfo.ContactLocalPosition -= UnachievedLocalDelta;
+					}
+
+					// Read, edit, write out the sync state based on the move results
+					FMoverSyncState PendingSyncState;
+					if (TargetMoverComp->BackendLiaisonComp->ReadPendingSyncState(OUT PendingSyncState))
+					{
+						if (FMoverDefaultSyncState* DefaultSyncState = PendingSyncState.SyncStateCollection.FindMutableDataByType<FMoverDefaultSyncState>())
+						{
+							FVector Velocity = DefaultSyncState->GetVelocity_WorldSpace();
+							DefaultSyncState->SetTransforms_WorldSpace(UpdatedComponent->GetComponentLocation(),
+								UpdatedComponent->GetComponentRotation(),
+								Velocity,
+								CurrentBaseInfo.MovementBase.Get(), CurrentBaseInfo.BoneName);
+
+							TargetMoverComp->BackendLiaisonComp->WritePendingSyncState(PendingSyncState);
+						}
+					}
+
+
+				}
+
+				SimBlackboard->Set(CommonBlackboard::LastAppliedDynamicMovementBase, CurrentBaseInfo);
+				bDidGetUpToDate = true;
+			}
+		}
+	}
+
+	if (!bDidGetUpToDate)
+	{
+		SimBlackboard->Invalidate(CommonBlackboard::LastAppliedDynamicMovementBase);
+	}
+}
+
+
+
+// FMoverDynamicBasedMovementTickFunction ////////////////////////////////////
+
+void FMoverDynamicBasedMovementTickFunction::ExecuteTick(float DeltaTime, enum ELevelTick TickType, ENamedThreads::Type CurrentThread, const FGraphEventRef& MyCompletionGraphEvent)
+{
+	FActorComponentTickFunction::ExecuteTickHelper(TargetMoverComp, /*bTickInEditor=*/ false, DeltaTime, TickType, [this](float DilatedTime)
+		{
+			UBasedMovementUtils::UpdateSimpleBasedMovement(TargetMoverComp);
+		});
+
+	if (bAutoDisableAfterTick)
+	{
+		SetTickFunctionEnable(false);
+	}
+}
+FString FMoverDynamicBasedMovementTickFunction::DiagnosticMessage()
+{
+	return TargetMoverComp->GetFullName() + TEXT("[FMoverDynamicBasedMovementTickFunction]");
+}
+FName FMoverDynamicBasedMovementTickFunction::DiagnosticContext(bool bDetailed)
+{
+	if (bDetailed)
+	{
+		return FName(*FString::Printf(TEXT("UMoverComponent/%s"), *GetFullNameSafe(TargetMoverComp)));
+	}
+	return FName(TEXT("FMoverDynamicBasedMovementTickFunction"));
+}
diff --git a/Engine/Plugins/Experimental/Mover/Source/Mover/Private/MoverComponent.cpp b/Engine/Plugins/Experimental/Mover/Source/Mover/Private/MoverComponent.cpp
--- a/Engine/Plugins/Experimental/Mover/Source/Mover/Private/MoverComponent.cpp
+++ b/Engine/Plugins/Experimental/Mover/Source/Mover/Private/MoverComponent.cpp
@@ -46,6 +46,11 @@ UMoverComponent::UMoverComponent()
 	PrimaryComponentTick.TickGroup = TG_PrePhysics;
 	PrimaryComponentTick.bCanEverTick = true;
 
+	BasedMovementTickFunction.bCanEverTick = true;
+	BasedMovementTickFunction.bStartWithTickEnabled = false;
+	BasedMovementTickFunction.SetTickFunctionEnable(false);
+	BasedMovementTickFunction.TickGroup = TG_PostPhysics;
+
 	bWantsInitializeComponent = true;
 	bAutoActivate = true;
 
@@ -163,6 +168,23 @@ void UMoverComponent::RegisterComponentTickFunctions(bool bRegister)
 	{
 		Owner->PrimaryActorTick.AddPrerequisite(this, PrimaryComponentTick);
 	}
+
+
+	if (bRegister)
+	{
+		if (SetupActorComponentTickFunction(&BasedMovementTickFunction))
+		{
+			BasedMovementTickFunction.TargetMoverComp = this;
+			BasedMovementTickFunction.AddPrerequisite(this, this->PrimaryComponentTick);
+		}
+	}
+	else
+	{
+		if (BasedMovementTickFunction.IsTickFunctionRegistered())
+		{
+			BasedMovementTickFunction.UnRegisterTickFunction();
+		}
+	}
 }
 
 
@@ -374,6 +396,11 @@ void UMoverComponent::SimulationTick(const FMoverTimeStep& InTimeStep, const FMo
 	CachedLastUsedInputCmd = SimInput.InputCmd;
 	bHasValidCachedUsedInput = true;
 
+	if (bSupportsKinematicBasedMovement)
+	{ 
+		UpdateBasedMovementScheduling(SimOutput);
+	}
+
 	OnPostSimulationTick.Broadcast(MoverTimeStep);
 
 	CachedLastSimTickTimeStep = MoverTimeStep;
@@ -471,6 +498,55 @@ void UMoverComponent::OnHandleImpact(const FMoverOnImpactParams& ImpactParams)
 	// TODO: Handle physics impacts here - ie when player runs into box, impart force onto box
 }
 
+void UMoverComponent::UpdateBasedMovementScheduling(const FMoverTickEndData& SimOutput)
+{
+	// If we have a dynamic movement base, enable later based movement tick
+	UPrimitiveComponent* SyncStateDynamicBase = nullptr;
+	if (const FMoverDefaultSyncState* OutputSyncState = SimOutput.SyncState.SyncStateCollection.FindDataByType<FMoverDefaultSyncState>())
+	{
+		if (UBasedMovementUtils::IsADynamicBase(OutputSyncState->GetMovementBase()))
+		{
+			SyncStateDynamicBase = OutputSyncState->GetMovementBase();
+		}
+	}
+
+	// Remove any stale dependency
+	if (MovementBaseDependency && (MovementBaseDependency != SyncStateDynamicBase))
+	{
+		UBasedMovementUtils::RemoveTickDependency(BasedMovementTickFunction, MovementBaseDependency);
+		MovementBaseDependency = nullptr;
+	}
+
+	// Set up current dependencies
+	if (SyncStateDynamicBase)
+	{
+		BasedMovementTickFunction.SetTickFunctionEnable(true);
+
+		if (UBasedMovementUtils::IsBaseSimulatingPhysics(SyncStateDynamicBase))
+		{
+			BasedMovementTickFunction.TickGroup = TG_PostPhysics;
+		}
+		else
+		{
+			BasedMovementTickFunction.TickGroup = TG_PrePhysics;
+		}
+
+		if (MovementBaseDependency == nullptr)
+		{
+			UBasedMovementUtils::AddTickDependency(BasedMovementTickFunction, SyncStateDynamicBase);
+			MovementBaseDependency = SyncStateDynamicBase;
+		}
+	}
+	else
+	{
+		BasedMovementTickFunction.SetTickFunctionEnable(false);
+		MovementBaseDependency = nullptr;
+
+		SimBlackboard->Invalidate(CommonBlackboard::LastFoundDynamicMovementBase);
+		SimBlackboard->Invalidate(CommonBlackboard::LastAppliedDynamicMovementBase);
+	}
+}
+
 void UMoverComponent::SetUpdatedComponent(USceneComponent* NewUpdatedComponent)
 {
 	// Remove delegates from old component
@@ -1232,6 +1308,11 @@ const UMoverBlackboard* UMoverComponent::GetSimBlackboard() const
 	return SimBlackboard;
 }
 
+UMoverBlackboard* UMoverComponent::GetSimBlackboard_Mutable() const
+{
+	return SimBlackboard;
+}
+
 bool UMoverComponent::HasValidCachedInputCmd() const
 {
 	return bHasValidCachedUsedInput;
diff --git a/Engine/Plugins/Experimental/Mover/Source/Mover/Private/PhysicsMover/Modes/PhysicsDrivenSwimmingMode.cpp b/Engine/Plugins/Experimental/Mover/Source/Mover/Private/PhysicsMover/Modes/PhysicsDrivenSwimmingMode.cpp
--- a/Engine/Plugins/Experimental/Mover/Source/Mover/Private/PhysicsMover/Modes/PhysicsDrivenSwimmingMode.cpp
+++ b/Engine/Plugins/Experimental/Mover/Source/Mover/Private/PhysicsMover/Modes/PhysicsDrivenSwimmingMode.cpp
@@ -170,8 +170,11 @@ bool UPhysicsDrivenSwimmingMode::AttemptTeleport(USceneComponent* UpdatedCompone
 		nullptr); // no movement base
 
 	// TODO: instead of invalidating it, consider checking for a floor. Possibly a dynamic base?
-	GetBlackboard_Mutable()->Invalidate(CommonBlackboard::LastFloorResult);
-	GetBlackboard_Mutable()->Invalidate(CommonBlackboard::LastMovementBase);
+	if (UMoverBlackboard* SimBlackboard = GetBlackboard_Mutable())
+	{
+		SimBlackboard->Invalidate(CommonBlackboard::LastFloorResult);
+		SimBlackboard->Invalidate(CommonBlackboard::LastFoundDynamicMovementBase);
+	}
 
 	return true;
 }
diff --git a/Engine/Plugins/Experimental/Mover/Source/Mover/Private/PhysicsMover/Modes/PhysicsDrivenWalkingMode.cpp b/Engine/Plugins/Experimental/Mover/Source/Mover/Private/PhysicsMover/Modes/PhysicsDrivenWalkingMode.cpp
--- a/Engine/Plugins/Experimental/Mover/Source/Mover/Private/PhysicsMover/Modes/PhysicsDrivenWalkingMode.cpp
+++ b/Engine/Plugins/Experimental/Mover/Source/Mover/Private/PhysicsMover/Modes/PhysicsDrivenWalkingMode.cpp
@@ -228,8 +228,11 @@ bool UPhysicsDrivenWalkingMode::AttemptTeleport(USceneComponent* UpdatedComponen
 		nullptr); // no movement base
 
 	// TODO: instead of invalidating it, consider checking for a floor. Possibly a dynamic base?
-	GetBlackboard_Mutable()->Invalidate(CommonBlackboard::LastFloorResult);
-	GetBlackboard_Mutable()->Invalidate(CommonBlackboard::LastMovementBase);
+	if (UMoverBlackboard* SimBlackboard = GetBlackboard_Mutable())
+	{
+		SimBlackboard->Invalidate(CommonBlackboard::LastFloorResult);
+		SimBlackboard->Invalidate(CommonBlackboard::LastFoundDynamicMovementBase);
+	}
 
 	return true;
 }
diff --git a/Engine/Plugins/Experimental/Mover/Source/Mover/Public/DefaultMovementSet/Modes/WalkingMode.h b/Engine/Plugins/Experimental/Mover/Source/Mover/Public/DefaultMovementSet/Modes/WalkingMode.h
--- a/Engine/Plugins/Experimental/Mover/Source/Mover/Public/DefaultMovementSet/Modes/WalkingMode.h
+++ b/Engine/Plugins/Experimental/Mover/Source/Mover/Public/DefaultMovementSet/Modes/WalkingMode.h
@@ -21,13 +21,6 @@ class MOVER_API UWalkingMode : public UBaseMovementMode
 	GENERATED_UCLASS_BODY()
 
 public:
-	/**
- 	 * Whether the actor ignores changes in rotation of the base it is standing on.
- 	 * If true, the actor maintains current world rotation.
- 	 * If false, the actor rotates with the moving base.
- 	 */
-	UPROPERTY(BlueprintReadWrite, EditAnywhere, Category=Mover)
-	bool bIgnoreBaseRotation = true;
 	
 	UFUNCTION(BlueprintCallable, Category=Mover)
 	virtual void OnGenerateMove(const FMoverTickStartData& StartState, const FMoverTimeStep& TimeStep, FProposedMove& OutProposedMove) const override;
@@ -48,7 +41,7 @@ protected:
 	virtual bool AttemptTeleport(USceneComponent* UpdatedComponent, const FVector& TeleportPos, const FRotator& TeleportRot, const FVector& PriorVelocity, FMoverTickEndData& Output);
 
 
-	void CaptureFinalState(USceneComponent* UpdatedComponent, const FFloorCheckResult& FloorResult, const FMovementRecord& Record, FMoverDefaultSyncState& OutputSyncState) const;
+	void CaptureFinalState(USceneComponent* UpdatedComponent, bool bDidAttemptMovement, const FFloorCheckResult& FloorResult, const FMovementRecord& Record, FMoverDefaultSyncState& OutputSyncState) const;
 
 	FRelativeBaseInfo UpdateFloorAndBaseInfo(const FFloorCheckResult& FloorResult) const;
 
diff --git a/Engine/Plugins/Experimental/Mover/Source/Mover/Public/DefaultMovementSet/Settings/CommonLegacyMovementSettings.h b/Engine/Plugins/Experimental/Mover/Source/Mover/Public/DefaultMovementSet/Settings/CommonLegacyMovementSettings.h
--- a/Engine/Plugins/Experimental/Mover/Source/Mover/Public/DefaultMovementSet/Settings/CommonLegacyMovementSettings.h
+++ b/Engine/Plugins/Experimental/Mover/Source/Mover/Public/DefaultMovementSet/Settings/CommonLegacyMovementSettings.h
@@ -92,6 +92,13 @@ public:
 	UPROPERTY(BlueprintReadWrite, EditAnywhere, Category="General", meta = (ClampMin = "0", UIMin = "0", ForceUnits = "Multiplier"))
 	float TurningBoost = 8.f;
 
+ 	/** Whether the actor ignores changes in rotation of the base it is standing on when using based movement.
+  	 * If true, the actor maintains its current world rotation.
+  	 * If false, the actor rotates with the moving base.
+  	 */
+ 	UPROPERTY(BlueprintReadWrite, EditAnywhere, Category = "General")
+	bool bIgnoreBaseRotation = false;
+
 	/** Instantaneous speed induced in an actor upon jumping */
 	UPROPERTY(BlueprintReadWrite, EditAnywhere, Category="Jumping", meta = (ClampMin = "0", UIMin = "0", ForceUnits = "cm/s"))
 	float JumpUpwardsSpeed = 500.0f;
diff --git a/Engine/Plugins/Experimental/Mover/Source/Mover/Public/MoveLibrary/BasedMovementUtils.h b/Engine/Plugins/Experimental/Mover/Source/Mover/Public/MoveLibrary/BasedMovementUtils.h
--- a/Engine/Plugins/Experimental/Mover/Source/Mover/Public/MoveLibrary/BasedMovementUtils.h
+++ b/Engine/Plugins/Experimental/Mover/Source/Mover/Public/MoveLibrary/BasedMovementUtils.h
@@ -3,8 +3,10 @@
 #pragma once
 
 #include "Kismet/BlueprintFunctionLibrary.h"
+#include "Engine/EngineBaseTypes.h"
 #include "BasedMovementUtils.generated.h"
 
+class UMoverComponent;
 struct FMovementRecord;
 struct FFloorCheckResult;
 
@@ -16,7 +18,7 @@ struct MOVER_API FRelativeBaseInfo
 
 	/** Component we are moving relative to */
 	UPROPERTY(VisibleInstanceOnly, BlueprintReadOnly, Category="Mover")
-	TObjectPtr<UPrimitiveComponent> MovementBase = nullptr;
+	TWeakObjectPtr<UPrimitiveComponent> MovementBase = nullptr;
 
 	/** Bone name on component, for skeletal meshes. NAME_None if not a skeletal mesh or if bone is invalid. */
 	UPROPERTY(VisibleInstanceOnly, BlueprintReadOnly, Category = "Mover")
@@ -30,9 +32,9 @@ struct MOVER_API FRelativeBaseInfo
 	UPROPERTY(VisibleInstanceOnly, BlueprintReadOnly, Category = "Mover")
 	FQuat Rotation = FQuat::Identity;
 
-	/** Last captured location of the tethering point where the Mover actor is "attached". */
+	/** Last captured location of the tethering point where the Mover actor is "attached", relative to the base. */
 	UPROPERTY(VisibleInstanceOnly, BlueprintReadOnly, Category = "Mover")
-	FVector ContactLocalPosition = FVector::ZeroVector;;
+	FVector ContactLocalPosition = FVector::ZeroVector;
 
 public:
 
@@ -44,6 +46,8 @@ public:
 
 	void SetFromFloorResult(const FFloorCheckResult& FloorTestResult);
 	void SetFromComponent(UPrimitiveComponent* InRelativeComp, FName InBoneName=NAME_None);
+
+	FString ToString() const;
 };
 
 /**
@@ -59,8 +63,9 @@ public:
 	UFUNCTION(BlueprintCallable, Category = "Mover/MovementBases")
 	static bool IsADynamicBase(const UPrimitiveComponent* MovementBase);
 
-	/** Attempts to move a component relative to any change of the movement base. OldBaseInfo will typically be captured from the prior simulation frame. */
-	static bool TryMoveToStayWithBase(USceneComponent* UpdatedComponent, UPrimitiveComponent* UpdatedPrimitive, const FRelativeBaseInfo& OldBaseInfo, FMovementRecord& MoveRecord, const bool& bIgnoreBaseRotation);
+	/** Determine whether MovementBase's movement is performed via physics. */
+	UFUNCTION(BlueprintCallable, Category = "Mover/MovementBases")
+	static bool IsBaseSimulatingPhysics(const UPrimitiveComponent* MovementBase);
 	
 	/** Get the transform (local-to-world) for the given MovementBase, optionally at the location of a bone. Returns false if MovementBase is nullptr, or if BoneName is not a valid bone. */
 	UFUNCTION(BlueprintCallable, Category="Mover/MovementBases")
@@ -113,4 +118,52 @@ public:
 	/** Convert a world space rotator to a local rotator for a given MovementBase, optionally relative to the orientation of a bone. Returns false if MovementBase is nullptr, or if BoneName is not a valid bone. Scaling is ignored. */
 	UFUNCTION(BlueprintCallable, Category = "Mover/MovementBases")
 	static void TransformRotatorToLocal(FQuat BaseQuat, FRotator WorldSpaceRotator, FRotator& OutLocalRotator);
+
+
+	/** Makes it so BasedObjectTick ticks after NewBase's actor ticking */
+	static void AddTickDependency(FTickFunction& BasedObjectTick, UPrimitiveComponent* NewBase);
+	
+	/** Removes ticking dependency of BasedObjectTick on OldBase */
+	static void RemoveTickDependency(FTickFunction& BasedObjectTick, UPrimitiveComponent* OldBase);
+
+	/** Attempts to move the actor to keep up with its base's movement using a simple sweep. This function is not intended to be called during a Mover actor's simulation tick. */
+	static void UpdateSimpleBasedMovement(UMoverComponent* TargetMoverComp);
+};
+
+
+/**
+ * Tick function used to perform based movement at dynamic times throughout the world update time, typically out-of-band with the movement simulation
+ **/
+USTRUCT()
+struct FMoverDynamicBasedMovementTickFunction : public FTickFunction
+{
+	GENERATED_USTRUCT_BODY()
+
+	/** MoverComponent that is the target of this tick **/
+	UMoverComponent* TargetMoverComp;
+
+	/** If true, this tick function will self-disable after running */
+	bool bAutoDisableAfterTick = false;
+
+	/**
+	 * Abstract function actually execute the tick.
+	 * @param DeltaTime - frame time to advance, in seconds
+	 * @param TickType - kind of tick for this frame
+	 * @param CurrentThread - thread we are executing on, useful to pass along as new tasks are created
+	 * @param MyCompletionGraphEvent - completion event for this task. Useful for holding the completion of this task until certain child tasks are complete.
+	 **/
+	virtual void ExecuteTick(float DeltaTime, enum ELevelTick TickType, ENamedThreads::Type CurrentThread, const FGraphEventRef& MyCompletionGraphEvent) override;
+	/** Abstract function to describe this tick. Used to print messages about illegal cycles in the dependency graph **/
+	virtual FString DiagnosticMessage() override;
+	/** Function used to describe this tick for active tick reporting. **/
+	virtual FName DiagnosticContext(bool bDetailed) override;
+};
+
+template<>
+struct TStructOpsTypeTraits<FMoverDynamicBasedMovementTickFunction> : public TStructOpsTypeTraitsBase2<FMoverDynamicBasedMovementTickFunction>
+{
+	enum
+	{
+		WithCopy = false
+	};
 };
diff --git a/Engine/Plugins/Experimental/Mover/Source/Mover/Public/MoverComponent.h b/Engine/Plugins/Experimental/Mover/Source/Mover/Public/MoverComponent.h
--- a/Engine/Plugins/Experimental/Mover/Source/Mover/Public/MoverComponent.h
+++ b/Engine/Plugins/Experimental/Mover/Source/Mover/Public/MoverComponent.h
@@ -9,6 +9,7 @@
 #include "MovementMode.h"
 #include "MoverTypes.h"
 #include "LayeredMove.h"
+#include "MoveLibrary/BasedMovementUtils.h"
 #if UE_ENABLE_INCLUDE_ORDER_DEPRECATED_IN_5_4
 #include "Engine/HitResult.h"
 #endif // UE_ENABLE_INCLUDE_ORDER_DEPRECATED_IN_5_4
@@ -256,6 +257,8 @@ public:	// Queries
 	UFUNCTION(BlueprintPure, Category=Mover)
 	const UMoverBlackboard* GetSimBlackboard() const;
 
+	UMoverBlackboard* GetSimBlackboard_Mutable() const;
+
 	/** Find settings object by type. Returns null if there is none of that type */
 	const IMovementSettingsInterface* FindSharedSettings(const UClass* ByType) const { return FindSharedSettings_Mutable(ByType); }
 	template<class T>
@@ -343,6 +346,11 @@ protected:
 	virtual void PhysicsVolumeChanged(class APhysicsVolume* NewVolume);
 
 	virtual void OnHandleImpact(const FMoverOnImpactParams& ImpactParams);
+
+	/** internal function to perform post-sim scheduling to optionally support simple based movement */
+	void UpdateBasedMovementScheduling(const FMoverTickEndData& SimOutput);
+
+	TObjectPtr<UPrimitiveComponent> MovementBaseDependency;	// used internally for based movement scheduling management
 	
 	/** internal function to ensure SharedSettings array matches what's needed by the list of Movement Modes */
 	void RefreshSharedSettings();
@@ -374,6 +382,9 @@ protected:
 
 	TWeakInterfacePtr<IMoverBackendLiaisonInterface> BackendLiaisonComp;
 
+	/** Tick function that may be called anytime after this actor's movement step, useful as a way to support based movement on objects that are not */
+	FMoverDynamicBasedMovementTickFunction BasedMovementTickFunction;
+
 private:
 	/** Collection of settings objects that are shared between movement modes. This list is automatically managed based on the @MovementModes contents. */
 	UPROPERTY(EditDefaultsOnly, EditFixedSize, Instanced, Category = Mover, meta = (NoResetToDefault, MustImplement = "/Script/Mover.MovementSettingsInterface"))
@@ -387,6 +398,10 @@ private:
 	UPROPERTY(EditDefaultsOnly, Category="Mover|Gravity", meta=(ForceUnits = "cm/s^2"))
 	FVector GravityAccelOverride;
 
+	/** If enabled, this actor will be moved to follow a base actor that it's standing on. Typically disabled for physics-based movement, which handles based movement internally. */
+	UPROPERTY(EditDefaultsOnly, Category = "Mover")
+	bool bSupportsKinematicBasedMovement = true;
+
 	/** Transient flag indicating whether we are executing OnRegister(). */
 	bool bInOnRegister = false;
 
@@ -407,4 +422,5 @@ private:
 	friend class UBaseMovementMode;
 	friend class UMoverNetworkPhysicsLiaisonComponent;
 	friend class UMoverDebugComponent;
+	friend class UBasedMovementUtils;
 };
diff --git a/Engine/Plugins/Experimental/Mover/Source/Mover/Public/MoverSimulationTypes.h b/Engine/Plugins/Experimental/Mover/Source/Mover/Public/MoverSimulationTypes.h
--- a/Engine/Plugins/Experimental/Mover/Source/Mover/Public/MoverSimulationTypes.h
+++ b/Engine/Plugins/Experimental/Mover/Source/Mover/Public/MoverSimulationTypes.h
@@ -26,7 +26,8 @@ namespace CommonBlackboard
 {
 	const FName LastFloorResult = TEXT("LastFloor");
 	const FName LastWaterResult = TEXT("LastWater");
-	const FName LastMovementBase = TEXT("LastMovementBase");
+	const FName LastFoundDynamicMovementBase = TEXT("LastFoundDynamicMovementBase");
+	const FName LastAppliedDynamicMovementBase = TEXT("LastAppliedDynamicMovementBase");
 	const FName TimeSinceSupported = TEXT("TimeSinceSupported");
 }
 
