Mover: help prevent jittery motion due to bad setup. Now ensuring if a Mover actor is set to replicate movement, which will conflict with the MoverComponent.

diff --git a/Engine/Plugins/Experimental/Mover/Source/Mover/Private/MoverComponent.cpp b/Engine/Plugins/Experimental/Mover/Source/Mover/Private/MoverComponent.cpp
--- a/Engine/Plugins/Experimental/Mover/Source/Mover/Private/MoverComponent.cpp
+++ b/Engine/Plugins/Experimental/Mover/Source/Mover/Private/MoverComponent.cpp
@@ -219,6 +219,13 @@ void UMoverComponent::BeginPlay()
 	{
 		MovementMixer = NewObject<UMovementMixer>(this, TEXT("Default Movement Mixer"));
 	}
+
+	if (const AActor* OwnerActor = GetOwner())
+	{
+		ensureMsgf(!OwnerActor->IsReplicatingMovement(), 
+		           TEXT("MoverComponent owning actor %s has the ReplicateMovement property enabled. This will conflict with MoverComponent and cause poor quality movement. Please disable it."),
+		           *GetNameSafe(GetOwner()));
+	}
 }
 
 void UMoverComponent::TickComponent(float DeltaTime, enum ELevelTick TickType, FActorComponentTickFunction* ThisTickFunction)
