Mover plugin: narrowing ensure condition to Network Prediction-simulated actors only.

#rnx

diff --git a/Engine/Plugins/Experimental/Mover/Source/Mover/Private/Backends/MoverNetworkPredictionLiaison.cpp b/Engine/Plugins/Experimental/Mover/Source/Mover/Private/Backends/MoverNetworkPredictionLiaison.cpp
--- a/Engine/Plugins/Experimental/Mover/Source/Mover/Private/Backends/MoverNetworkPredictionLiaison.cpp
+++ b/Engine/Plugins/Experimental/Mover/Source/Mover/Private/Backends/MoverNetworkPredictionLiaison.cpp
@@ -118,6 +118,14 @@ bool UMoverNetworkPredictionLiaisonComponent::WritePendingSyncState(const FMover
 void UMoverNetworkPredictionLiaisonComponent::BeginPlay()
 {
 	Super::BeginPlay();
+
+	if (const AActor* OwnerActor = GetOwner())
+	{
+		ensureMsgf(!OwnerActor->IsReplicatingMovement(),
+			TEXT("MoverComponent owning actor %s has the ReplicateMovement property enabled. This will conflict with Network Prediction and cause poor quality movement. Please disable it."),
+			*GetNameSafe(GetOwner()));
+	}
+
 	if (StartingOutSync && StartingOutAux)
 	{
 		if (FMoverDefaultSyncState* StartingSyncState = StartingOutSync->SyncStateCollection.FindMutableDataByType<FMoverDefaultSyncState>())
diff --git a/Engine/Plugins/Experimental/Mover/Source/Mover/Private/MoverComponent.cpp b/Engine/Plugins/Experimental/Mover/Source/Mover/Private/MoverComponent.cpp
--- a/Engine/Plugins/Experimental/Mover/Source/Mover/Private/MoverComponent.cpp
+++ b/Engine/Plugins/Experimental/Mover/Source/Mover/Private/MoverComponent.cpp
@@ -219,13 +219,6 @@ void UMoverComponent::BeginPlay()
 	{
 		MovementMixer = NewObject<UMovementMixer>(this, TEXT("Default Movement Mixer"));
 	}
-
-	if (const AActor* OwnerActor = GetOwner())
-	{
-		ensureMsgf(!OwnerActor->IsReplicatingMovement(), 
-		           TEXT("MoverComponent owning actor %s has the ReplicateMovement property enabled. This will conflict with MoverComponent and cause poor quality movement. Please disable it."),
-		           *GetNameSafe(GetOwner()));
-	}
 }
 
 void UMoverComponent::TickComponent(float DeltaTime, enum ELevelTick TickType, FActorComponentTickFunction* ThisTickFunction)
