Physics Mover - Move contact filtering to be movement mode specific

diff --git a/Engine/Plugins/Experimental/Mover/Source/Mover/Private/Backends/MoverNetworkPhysicsLiaison.cpp b/Engine/Plugins/Experimental/Mover/Source/Mover/Private/Backends/MoverNetworkPhysicsLiaison.cpp
--- a/Engine/Plugins/Experimental/Mover/Source/Mover/Private/Backends/MoverNetworkPhysicsLiaison.cpp
+++ b/Engine/Plugins/Experimental/Mover/Source/Mover/Private/Backends/MoverNetworkPhysicsLiaison.cpp
@@ -869,10 +869,18 @@ void UMoverNetworkPhysicsLiaisonComponent::OnPreSimulate_Internal(const FPhysics
 		return;
 	}
 
-	const IPhysicsCharacterMovementModeInterface* PhysicsMode = Cast<const IPhysicsCharacterMovementModeInterface>(MoverComp->ModeFSM->FindMovementMode(Input.SyncState.MovementMode));
-	if (!PhysicsMode)
+	// Check that the input movement mode is valid
+	if (MoverComp->MovementModes.Contains(Input.SyncState.MovementMode))
+	{
+		const IPhysicsCharacterMovementModeInterface* PhysicsMode = Cast<const IPhysicsCharacterMovementModeInterface>(MoverComp->MovementModes[Input.SyncState.MovementMode]);
+		if (!PhysicsMode)
+		{
+			UE_LOG(LogMover, Verbose, TEXT("Attempting to run non-physics movement mode %s in physics mover update."), *Input.SyncState.MovementMode.ToString());
+			return;
+		}
+	}
+	else
 	{
-		UE_LOG(LogMover, Verbose, TEXT("Attempting to run non-physics movement mode %s in physics mover update."), *Input.SyncState.MovementMode.ToString());
 		return;
 	}
 
@@ -1086,75 +1094,18 @@ void UMoverNetworkPhysicsLiaisonComponent::OnContactModification_Internal(const
 		return;
 	}
 
-	Chaos::FCharacterGroundConstraintHandle* ConstraintHandle = Constraint->GetProxy<Chaos::FCharacterGroundConstraintProxy>()->GetPhysicsThreadAPI();
-	if (!ConstraintHandle || !ConstraintHandle->IsEnabled() || !ConstraintHandle->GetCharacterParticle())
-	{
-		return;
-	}
-
-	Chaos::FPBDRigidParticleHandle* CharacterParticle = ConstraintHandle->GetCharacterParticle()->CastToRigidParticle();
-	if (!CharacterParticle || CharacterParticle->Disabled())
-	{
-		return;
-	}
-
-	const Chaos::FGeometryParticleHandle* GroundParticle = ConstraintHandle->GetGroundParticle();
-	if (!GroundParticle)
+	if (MoverComp->MovementModes.Contains(Input.SyncState.MovementMode))
 	{
-		return;
-	}
-
-	float PawnHalfHeight;
-	float PawnRadius;
-	MoverComp->UpdatedComponent->CalcBoundingCylinder(PawnRadius, PawnHalfHeight);
-
-	const float CharacterHeight = CharacterParticle->GetX().Z;
-	const float EndCapHeight = CharacterHeight - PawnHalfHeight + PawnRadius;
-
-	const float CosThetaMax = 0.707f;
-
-	float MinContactHeightStepUps = CharacterHeight - 1.0e10f;
-	if (Input.SyncState.MovementMode == DefaultModeNames::Walking)
-	{
-		if (const UPhysicsDrivenWalkingMode* WalkingMode = Cast<UPhysicsDrivenWalkingMode>(MoverComp->FindMode_Mutable<UPhysicsDrivenWalkingMode>()))
+		if (const IPhysicsCharacterMovementModeInterface* PhysicsMode = Cast<const IPhysicsCharacterMovementModeInterface>(MoverComp->MovementModes[Input.SyncState.MovementMode]))
 		{
-			// Contacts on the character capsule below the MaxStepHeight tend to snag the character when stepping
-			// up or down, so disable them
-
-			if (const UCommonLegacyMovementSettings* Settings = MoverComp->FindSharedSettings<UCommonLegacyMovementSettings>())
+			Chaos::FCharacterGroundConstraintHandle* ConstraintHandle = Constraint->GetProxy<Chaos::FCharacterGroundConstraintProxy>()->GetPhysicsThreadAPI();
+			if (!ConstraintHandle || !ConstraintHandle->IsEnabled() || !ConstraintHandle->GetCharacterParticle())
 			{
-				const float StepDistance = FMath::Abs(WalkingMode->TargetHeight - ConstraintHandle->GetData().GroundDistance);
-				if (StepDistance >= GPhysicsDrivenMotionDebugParams.MinStepUpDistance)
-				{
-					MinContactHeightStepUps = CharacterHeight - WalkingMode->TargetHeight + Settings->MaxStepHeight;
-				}
+				return;
 			}
-		}
-	}
-
-	for (Chaos::FContactPairModifier& PairModifier : Modifier.GetContacts(CharacterParticle))
-	{
-		const int32 CharacterIdx = CharacterParticle == PairModifier.GetParticlePair()[0] ? 0 : 1;
-		const int32 OtherIdx = CharacterIdx == 0 ? 1 : 0;
-
-		for (int32 Idx = 0; Idx < PairModifier.GetNumContacts(); ++Idx)
-		{
-			Chaos::FVec3 Point0, Point1;
-			PairModifier.GetWorldContactLocations(Idx, Point0, Point1);
-			Chaos::FVec3 CharacterPoint = CharacterIdx == 0 ? Point0 : Point1;
 
-			Chaos::FVec3 ContactNormal = PairModifier.GetWorldNormal(Idx);
-			if ((ContactNormal.Z > CosThetaMax) && CharacterPoint.Z < EndCapHeight)
-			{
-				// Disable any nearly vertical contact with the end cap of the capsule
-				// This will be handled by the character ground constraint
-				PairModifier.SetContactPointDisabled(Idx);
-			}
-			else if ((CharacterPoint.Z < MinContactHeightStepUps) && (GroundParticle == PairModifier.GetParticlePair()[OtherIdx]))
-			{
-				// In the case of steps ups disable all contacts below the max step height
-				PairModifier.SetContactPointDisabled(Idx);
-			}
+			const FPhysicsMoverSimulationContactModifierParams Params { ConstraintHandle, MoverComp->UpdatedCompAsPrimitive };
+			PhysicsMode->OnContactModification_Internal(Params, Modifier);
 		}
 	}
 }
diff --git a/Engine/Plugins/Experimental/Mover/Source/Mover/Private/PhysicsMover/Modes/PhysicsDrivenWalkingMode.cpp b/Engine/Plugins/Experimental/Mover/Source/Mover/Private/PhysicsMover/Modes/PhysicsDrivenWalkingMode.cpp
--- a/Engine/Plugins/Experimental/Mover/Source/Mover/Private/PhysicsMover/Modes/PhysicsDrivenWalkingMode.cpp
+++ b/Engine/Plugins/Experimental/Mover/Source/Mover/Private/PhysicsMover/Modes/PhysicsDrivenWalkingMode.cpp
@@ -3,13 +3,15 @@
 #include "PhysicsMover/Modes/PhysicsDrivenWalkingMode.h"
 
 #include "Chaos/Character/CharacterGroundConstraint.h"
+#include "Chaos/Character/CharacterGroundConstraintContainer.h"
+#include "Chaos/ContactModification.h"
 #include "Chaos/PhysicsObject.h"
 #include "Chaos/PhysicsObjectInternalInterface.h"
-#include "GameFramework/PhysicsVolume.h"
 #include "DefaultMovementSet/LayeredMoves/BasicLayeredMoves.h"
-#include "MoverComponent.h"
 #include "DefaultMovementSet/Settings/CommonLegacyMovementSettings.h"
+#include "GameFramework/PhysicsVolume.h"
 #include "Math/UnitConversion.h"
+#include "MoverComponent.h"
 #include "MoveLibrary/WaterMovementUtils.h"
 #include "PhysicsMover/PhysicsMovementUtils.h"
 #include "PhysicsMover/PhysicsMoverSimulationTypes.h"
@@ -39,6 +41,70 @@ void UPhysicsDrivenWalkingMode::UpdateConstraintSettings(Chaos::FCharacterGround
 	Constraint.SetTargetHeight(TargetHeight);
 }
 
+void UPhysicsDrivenWalkingMode::OnContactModification_Internal(const FPhysicsMoverSimulationContactModifierParams& Params, Chaos::FCollisionContactModifier& Modifier) const
+{
+	Chaos::EnsureIsInPhysicsThreadContext();
+
+	if (!Params.ConstraintHandle || !Params.UpdatedPrimitive)
+	{
+		return;
+	}
+
+	Chaos::FPBDRigidParticleHandle* CharacterParticle = Params.ConstraintHandle->GetCharacterParticle()->CastToRigidParticle();
+	if (!CharacterParticle || CharacterParticle->Disabled())
+	{
+		return;
+	}
+
+	const Chaos::FGeometryParticleHandle* GroundParticle = Params.ConstraintHandle->GetGroundParticle();
+	if (!GroundParticle)
+	{
+		return;
+	}
+
+	float PawnHalfHeight;
+	float PawnRadius;
+	Params.UpdatedPrimitive->CalcBoundingCylinder(PawnRadius, PawnHalfHeight);
+
+	const float CharacterHeight = CharacterParticle->GetX().Z;
+	const float EndCapHeight = CharacterHeight - PawnHalfHeight + PawnRadius;
+
+	const float CosThetaMax = 0.707f;
+
+	float MinContactHeightStepUps = CharacterHeight - 1.0e10f;
+	const float StepDistance = FMath::Abs(TargetHeight - Params.ConstraintHandle->GetData().GroundDistance);
+	if (StepDistance >= GPhysicsDrivenMotionDebugParams.MinStepUpDistance)
+	{
+		MinContactHeightStepUps = CharacterHeight - TargetHeight + CommonLegacySettings->MaxStepHeight;
+	}
+
+	for (Chaos::FContactPairModifier& PairModifier : Modifier.GetContacts(CharacterParticle))
+	{
+		const int32 CharacterIdx = CharacterParticle == PairModifier.GetParticlePair()[0] ? 0 : 1;
+		const int32 OtherIdx = CharacterIdx == 0 ? 1 : 0;
+
+		for (int32 Idx = 0; Idx < PairModifier.GetNumContacts(); ++Idx)
+		{
+			Chaos::FVec3 Point0, Point1;
+			PairModifier.GetWorldContactLocations(Idx, Point0, Point1);
+			Chaos::FVec3 CharacterPoint = CharacterIdx == 0 ? Point0 : Point1;
+
+			Chaos::FVec3 ContactNormal = PairModifier.GetWorldNormal(Idx);
+			if ((ContactNormal.Z > CosThetaMax) && CharacterPoint.Z < EndCapHeight)
+			{
+				// Disable any nearly vertical contact with the end cap of the capsule
+				// This will be handled by the character ground constraint
+				PairModifier.SetContactPointDisabled(Idx);
+			}
+			else if ((CharacterPoint.Z < MinContactHeightStepUps) && (GroundParticle == PairModifier.GetParticlePair()[OtherIdx]))
+			{
+				// In the case of steps ups disable all contacts below the max step height
+				PairModifier.SetContactPointDisabled(Idx);
+			}
+		}
+	}
+}
+
 #if WITH_EDITOR
 EDataValidationResult UPhysicsDrivenWalkingMode::IsDataValid(FDataValidationContext& Context) const
 {
diff --git a/Engine/Plugins/Experimental/Mover/Source/Mover/Public/PhysicsMover/Modes/PhysicsDrivenWalkingMode.h b/Engine/Plugins/Experimental/Mover/Source/Mover/Public/PhysicsMover/Modes/PhysicsDrivenWalkingMode.h
--- a/Engine/Plugins/Experimental/Mover/Source/Mover/Public/PhysicsMover/Modes/PhysicsDrivenWalkingMode.h
+++ b/Engine/Plugins/Experimental/Mover/Source/Mover/Public/PhysicsMover/Modes/PhysicsDrivenWalkingMode.h
@@ -25,6 +25,7 @@ public:
 	virtual bool AttemptTeleport(USceneComponent* UpdatedComponent, const FVector& TeleportPos, const FRotator& TeleportRot, const FVector& PriorVelocity, FMoverTickEndData& Output);
 
 	virtual void UpdateConstraintSettings(Chaos::FCharacterGroundConstraint& Constraint) const override;
+	virtual void OnContactModification_Internal(const FPhysicsMoverSimulationContactModifierParams& Params, Chaos::FCollisionContactModifier& Modifier) const override;
 
 #if WITH_EDITOR
 	virtual EDataValidationResult IsDataValid(class FDataValidationContext& Context) const override;
diff --git a/Engine/Plugins/Experimental/Mover/Source/Mover/Public/PhysicsMover/PhysicsMoverSimulationTypes.h b/Engine/Plugins/Experimental/Mover/Source/Mover/Public/PhysicsMover/PhysicsMoverSimulationTypes.h
--- a/Engine/Plugins/Experimental/Mover/Source/Mover/Public/PhysicsMover/PhysicsMoverSimulationTypes.h
+++ b/Engine/Plugins/Experimental/Mover/Source/Mover/Public/PhysicsMover/PhysicsMoverSimulationTypes.h
@@ -9,6 +9,17 @@
 #include "UObject/WeakInterfacePtr.h"
 #include "PhysicsMoverSimulationTypes.generated.h"
 
+//////////////////////////////////////////////////////////////////////////
+
+namespace Chaos
+{
+	class FCharacterGroundConstraint;
+	class FCharacterGroundConstraintHandle;
+	class FCollisionContactModifier;
+}
+
+class UPrimitiveComponent;
+
 //////////////////////////////////////////////////////////////////////////
 // Debug
 
@@ -33,8 +44,6 @@ struct FPhysicsMoverSimulationTickParams
 	float DeltaTimeSeconds = 0.0f;
 };
 
-namespace Chaos { class FCharacterGroundConstraintHandle; }
-
 struct FPhysicsMoverAsyncInput
 {
 	bool IsValid() const
@@ -60,8 +69,14 @@ struct FPhysicsMoverAsyncOutput
 //////////////////////////////////////////////////////////////////////////
 // Movement modes
 
+struct FPhysicsMoverSimulationContactModifierParams
+{
+	Chaos::FCharacterGroundConstraintHandle* ConstraintHandle;
+	UPrimitiveComponent* UpdatedPrimitive;
+};
+
 /**
- * PhysicsDrivenMotionMode: Interface for movement modes that are for physics driven motion
+ * UPhysicsCharacterMovementModeInterface: Interface for movement modes that are for physics driven motion
  * A physics driven motion mode needs to update the character ground constraint with the
  * parameters associated with that mode
  */
@@ -71,14 +86,16 @@ class UPhysicsCharacterMovementModeInterface : public UInterface
 	GENERATED_BODY()
 };
 
-namespace Chaos { class FCharacterGroundConstraint; }
-
 class IPhysicsCharacterMovementModeInterface
 {
 	GENERATED_BODY()
 
 public:
+	// Update the constraint settings on the game thread
 	virtual void UpdateConstraintSettings(Chaos::FCharacterGroundConstraint& Constraint) const = 0;
+
+	// Optionally run contact modification on the physics thread
+	virtual void OnContactModification_Internal(const FPhysicsMoverSimulationContactModifierParams& Params, Chaos::FCollisionContactModifier& Modifier) const {};
 };
 
 //////////////////////////////////////////////////////////////////////////
