Physics Mover - Remove current input validation

diff --git a/Engine/Plugins/Experimental/Mover/Source/Mover/Private/Backends/MoverNetworkPhysicsLiaison.cpp b/Engine/Plugins/Experimental/Mover/Source/Mover/Private/Backends/MoverNetworkPhysicsLiaison.cpp
--- a/Engine/Plugins/Experimental/Mover/Source/Mover/Private/Backends/MoverNetworkPhysicsLiaison.cpp
+++ b/Engine/Plugins/Experimental/Mover/Source/Mover/Private/Backends/MoverNetworkPhysicsLiaison.cpp
@@ -201,26 +201,8 @@ void UMoverNetworkPhysicsLiaisonComponent::SetCurrentStateData(const FMoverSyncS
 
 bool UMoverNetworkPhysicsLiaisonComponent::ValidateInputData(FMoverInputCmdContext& InputCmd) const
 {
-	bool bValidInputs = true;
-	if (FMovementSettingsInputs* MovementSettings = InputCmd.InputCollection.FindMutableDataByType<FMovementSettingsInputs>())
-	{
-		// TODO - proper data validation
-		const float MaxMaxSpeed = 1000.0f;
-		if (MovementSettings->MaxSpeed > MaxMaxSpeed)
-		{
-			MovementSettings->MaxSpeed = MaxMaxSpeed;
-			bValidInputs = false;
-		}
-
-		const float MaxAcceleration = 10000.0f;
-		if (MovementSettings->Acceleration > MaxAcceleration)
-		{
-			MovementSettings->Acceleration = MaxAcceleration;
-			bValidInputs = false;
-		}
-	}
-
-	return bValidInputs;
+	// TODO - proper data validation
+	return true;
 }
 
 UMoverNetworkPhysicsLiaisonComponent::UMoverNetworkPhysicsLiaisonComponent()
