Mover: fixed crash when using the gameplay debugger tool on a Mover actor that does not have authored inputs

#rnx

diff --git a/Engine/Plugins/Experimental/Mover/Source/Mover/Private/MoverComponent.cpp b/Engine/Plugins/Experimental/Mover/Source/Mover/Private/MoverComponent.cpp
--- a/Engine/Plugins/Experimental/Mover/Source/Mover/Private/MoverComponent.cpp
+++ b/Engine/Plugins/Experimental/Mover/Source/Mover/Private/MoverComponent.cpp
@@ -1063,13 +1063,14 @@ FRotator UMoverComponent::GetTargetOrientation() const
 	if (bHasValidCachedUsedInput)
 	{
 		const FMoverInputCmdContext& LastInputCmd = GetLastInputCmd();
-		const FCharacterDefaultInputs* MoverInputs = LastInputCmd.InputCollection.FindDataByType<FCharacterDefaultInputs>();
-
-		const FVector TargetOrientationDir = MoverInputs->GetOrientationIntentDir_WorldSpace();
-
-		if (!TargetOrientationDir.IsNearlyZero())
+		if (const FCharacterDefaultInputs* MoverInputs = LastInputCmd.InputCollection.FindDataByType<FCharacterDefaultInputs>())
 		{
-			return TargetOrientationDir.ToOrientationRotator();
+			const FVector TargetOrientationDir = MoverInputs->GetOrientationIntentDir_WorldSpace();
+
+			if (!TargetOrientationDir.IsNearlyZero())
+			{
+				return TargetOrientationDir.ToOrientationRotator();
+			}
 		}
 	}
 	
