Mover: converted StartingMovementMode property editing to a dropdown list, making it less error-prone

#rnx

diff --git a/Engine/Plugins/Experimental/Mover/Source/Mover/Private/MoverComponent.cpp b/Engine/Plugins/Experimental/Mover/Source/Mover/Private/MoverComponent.cpp
--- a/Engine/Plugins/Experimental/Mover/Source/Mover/Private/MoverComponent.cpp
+++ b/Engine/Plugins/Experimental/Mover/Source/Mover/Private/MoverComponent.cpp
@@ -767,6 +767,22 @@ bool UMoverComponent::ValidateSetup(FDataValidationContext& Context) const
 	return !bDidFindAnyProblems;
 }
 
+TArray<FString> UMoverComponent::GetStartingMovementModeNames()
+{
+	TArray<FString> PossibleModeNames;
+
+	PossibleModeNames.Add(TEXT(""));
+
+	for (const TPair<FName, TObjectPtr<UBaseMovementMode>>& Element : MovementModes)
+	{
+		FString ModeNameAsString;
+		Element.Key.ToString(ModeNameAsString);
+		PossibleModeNames.Add(ModeNameAsString);
+	}
+
+	return PossibleModeNames;
+}
+
 #endif // WITH_EDITOR
 
 
diff --git a/Engine/Plugins/Experimental/Mover/Source/Mover/Public/MoverComponent.h b/Engine/Plugins/Experimental/Mover/Source/Mover/Public/MoverComponent.h
--- a/Engine/Plugins/Experimental/Mover/Source/Mover/Public/MoverComponent.h
+++ b/Engine/Plugins/Experimental/Mover/Source/Mover/Public/MoverComponent.h
@@ -115,8 +115,8 @@ public:
 	TMap<FName, TObjectPtr<UBaseMovementMode>> MovementModes;
 
 	// Name of the first mode to start in when simulation begins. Must have a mapping in MovementModes. Only used during initialization.
-	UPROPERTY(EditAnywhere, BlueprintReadOnly, Category=Mover)
-	FName StartingMovementMode;
+	UPROPERTY(EditAnywhere, BlueprintReadOnly, Category=Mover, meta=(GetOptions=GetStartingMovementModeNames))
+	FName StartingMovementMode = NAME_None;
 
 	// Transition checks that are always evaluated regardless of mode. Evaluated in order, stopping at the first successful transition check. Mode-owned transitions take precedence. */
 	UPROPERTY(EditAnywhere, BlueprintReadWrite, Instanced, Category=Mover)
@@ -340,6 +340,9 @@ protected:
 	virtual void PostTransacted(const FTransactionObjectEvent& TransactionEvent) override;
 	bool ValidateSetup(class FDataValidationContext& ValidationErrors) const;
 	virtual EDataValidationResult IsDataValid(class FDataValidationContext& Context) const override;
+
+	UFUNCTION()
+	TArray<FString> GetStartingMovementModeNames();
 	#endif // WITH_EDITOR
 
 	UFUNCTION()
