Physics Mover
- Improve ground query code to prevent jitter and better handle unwalkable surfaces
- Fix physics mover mode settings so that the default is good for the UE example character

diff --git a/Engine/Plugins/Experimental/Mover/Source/Mover/Private/PhysicsMover/Modes/PhysicsDrivenFallingMode.cpp b/Engine/Plugins/Experimental/Mover/Source/Mover/Private/PhysicsMover/Modes/PhysicsDrivenFallingMode.cpp
--- a/Engine/Plugins/Experimental/Mover/Source/Mover/Private/PhysicsMover/Modes/PhysicsDrivenFallingMode.cpp
+++ b/Engine/Plugins/Experimental/Mover/Source/Mover/Private/PhysicsMover/Modes/PhysicsDrivenFallingMode.cpp
@@ -3,12 +3,12 @@
 #include "PhysicsMover/Modes/PhysicsDrivenFallingMode.h"
 
 #include "Chaos/Character/CharacterGroundConstraint.h"
-#include "GameFramework/PhysicsVolume.h"
-#include "MoverComponent.h"
 #include "DefaultMovementSet/Settings/CommonLegacyMovementSettings.h"
+#include "GameFramework/PhysicsVolume.h"
 #include "Math/UnitConversion.h"
 #include "MoveLibrary/FloorQueryUtils.h"
 #include "MoveLibrary/WaterMovementUtils.h"
+#include "MoverComponent.h"
 #include "PhysicsMover/PhysicsMovementUtils.h"
 #if WITH_EDITOR
 #include "Backends/MoverNetworkPhysicsLiaison.h"
@@ -85,15 +85,19 @@ void UPhysicsDrivenFallingMode::OnSimulationTick(const FSimulationTickParams& Pa
 
 	// Find floor
 
+	FFloorCheckResult FloorResult;
+	FWaterCheckResult WaterResult;
+
 	float PawnHalfHeight;
 	float PawnRadius;
 	UpdatedPrimitive->CalcBoundingCylinder(PawnRadius, PawnHalfHeight);
 
-	FFloorCheckResult FloorResult;
-	FWaterCheckResult WaterResult;
-	UPhysicsMovementUtils::FindFloor(StartingSyncState->GetLocation_WorldSpace(), StartingSyncState->GetVelocity_WorldSpace() * DeltaSeconds,
-	UpdatedPrimitive, UpDir, PawnRadius, TargetHeight, CommonLegacySettings->MaxStepHeight,
-		CommonLegacySettings->MaxWalkSlopeCosine, FloorResult, WaterResult);
+	const float QueryDistance = FMath::Max(1.1f * TargetHeight, TargetHeight - UpDir.Dot(ProposedMove.LinearVelocity) * DeltaSeconds);
+	const float ShrinkRadius = 5.0f; // TODO - Make this a user setting
+	const float QueryRadius = FMath::Max(PawnRadius - ShrinkRadius, 0.0f);
+
+	UPhysicsMovementUtils::FloorSweep(StartingSyncState->GetLocation_WorldSpace(), StartingSyncState->GetVelocity_WorldSpace() * DeltaSeconds,
+		UpdatedPrimitive, UpDir, QueryRadius, QueryDistance, CommonLegacySettings->MaxWalkSlopeCosine, TargetHeight, FloorResult, WaterResult);
 
 	SimBlackboard->Set(CommonBlackboard::LastFloorResult, FloorResult);
 	SimBlackboard->Set(CommonBlackboard::LastWaterResult, WaterResult);
diff --git a/Engine/Plugins/Experimental/Mover/Source/Mover/Private/PhysicsMover/Modes/PhysicsDrivenSwimmingMode.cpp b/Engine/Plugins/Experimental/Mover/Source/Mover/Private/PhysicsMover/Modes/PhysicsDrivenSwimmingMode.cpp
--- a/Engine/Plugins/Experimental/Mover/Source/Mover/Private/PhysicsMover/Modes/PhysicsDrivenSwimmingMode.cpp
+++ b/Engine/Plugins/Experimental/Mover/Source/Mover/Private/PhysicsMover/Modes/PhysicsDrivenSwimmingMode.cpp
@@ -3,13 +3,13 @@
 #include "PhysicsMover/Modes/PhysicsDrivenSwimmingMode.h"
 
 #include "Chaos/Character/CharacterGroundConstraint.h"
+#include "DefaultMovementSet/LayeredMoves/BasicLayeredMoves.h"
+#include "DefaultMovementSet/Settings/CommonLegacyMovementSettings.h"
+#include "GameFramework/PhysicsVolume.h"
 #include "Math/UnitConversion.h"
 #include "MoveLibrary/FloorQueryUtils.h"
 #include "MoveLibrary/WaterMovementUtils.h"
-#include "DefaultMovementSet/LayeredMoves/BasicLayeredMoves.h"
 #include "MoverComponent.h"
-#include "DefaultMovementSet/Settings/CommonLegacyMovementSettings.h"
-#include "GameFramework/PhysicsVolume.h"
 #include "PhysicsMover/PhysicsMovementUtils.h"
 #if WITH_EDITOR
 #include "Backends/MoverNetworkPhysicsLiaison.h"
@@ -94,11 +94,12 @@ void UPhysicsDrivenSwimmingMode::OnSimulationTick(const FSimulationTickParams& P
 	float PawnHalfHeight;
 	float PawnRadius;
 	UpdatedPrimitive->CalcBoundingCylinder(PawnRadius, PawnHalfHeight);
+
+	const float QueryDistance= 2.0f * PawnHalfHeight;
 	
-	UPhysicsMovementUtils::FindFloor(StartingSyncState->GetLocation_WorldSpace(), StartingSyncState->GetVelocity_WorldSpace() * DeltaSeconds,
-		UpdatedPrimitive, UpDir, PawnRadius, TargetHeight, CommonLegacySettings->MaxStepHeight,
-		CommonLegacySettings->MaxWalkSlopeCosine, FloorResult, WaterResult);
-	
+	UPhysicsMovementUtils::FloorSweep(StartingSyncState->GetLocation_WorldSpace(), StartingSyncState->GetVelocity_WorldSpace() * DeltaSeconds,
+		UpdatedPrimitive, UpDir, PawnRadius, QueryDistance, CommonLegacySettings->MaxWalkSlopeCosine, TargetHeight, FloorResult, WaterResult);
+
 	SimBlackboard->Set(CommonBlackboard::LastFloorResult, FloorResult);
 	SimBlackboard->Set(CommonBlackboard::LastWaterResult, WaterResult);
 	
diff --git a/Engine/Plugins/Experimental/Mover/Source/Mover/Private/PhysicsMover/Modes/PhysicsDrivenWalkingMode.cpp b/Engine/Plugins/Experimental/Mover/Source/Mover/Private/PhysicsMover/Modes/PhysicsDrivenWalkingMode.cpp
--- a/Engine/Plugins/Experimental/Mover/Source/Mover/Private/PhysicsMover/Modes/PhysicsDrivenWalkingMode.cpp
+++ b/Engine/Plugins/Experimental/Mover/Source/Mover/Private/PhysicsMover/Modes/PhysicsDrivenWalkingMode.cpp
@@ -9,9 +9,11 @@
 #include "Chaos/PhysicsObjectInternalInterface.h"
 #include "DefaultMovementSet/LayeredMoves/BasicLayeredMoves.h"
 #include "DefaultMovementSet/Settings/CommonLegacyMovementSettings.h"
+#include "Engine/World.h"
 #include "GameFramework/PhysicsVolume.h"
 #include "Math/UnitConversion.h"
 #include "MoverComponent.h"
+#include "MoveLibrary/GroundMovementUtils.h"
 #include "MoveLibrary/WaterMovementUtils.h"
 #include "PhysicsMover/PhysicsMovementUtils.h"
 #include "PhysicsMover/PhysicsMoverSimulationTypes.h"
@@ -39,6 +41,7 @@ void UPhysicsDrivenWalkingMode::UpdateConstraintSettings(Chaos::FCharacterGround
 	Constraint.SetTwistTorqueLimit(FUnitConversion::Convert(TwistTorqueLimit, EUnit::NewtonMeters, EUnit::KilogramCentimetersSquaredPerSecondSquared));
 	Constraint.SetSwingTorqueLimit(FUnitConversion::Convert(SwingTorqueLimit, EUnit::NewtonMeters, EUnit::KilogramCentimetersSquaredPerSecondSquared));
 	Constraint.SetTargetHeight(TargetHeight);
+	Constraint.SetDampingFactor(GroundDamping);
 }
 
 void UPhysicsDrivenWalkingMode::OnContactModification_Internal(const FPhysicsMoverSimulationContactModifierParams& Params, Chaos::FCollisionContactModifier& Modifier) const
@@ -105,6 +108,142 @@ void UPhysicsDrivenWalkingMode::OnContactModification_Internal(const FPhysicsMov
 	}
 }
 
+bool UPhysicsDrivenWalkingMode::CanStepUpOnHitSurface(const FFloorCheckResult& FloorResult) const
+{
+	const float StepHeight = TargetHeight - FloorResult.FloorDist;
+
+	bool bWalkable = StepHeight <= CommonLegacySettings->MaxStepHeight;
+	constexpr float MinStepHeight = 2.0f;
+	const bool SteppingUp = StepHeight > MinStepHeight;
+	if (bWalkable && SteppingUp)
+	{
+		bWalkable = UGroundMovementUtils::CanStepUpOnHitSurface(FloorResult.HitResult);
+	}
+
+	return bWalkable;
+}
+
+void UPhysicsDrivenWalkingMode::FloorCheck(const FMoverDefaultSyncState& SyncState, const FProposedMove& ProposedMove, UPrimitiveComponent* UpdatedPrimitive, float DeltaSeconds,
+	FFloorCheckResult& OutFloorResult, FWaterCheckResult& OutWaterResult, FVector& OutDeltaPos) const
+{
+	const FVector UpDir = GetMoverComponent()->GetUpDirection();
+
+	FVector DeltaPos = ProposedMove.LinearVelocity * DeltaSeconds;
+	OutDeltaPos = DeltaPos;
+
+	float PawnHalfHeight;
+	float PawnRadius;
+	UpdatedPrimitive->CalcBoundingCylinder(PawnRadius, PawnHalfHeight);
+
+	const float FloorSweepDistance = TargetHeight + CommonLegacySettings->MaxStepHeight;
+	const float ShrinkRadius = 1.0f;
+	const float QueryRadius = FMath::Max(PawnRadius - ShrinkRadius, 0.0f);
+	UPhysicsMovementUtils::FloorSweep(SyncState.GetLocation_WorldSpace(), DeltaPos, UpdatedPrimitive, UpDir,
+		QueryRadius, FloorSweepDistance, CommonLegacySettings->MaxWalkSlopeCosine, TargetHeight, OutFloorResult, OutWaterResult);
+
+	if (!OutFloorResult.bBlockingHit)
+	{
+		// Floor not found
+		return;
+	}
+
+	bool bWalkableFloor = OutFloorResult.bWalkableFloor && CanStepUpOnHitSurface(OutFloorResult);
+	if (bWalkableFloor)
+	{
+		// Walkable floor found
+		OutFloorResult.bWalkableFloor = true;
+		return;
+	}
+
+	// Hit something but not walkable. Try a new query to find a walkable surface
+	const float StepBlockedHeight = TargetHeight - PawnHalfHeight + PawnRadius;
+	const float StepHeight = TargetHeight - OutFloorResult.FloorDist;
+
+	if (StepHeight > StepBlockedHeight)
+	{
+		// Collision should prevent movement. Just try to find ground at start of movement
+		const float ShrinkMultiplier = 0.75f;
+		UPhysicsMovementUtils::FloorSweep(SyncState.GetLocation_WorldSpace(), DeltaPos, UpdatedPrimitive, UpDir,
+			ShrinkMultiplier * QueryRadius, FloorSweepDistance, CommonLegacySettings->MaxWalkSlopeCosine, TargetHeight, OutFloorResult, OutWaterResult);
+
+		OutFloorResult.bWalkableFloor = OutFloorResult.bWalkableFloor && CanStepUpOnHitSurface(OutFloorResult);
+		return;
+	}
+
+	if (DeltaPos.SizeSquared() < UE_SMALL_NUMBER)
+	{
+		// Stationary
+		OutDeltaPos = FVector::ZeroVector;
+		return;
+	}
+
+	// Try to limit the movement to remain on a walkable surface
+	FVector NewDeltaPos = DeltaPos;
+	float NewQueryRadius = QueryRadius;
+
+	FVector HorizSurfaceDir = FVector::VectorPlaneProject(OutFloorResult.HitResult.ImpactNormal, UpDir);
+	float HorizSurfaceDirSizeSq = HorizSurfaceDir.SizeSquared();
+	bool bFoundOutwardDir = false;
+	if (HorizSurfaceDirSizeSq > UE_SMALL_NUMBER)
+	{
+		HorizSurfaceDir *= FMath::InvSqrt(HorizSurfaceDirSizeSq);
+		bFoundOutwardDir = true;
+	}
+	else
+	{
+		// Flat unwalkable surface. Try and get the horizontal direction from the normal instead
+		HorizSurfaceDir = FVector::VectorPlaneProject(OutFloorResult.HitResult.Normal, UpDir);
+		HorizSurfaceDirSizeSq = HorizSurfaceDir.SizeSquared();
+
+		if (HorizSurfaceDirSizeSq > UE_SMALL_NUMBER)
+		{
+			HorizSurfaceDir *= FMath::InvSqrt(HorizSurfaceDirSizeSq);
+			bFoundOutwardDir = true;
+		}
+	}
+
+	if (bFoundOutwardDir)
+	{
+		// If we're moving away try a ray query at the end of the motion
+		const float DP = DeltaPos.Dot(HorizSurfaceDir);
+		if (DP > 0.0f)
+		{
+			NewQueryRadius = 0.0f;
+		}
+		else
+		{
+			NewQueryRadius = 0.75f * QueryRadius;
+			NewDeltaPos = DeltaPos - DP * HorizSurfaceDir;
+		}
+
+		UPhysicsMovementUtils::FloorSweep(SyncState.GetLocation_WorldSpace(), NewDeltaPos, UpdatedPrimitive, UpDir,
+			NewQueryRadius, FloorSweepDistance, CommonLegacySettings->MaxWalkSlopeCosine, TargetHeight, OutFloorResult, OutWaterResult);
+
+		OutFloorResult.bWalkableFloor = OutFloorResult.bWalkableFloor && CanStepUpOnHitSurface(OutFloorResult);
+		if (OutFloorResult.bWalkableFloor)
+		{
+			OutDeltaPos = NewDeltaPos;
+		}
+		else
+		{
+			OutFloorResult.bWalkableFloor = false;
+		}
+	}
+	else
+	{
+		// Try a query at the start of the movement to find a walkable surface and prevent movement
+
+		NewDeltaPos = FVector::ZeroVector;
+		NewQueryRadius = 0.75f * QueryRadius;
+
+		UPhysicsMovementUtils::FloorSweep(SyncState.GetLocation_WorldSpace(), NewDeltaPos, UpdatedPrimitive, UpDir,
+			NewQueryRadius, FloorSweepDistance, CommonLegacySettings->MaxWalkSlopeCosine, TargetHeight, OutFloorResult, OutWaterResult);
+
+		OutFloorResult.bWalkableFloor = OutFloorResult.bWalkableFloor && CanStepUpOnHitSurface(OutFloorResult);
+		OutDeltaPos = NewDeltaPos;
+	}
+}
+
 #if WITH_EDITOR
 EDataValidationResult UPhysicsDrivenWalkingMode::IsDataValid(FDataValidationContext& Context) const
 {
@@ -160,15 +299,13 @@ void UPhysicsDrivenWalkingMode::OnSimulationTick(const FSimulationTickParams& Pa
 	}
 
 	// Floor query
-	float PawnHalfHeight;
-	float PawnRadius;
-	UpdatedPrimitive->CalcBoundingCylinder(PawnRadius, PawnHalfHeight);
-
 	FFloorCheckResult FloorResult;
 	FWaterCheckResult WaterResult;
-	UPhysicsMovementUtils::FindFloor(StartingSyncState->GetLocation_WorldSpace(), StartingSyncState->GetVelocity_WorldSpace() * DeltaSeconds,
-		UpdatedPrimitive, UpDir, PawnRadius, TargetHeight, CommonLegacySettings->MaxStepHeight,
-		CommonLegacySettings->MaxWalkSlopeCosine, FloorResult, WaterResult);
+	FVector OutDeltaPos;
+
+	FloorCheck(*StartingSyncState, ProposedMove, UpdatedPrimitive, DeltaSeconds, FloorResult, WaterResult, OutDeltaPos);
+
+	ProposedMove.LinearVelocity = OutDeltaPos / DeltaSeconds;
 
 	SimBlackboard->Set(CommonBlackboard::LastFloorResult, FloorResult);
 	SimBlackboard->Set(CommonBlackboard::LastWaterResult, WaterResult);
diff --git a/Engine/Plugins/Experimental/Mover/Source/Mover/Private/PhysicsMover/PhysicsMovementUtils.cpp b/Engine/Plugins/Experimental/Mover/Source/Mover/Private/PhysicsMover/PhysicsMovementUtils.cpp
--- a/Engine/Plugins/Experimental/Mover/Source/Mover/Private/PhysicsMover/PhysicsMovementUtils.cpp
+++ b/Engine/Plugins/Experimental/Mover/Source/Mover/Private/PhysicsMover/PhysicsMovementUtils.cpp
@@ -18,17 +18,18 @@
 
 extern FPhysicsDrivenMotionDebugParams GPhysicsDrivenMotionDebugParams;
 
-void UPhysicsMovementUtils::FindFloor(
+void UPhysicsMovementUtils::FloorSweep(
 	const FVector& Location,
 	const FVector& DeltaPos,
 	const UPrimitiveComponent* UpdatedPrimitive,
 	const FVector& UpDir,
 	float QueryRadius,
-	float TargetHeight,
-	float MaxStepHeight,
+	float QueryDistance,
 	float MaxWalkSlopeCosine,
+	float TargetHeight,
 	FFloorCheckResult& OutFloorResult,
-	FWaterCheckResult& OutWaterResult)
+	FWaterCheckResult& OutWaterResult
+)
 {
 	if (const UWorld* World = UpdatedPrimitive->GetWorld())
 	{
@@ -42,17 +43,17 @@ void UPhysicsMovementUtils::FindFloor(
 		ResponseParams.CollisionResponse.SetResponse(ECC_Destructible, ECR_Block);
 		ResponseParams.CollisionResponse.SetResponse(ECC_PhysicsBody, ECR_Block);
 
-		const FVector DeltaPosVert = DeltaPos.ProjectOnTo(UpDir);
-		const FVector DeltaPosHoriz = DeltaPos - DeltaPosVert;
+		TArray<FHitResult> Hits;
 
-		const FVector WaterQueryLocation = Location - (TargetHeight * UpDir);
+		const float DeltaPosVertLength = DeltaPos.Dot(UpDir);
+		const FVector DeltaPosHoriz = DeltaPos - DeltaPosVertLength * UpDir;
 
-		const float StartOffset = TargetHeight - MaxStepHeight - QueryRadius;
-		const float SweepDistance = MaxStepHeight + QueryRadius + FMath::Max(MaxStepHeight, -DeltaPosVert.Dot(UpDir));
+		// Make sure the query is long enough to include the vertical movement
+		const float AdjustedQueryDistance = FMath::Max(UE_KINDA_SMALL_NUMBER + DeltaPosVertLength + TargetHeight, QueryDistance);
 
-		FVector Start = Location + DeltaPosHoriz - StartOffset * UpDir;
-		FVector End = Start - SweepDistance * UpDir;
-		TArray<FHitResult> Hits;
+		// The bottom of the query shape should be at the integrated location (ignoring vertical movement)
+		FVector Start = Location + DeltaPosHoriz + (QueryRadius + UE_KINDA_SMALL_NUMBER) * UpDir;
+		FVector End = Start - AdjustedQueryDistance * UpDir;
 		FHitResult OutHit;
 		if (World->SweepMultiByChannel(Hits, Start, End, FQuat::Identity, CollisionChannel, FCollisionShape::MakeSphere(QueryRadius), QueryParams, ResponseParams))
 		{
@@ -71,11 +72,7 @@ void UPhysicsMovementUtils::FindFloor(
 
 		if (OutHit.bBlockingHit)
 		{
-			const float Distance = UpDir.Dot(Location - OutHit.ImpactPoint);
-			const float StepHeight = TargetHeight - Distance;
-			const float MinStepHeight = 5.0f;
-
-			bool bWalkable = IsHitSurfaceWalkableWithStepUpCheck(OutHit, StepHeight, MaxStepHeight, MinStepHeight, MaxWalkSlopeCosine);
+			bool bWalkable = UFloorQueryUtils::IsHitSurfaceWalkable(OutHit, MaxWalkSlopeCosine);
 
 #if PHYSICSDRIVENMOTION_DEBUG_DRAW
 			if (GPhysicsDrivenMotionDebugParams.DebugDrawGroundQueries)
@@ -85,146 +82,19 @@ void UPhysicsMovementUtils::FindFloor(
 				Chaos::FDebugDrawQueue::GetInstance().DrawDebugCapsule(Center, 0.5f * OutHit.Distance + QueryRadius, QueryRadius, FQuat::Identity, Color, false, -1.f, 10, 1.0f);
 			}
 #endif
-
-			GetWaterResultFromHitResults(Hits, WaterQueryLocation, TargetHeight, OutWaterResult);
-
-			if (bWalkable)
-			{
-				// Found a walkable surface
-				OutFloorResult.bBlockingHit = true;
-				OutFloorResult.bWalkableFloor = true;
-				OutFloorResult.FloorDist = UpDir.Dot(Location - OutHit.ImpactPoint);;
-				OutFloorResult.HitResult = OutHit;
-			}
-			else
-			{
-				// Hit something but not walkable. Fire some line queries to try and find a walkable surface
-
-				bool bLineQueryWalkable = false;
-				float LineQueryDistance = 0.0f;
-
-				FVector MovementDir = FVector::ForwardVector;
-				FVector PerpDir = FVector::RightVector;
-				FVector Fwd = DeltaPos - DeltaPos.Dot(UpDir) * UpDir;
-				const float FwdSizeSq = Fwd.SizeSquared();
-				if (FwdSizeSq > UE_SMALL_NUMBER)
-				{
-					MovementDir = Fwd * FMath::InvSqrtEst(FwdSizeSq);
-					PerpDir = UpDir.Cross(MovementDir);
-				}
-
-				const int NumLineTraces = 4;
-				FVector StartPoints[NumLineTraces];
-
-				StartPoints[0] = Start + QueryRadius * MovementDir;
-				StartPoints[1] = Start - QueryRadius * MovementDir;
-				StartPoints[2] = Start - QueryRadius * PerpDir;
-				StartPoints[3] = Start + QueryRadius * PerpDir;
-
-				FHitResult LineQueryHit[NumLineTraces];
-
-				float MinDistance = SweepDistance + QueryRadius;
-				int32 MinDistanceIdx = INDEX_NONE;
-				for (int32 Idx = 0; Idx < NumLineTraces; ++Idx)
-				{
-					Start = StartPoints[Idx];
-					End = Start - UpDir * SweepDistance;
-					World->LineTraceSingleByChannel(LineQueryHit[Idx], Start, End, CollisionChannel, QueryParams, ResponseParams);
-
-					if (LineQueryHit[Idx].IsValidBlockingHit())
-					{
-						LineQueryDistance = UpDir.Dot(Location - LineQueryHit[Idx].ImpactPoint);
-						const float LineQueryStepHeight = TargetHeight - LineQueryDistance;
-
-						bLineQueryWalkable = IsHitSurfaceWalkableWithStepUpCheck(LineQueryHit[Idx], LineQueryStepHeight, MaxStepHeight, MinStepHeight, MaxWalkSlopeCosine);
-
-						if (bLineQueryWalkable && LineQueryDistance < MinDistance)
-						{
-							MinDistance = LineQueryDistance;
-							MinDistanceIdx = Idx;
-						}
-					}
-
-#if PHYSICSDRIVENMOTION_DEBUG_DRAW
-					if (GPhysicsDrivenMotionDebugParams.DebugDrawGroundQueries)
-					{
-						FVector Dir = End - Start;
-						const float SizeSq = Dir.SizeSquared();
-						if (SizeSq > UE_SMALL_NUMBER)
-						{
-							Dir *= FMath::InvSqrt(SizeSq);
-						}
-
-						if (LineQueryHit[Idx].bBlockingHit)
-						{
-							const FVector Center = Start + 0.5f * LineQueryHit[Idx].Distance * Dir;
-							const FColor Color = bLineQueryWalkable ? FColor::Green : FColor::Red;
-							Chaos::FDebugDrawQueue::GetInstance().DrawDebugLine(Start, LineQueryHit[Idx].ImpactPoint, Color, false, -1.f, 10, 1.0f);
-						}
-						else
-						{
-							const FVector Center = 0.5f * (Start + End);
-							float Dist = SizeSq > UE_SMALL_NUMBER ? FMath::Sqrt(SizeSq) : 0.0f;
-							Chaos::FDebugDrawQueue::GetInstance().DrawDebugLine(Start, End, FColor::Silver, false, -1.f, 10, 1.0f);
-						}
-					}
-#endif
-				}
-
-				if (MinDistanceIdx != INDEX_NONE)
-				{
-					// Found a walkable surface
-					OutFloorResult.bBlockingHit = true;
-					OutFloorResult.bWalkableFloor = true;
-					OutFloorResult.FloorDist = MinDistance;
-					OutFloorResult.HitResult = LineQueryHit[MinDistanceIdx];
-				}
-				else
-				{
-					// Didn't find a walkable surface. Use original sphere query result
-					OutFloorResult.bBlockingHit = true;
-					OutFloorResult.bWalkableFloor = false;
-					OutFloorResult.FloorDist = Distance;
-					OutFloorResult.HitResult = OutHit;
-				}
-
-				// May have missed a water result due to the original query blocking on a non-walkable surface
-				// so run again with no blocking results
-				ResponseParams.CollisionResponse.SetAllChannels(ECR_Overlap);
-				World->SweepMultiByChannel(Hits, Start, End, FQuat::Identity, CollisionChannel, FCollisionShape::MakeSphere(QueryRadius), QueryParams, ResponseParams);
-				GetWaterResultFromHitResults(Hits, WaterQueryLocation, TargetHeight, OutWaterResult);
-			}
+			OutFloorResult.bBlockingHit = true;
+			OutFloorResult.bWalkableFloor = bWalkable;
+			OutFloorResult.FloorDist = UpDir.Dot(Location - OutHit.ImpactPoint);
+			OutFloorResult.HitResult = OutHit;
 		}
 		else
 		{
-			// Sweep didn't hit anything blocking
 			OutFloorResult.Clear();
 			OutFloorResult.FloorDist = 1.0e10f;
-
-			// In deep waters i.e., WaterDepth > SweepDistance, we may hit only water and no blocking surfaces
-			GetWaterResultFromHitResults(Hits, WaterQueryLocation, TargetHeight, OutWaterResult);
 		}
-	}
-}
 
-bool UPhysicsMovementUtils::IsHitSurfaceWalkableWithStepUpCheck(const FHitResult& Hit, float StepHeight, float MaxStepHeight, float MinStepUpHeight, float MaxWalkSlopeCosine)
-{
-	bool bWalkable = false;
-	if (StepHeight <= MaxStepHeight)
-	{
-		bWalkable = UFloorQueryUtils::IsHitSurfaceWalkable(Hit, MaxWalkSlopeCosine);
-		if (bWalkable)
-		{
-			const float MinStepHeight = 1.0f;
-			const bool bSteppingUp = StepHeight > MinStepHeight;
-			if (bSteppingUp && !UGroundMovementUtils::CanStepUpOnHitSurface(Hit))
-			{
-				bWalkable = false;
-			}
-		}
+		GetWaterResultFromHitResults(Hits, Location, TargetHeight, OutWaterResult);
 	}
-
-	return bWalkable;
 }
 
 const Chaos::FPBDRigidParticleHandle* UPhysicsMovementUtils::GetRigidParticleHandleFromHitResult(const FHitResult& HitResult)
diff --git a/Engine/Plugins/Experimental/Mover/Source/Mover/Public/PhysicsMover/Modes/PhysicsDrivenFallingMode.h b/Engine/Plugins/Experimental/Mover/Source/Mover/Public/PhysicsMover/Modes/PhysicsDrivenFallingMode.h
--- a/Engine/Plugins/Experimental/Mover/Source/Mover/Public/PhysicsMover/Modes/PhysicsDrivenFallingMode.h
+++ b/Engine/Plugins/Experimental/Mover/Source/Mover/Public/PhysicsMover/Modes/PhysicsDrivenFallingMode.h
@@ -39,5 +39,5 @@ public:
 
 	// Target height for the character. This is the desired distance from the center of the capsule to the floor
 	UPROPERTY(EditAnywhere, Category = "Physics Mover", meta = (ClampMin = "0", UIMin = "0", ForceUnits = "cm"))
-	float TargetHeight = 54.0f;
+	float TargetHeight = 95.0f;
 };
diff --git a/Engine/Plugins/Experimental/Mover/Source/Mover/Public/PhysicsMover/Modes/PhysicsDrivenWalkingMode.h b/Engine/Plugins/Experimental/Mover/Source/Mover/Public/PhysicsMover/Modes/PhysicsDrivenWalkingMode.h
--- a/Engine/Plugins/Experimental/Mover/Source/Mover/Public/PhysicsMover/Modes/PhysicsDrivenWalkingMode.h
+++ b/Engine/Plugins/Experimental/Mover/Source/Mover/Public/PhysicsMover/Modes/PhysicsDrivenWalkingMode.h
@@ -10,6 +10,7 @@
 
 
 namespace Chaos { class FCharacterGroundConstraint; }
+struct FWaterCheckResult;
 
 /**
  * PhysicsDrivenWalkingMode: Override base kinematic walking mode for physics based motion.
@@ -49,7 +50,7 @@ public:
 
 	// Target height for the character. This is the desired distance from the center of the capsule to the floor
 	UPROPERTY(EditAnywhere, Category = "Physics Mover", meta = (ClampMin = "0", UIMin = "0", ForceUnits = "cm"))
-	float TargetHeight = 54.0f;
+	float TargetHeight = 95.0f;
 
 	// Damping factor to control the softness of the interaction between the character and the ground
 	// Set to 0 for no damping and 1 for maximum damping
@@ -73,8 +74,13 @@ public:
 	// This provides some grace period when walking off of an edge during which locomotion
 	// and jumping are still possible even though the character has started falling under gravity
 	UPROPERTY(EditAnywhere, Category = "Physics Mover", meta = (ClampMin = "0", UIMin = "0", ForceUnits = "s"))
-	float MaxUnsupportedTimeBeforeFalling = 0.05f;
+	float MaxUnsupportedTimeBeforeFalling = 0.06f;
 
 protected:
 	void SwitchToState(const FName& StateName, const FSimulationTickParams& Params, FMoverTickEndData& OutputState);
+
+	bool CanStepUpOnHitSurface(const FFloorCheckResult& FloorResult) const;
+
+	void FloorCheck(const FMoverDefaultSyncState& SyncState, const FProposedMove& ProposedMove, UPrimitiveComponent* UpdatedPrimitive, float DeltaSeconds,
+		FFloorCheckResult& OutFloorResult, FWaterCheckResult& OutWaterResult, FVector& OutDeltaPos) const;
 };
diff --git a/Engine/Plugins/Experimental/Mover/Source/Mover/Public/PhysicsMover/PhysicsMovementUtils.h b/Engine/Plugins/Experimental/Mover/Source/Mover/Public/PhysicsMover/PhysicsMovementUtils.h
--- a/Engine/Plugins/Experimental/Mover/Source/Mover/Public/PhysicsMover/PhysicsMovementUtils.h
+++ b/Engine/Plugins/Experimental/Mover/Source/Mover/Public/PhysicsMover/PhysicsMovementUtils.h
@@ -17,14 +17,12 @@ class UPrimitiveComponent;
 class MOVER_API UPhysicsMovementUtils
 {
 public:
-	static void FindFloor(const FVector& Location, const FVector& DeltaPos, const UPrimitiveComponent* UpdatedPrimitive, const FVector& UpDir, float QueryRadius, float TargetHeight, float MaxStepHeight, float MaxWalkSlopeCosine, FFloorCheckResult& OutFloorResult, FWaterCheckResult& OutWaterResult);
+	static void FloorSweep(const FVector& Location, const FVector& DeltaPos, const UPrimitiveComponent* UpdatedPrimitive, const FVector& UpDir,
+		float QueryRadius, float QueryDistance, float MaxWalkSlopeCosine, float TargetHeight, FFloorCheckResult& OutFloorResult, FWaterCheckResult& OutWaterResult);
 
 	// If the hit result hit something, return the particle handle
 	static const Chaos::FPBDRigidParticleHandle* GetRigidParticleHandleFromHitResult(const FHitResult& HitResult);
 
-	// Checks if the hit surface is walkable and, if stepping up, whether the surface can be stepped up on
-	static bool IsHitSurfaceWalkableWithStepUpCheck(const FHitResult& Hit, float StepHeight, float MaxStepHeight, float MinStepUpHeight, float MaxWalkSlopeCosine);
-
 	// Checks if any hit is with water and, if so, fills in the OutWaterResult
 	static bool GetWaterResultFromHitResults(const TArray<FHitResult>& Hits, const FVector& Location, const float TargetHeight, FWaterCheckResult& OutWaterResult);
 
