Mover 2.0 adjusting default decceleration

#rnx

diff --git a/Engine/Plugins/Experimental/Mover/Source/Mover/Public/DefaultMovementSet/Settings/CommonLegacyMovementSettings.h b/Engine/Plugins/Experimental/Mover/Source/Mover/Public/DefaultMovementSet/Settings/CommonLegacyMovementSettings.h
--- a/Engine/Plugins/Experimental/Mover/Source/Mover/Public/DefaultMovementSet/Settings/CommonLegacyMovementSettings.h
+++ b/Engine/Plugins/Experimental/Mover/Source/Mover/Public/DefaultMovementSet/Settings/CommonLegacyMovementSettings.h
@@ -78,7 +78,7 @@ public:
 	
 	/** Default max linear rate of deceleration when there is no controlled input */
 	UPROPERTY(BlueprintReadWrite, EditAnywhere, Category="General", meta = (ClampMin = "0", UIMin = "0", ForceUnits = "cm/s^2"))
-	float Deceleration = 8000.f;
+	float Deceleration = 4000.f;
 
 	/** Default max linear rate of acceleration for controlled input. May be scaled based on magnitude of input. */
 	UPROPERTY(BlueprintReadWrite, EditAnywhere, Category="General", meta = (ClampMin = "0", UIMin = "0", ForceUnits = "cm/s^2"))
