Mover 2.0 turn generator providing a safe way to change and edit turn generators at runtime

#rnx

diff --git a/Engine/Plugins/Experimental/Mover/Source/Mover/Private/DefaultMovementSet/Modes/WalkingMode.cpp b/Engine/Plugins/Experimental/Mover/Source/Mover/Private/DefaultMovementSet/Modes/WalkingMode.cpp
--- a/Engine/Plugins/Experimental/Mover/Source/Mover/Private/DefaultMovementSet/Modes/WalkingMode.cpp
+++ b/Engine/Plugins/Experimental/Mover/Source/Mover/Private/DefaultMovementSet/Modes/WalkingMode.cpp
@@ -300,6 +300,23 @@ void UWalkingMode::OnSimulationTick(const FSimulationTickParams& Params, FMoverT
 
 }
 
+UObject* UWalkingMode::GetTurnGenerator()
+{
+	return TurnGenerator;
+}
+
+void UWalkingMode::SetTurnGeneratorClass(TSubclassOf<UObject> TurnGeneratorClass)
+{
+	if (TurnGeneratorClass)
+	{
+		TurnGenerator = NewObject<UObject>(this, TurnGeneratorClass);
+	}
+	else
+	{
+		TurnGenerator = nullptr; // Clearing the turn generator is valid - will go back to the default turn generation
+	}
+}
+
 
 void UWalkingMode::OnRegistered(const FName ModeName)
 {
diff --git a/Engine/Plugins/Experimental/Mover/Source/Mover/Public/DefaultMovementSet/Modes/WalkingMode.h b/Engine/Plugins/Experimental/Mover/Source/Mover/Public/DefaultMovementSet/Modes/WalkingMode.h
--- a/Engine/Plugins/Experimental/Mover/Source/Mover/Public/DefaultMovementSet/Modes/WalkingMode.h
+++ b/Engine/Plugins/Experimental/Mover/Source/Mover/Public/DefaultMovementSet/Modes/WalkingMode.h
@@ -28,10 +28,17 @@ public:
 	UFUNCTION(BlueprintCallable, Category=Mover)
 	virtual void OnSimulationTick(const FSimulationTickParams& Params, FMoverTickEndData& OutputState) override;
 
+	// Returns the active turn generator. Note: you will need to cast the return value to the generator you expect to get, it can also be none
+	UFUNCTION(BlueprintPure, Category=Mover)
+	UObject* GetTurnGenerator();
+
+	// Sets the active turn generator to use the class provided. Note: To set it back to the default implementation pass in none
+	UFUNCTION(BlueprintCallable, Category=Mover)
+	void SetTurnGeneratorClass(UPARAM(meta=(MustImplement="/Script/Mover.TurnGeneratorInterface", AllowAbstract="false")) TSubclassOf<UObject> TurnGeneratorClass);
 
 protected:
 	/** Optional modular object for generating rotation towards desired orientation. If not specified, linear interpolation will be used. */
-	UPROPERTY(BlueprintReadWrite, EditAnywhere, Instanced, Category=Mover, meta=(MustImplement="/Script/Mover.TurnGeneratorInterface"))
+	UPROPERTY(EditAnywhere, Instanced, Category=Mover, meta=(MustImplement="/Script/Mover.TurnGeneratorInterface"))
 	TObjectPtr<UObject> TurnGenerator;
 
 	virtual void OnRegistered(const FName ModeName) override; 
diff --git a/Engine/Plugins/Experimental/Mover/Source/Mover/Public/MoveLibrary/ModularMovement.h b/Engine/Plugins/Experimental/Mover/Source/Mover/Public/MoveLibrary/ModularMovement.h
--- a/Engine/Plugins/Experimental/Mover/Source/Mover/Public/MoveLibrary/ModularMovement.h
+++ b/Engine/Plugins/Experimental/Mover/Source/Mover/Public/MoveLibrary/ModularMovement.h
@@ -37,7 +37,7 @@ public:
 /**
  * Modular mechanism for turning a Mover actor in place using linear angular velocity.
  */
-UCLASS(EditInlineNew)
+UCLASS(BlueprintType, EditInlineNew)
 class MOVER_API ULinearTurnGenerator : public UObject, public ITurnGeneratorInterface
 {
 	GENERATED_BODY()
@@ -64,7 +64,7 @@ public:
 /**
  * Modular mechanism for turning a Mover actor in place using an exact damped spring
  */
-UCLASS(EditInlineNew)
+UCLASS(BlueprintType, EditInlineNew)
 class MOVER_API UExactDampedTurnGenerator : public UObject, public ITurnGeneratorInterface
 {
 	GENERATED_BODY()
