Mover 2.0 splitting off character specific functions into a new character mover component. This change also fixes up any BP errors from the change

#rnx

#virtualized

diff --git a/Engine/Plugins/Experimental/Mover/Source/Mover/Private/DefaultMovementSet/CharacterMoverComponent.cpp b/Engine/Plugins/Experimental/Mover/Source/Mover/Private/DefaultMovementSet/CharacterMoverComponent.cpp
new file mode 100644
--- /dev/null
+++ b/Engine/Plugins/Experimental/Mover/Source/Mover/Private/DefaultMovementSet/CharacterMoverComponent.cpp
@@ -0,0 +1,63 @@
+﻿// Copyright Epic Games, Inc. All Rights Reserved.
+
+#include "DefaultMovementSet/CharacterMoverComponent.h"
+
+#include "DefaultMovementSet/Modes/FallingMode.h"
+#include "DefaultMovementSet/Modes/FlyingMode.h"
+#include "DefaultMovementSet/Modes/WalkingMode.h"
+#include "MoveLibrary/FloorQueryUtils.h"
+
+UCharacterMoverComponent::UCharacterMoverComponent()
+{
+	// Default movement modes
+	MovementModes.Add(DefaultModeNames::Walking, CreateDefaultSubobject<UWalkingMode>(TEXT("DefaultWalkingMode")));
+	MovementModes.Add(DefaultModeNames::Falling, CreateDefaultSubobject<UFallingMode>(TEXT("DefaultFallingMode")));
+	MovementModes.Add(DefaultModeNames::Flying,  CreateDefaultSubobject<UFlyingMode>(TEXT("DefaultFlyingMode")));
+
+	StartingMovementMode = DefaultModeNames::Falling;
+}
+
+bool UCharacterMoverComponent::IsFalling() const
+{
+	if (bHasValidCachedState)
+	{
+		return CachedLastSyncState.MovementMode == DefaultModeNames::Falling;
+	}
+
+	return false;
+}
+
+bool UCharacterMoverComponent::IsAirborne() const
+{
+	if (bHasValidCachedState)
+	{
+		return CachedLastSyncState.MovementMode == DefaultModeNames::Flying || CachedLastSyncState.MovementMode == DefaultModeNames::Falling;
+	}
+
+	return false;
+}
+
+bool UCharacterMoverComponent::IsOnGround() const
+{
+	if (bHasValidCachedState)
+	{
+		return CachedLastSyncState.MovementMode == DefaultModeNames::Walking;
+	}
+
+	return false;
+}
+
+bool UCharacterMoverComponent::IsSlopeSliding() const
+{
+	if (IsAirborne())
+	{
+		FFloorCheckResult HitResult;
+		const UMoverBlackboard* MoverBlackboard = GetSimBlackboard();
+		if (MoverBlackboard->TryGet(CommonBlackboard::LastFloorResult, HitResult))
+		{
+			return HitResult.bBlockingHit && !HitResult.bWalkableFloor;
+		}
+	}
+
+	return false;
+}
diff --git a/Engine/Plugins/Experimental/Mover/Source/Mover/Private/MoverComponent.cpp b/Engine/Plugins/Experimental/Mover/Source/Mover/Private/MoverComponent.cpp
--- a/Engine/Plugins/Experimental/Mover/Source/Mover/Private/MoverComponent.cpp
+++ b/Engine/Plugins/Experimental/Mover/Source/Mover/Private/MoverComponent.cpp
@@ -54,13 +54,6 @@ UMoverComponent::UMoverComponent()
 	bWantsInitializeComponent = true;
 	bAutoActivate = true;
 
-	// Default movement modes
-	MovementModes.Add(DefaultModeNames::Walking, CreateDefaultSubobject<UWalkingMode>(TEXT("DefaultWalkingMode")));
-	MovementModes.Add(DefaultModeNames::Falling, CreateDefaultSubobject<UFallingMode>(TEXT("DefaultFallingMode")));
-	MovementModes.Add(DefaultModeNames::Flying,  CreateDefaultSubobject<UFlyingMode>(TEXT("DefaultFlyingMode")));
-
-	StartingMovementMode = DefaultModeNames::Falling;
-
 	PersistentSyncStateDataTypes.Add(FMoverDataPersistence(FMoverDefaultSyncState::StaticStruct(), true));
 
 	BackendClass = UMoverNetworkPredictionLiaisonComponent::StaticClass();
@@ -1142,50 +1135,6 @@ FVector UMoverComponent::GetUpDirection() const
 }
 
 
-bool UMoverComponent::IsFalling() const
-{ 
-	if (bHasValidCachedState)
-	{
-		return CachedLastSyncState.MovementMode == DefaultModeNames::Falling;
-	}
-
-	return false;
-}
-
-bool UMoverComponent::IsAirborne() const
-{
-	if (bHasValidCachedState)
-	{
-		return CachedLastSyncState.MovementMode == DefaultModeNames::Flying || CachedLastSyncState.MovementMode == DefaultModeNames::Falling;
-	}
-
-	return false;
-}
-
-bool UMoverComponent::IsOnGround() const
-{
-	if (bHasValidCachedState)
-	{
-		return CachedLastSyncState.MovementMode == DefaultModeNames::Walking;
-	}
-
-	return false;
-}
-
-bool UMoverComponent::IsSlopeSliding() const
-{
-	if (IsAirborne())
-	{
-		FFloorCheckResult HitResult;
-		if (SimBlackboard->TryGet(CommonBlackboard::LastFloorResult, HitResult))
-		{
-			return HitResult.bBlockingHit && !HitResult.bWalkableFloor;
-		}
-	}
-
-	return false;
-}
-
 TArray<FTrajectorySampleInfo> UMoverComponent::GetFutureTrajectory(float FutureSeconds, float SamplesPerSecond) const
 {
 	FMoverTickStartData StartingState;
diff --git a/Engine/Plugins/Experimental/Mover/Source/Mover/Public/DefaultMovementSet/CharacterMoverComponent.h b/Engine/Plugins/Experimental/Mover/Source/Mover/Public/DefaultMovementSet/CharacterMoverComponent.h
new file mode 100644
--- /dev/null
+++ b/Engine/Plugins/Experimental/Mover/Source/Mover/Public/DefaultMovementSet/CharacterMoverComponent.h
@@ -0,0 +1,31 @@
+﻿// Copyright Epic Games, Inc. All Rights Reserved.
+
+#pragma once
+
+#include "MoverComponent.h"
+#include "CharacterMoverComponent.generated.h"
+
+UCLASS(BlueprintType, meta = (BlueprintSpawnableComponent))
+class MOVER_API UCharacterMoverComponent : public UMoverComponent
+{
+	GENERATED_BODY()
+	
+public:
+	UCharacterMoverComponent();
+	
+	// Is this actor in a falling state? Note that this includes upwards motion induced by jumping.
+	UFUNCTION(BlueprintPure, Category = Mover)
+	virtual bool IsFalling() const;
+
+	// Is this actor in a airborne state? (e.g. Flying, Falling)
+	UFUNCTION(BlueprintPure, Category = Mover)
+	virtual bool IsAirborne() const;
+
+	// Is this actor in a grounded state? (e.g. Walking)
+	UFUNCTION(BlueprintPure, Category = Mover)
+	virtual bool IsOnGround() const;
+
+	// Is this actor sliding on an unwalkable slope
+	UFUNCTION(BlueprintPure, Category = Mover)
+	virtual bool IsSlopeSliding() const;
+};
diff --git a/Engine/Plugins/Experimental/Mover/Source/Mover/Public/MoverComponent.h b/Engine/Plugins/Experimental/Mover/Source/Mover/Public/MoverComponent.h
--- a/Engine/Plugins/Experimental/Mover/Source/Mover/Public/MoverComponent.h
+++ b/Engine/Plugins/Experimental/Mover/Source/Mover/Public/MoverComponent.h
@@ -201,22 +201,6 @@ public:	// Queries
 	UFUNCTION(BlueprintPure, Category = Mover)
 	FRotator GetTargetOrientation() const;
 
-	// Is this actor in a falling state? Note that this includes upwards motion induced by jumping.
-	UFUNCTION(BlueprintPure, Category = Mover)
-	virtual bool IsFalling() const;
-
-	// Is this actor in a airborne state? (e.g. Flying, Falling)
-	UFUNCTION(BlueprintPure, Category = Mover)
-	virtual bool IsAirborne() const;
-
-	// Is this actor in a grounded state? (e.g. Walking)
-	UFUNCTION(BlueprintPure, Category = Mover)
-	virtual bool IsOnGround() const;
-
-	// Is this actor sliding on an unwalkable slope
-	UFUNCTION(BlueprintPure, Category = Mover)
-	virtual bool IsSlopeSliding() const;
-
 	/** Get a sampling of where the actor is projected to be in the future, based on a current state. Note that this is projecting ideal movement without doing full simulation and collision. */
 	UFUNCTION(BlueprintCallable, BlueprintPure = false, Category = Mover)
 	TArray<FTrajectorySampleInfo> GetFutureTrajectory(float FutureSeconds, float SamplesPerSecond) const;
diff --git a/Engine/Plugins/Experimental/MoverExamples/Source/MoverExamples/Private/CharacterVariants/Ziplining/ZipliningTransitions.cpp b/Engine/Plugins/Experimental/MoverExamples/Source/MoverExamples/Private/CharacterVariants/Ziplining/ZipliningTransitions.cpp
--- a/Engine/Plugins/Experimental/MoverExamples/Source/MoverExamples/Private/CharacterVariants/Ziplining/ZipliningTransitions.cpp
+++ b/Engine/Plugins/Experimental/MoverExamples/Source/MoverExamples/Private/CharacterVariants/Ziplining/ZipliningTransitions.cpp
@@ -3,7 +3,7 @@
 #include "CharacterVariants/Ziplining/ZipliningTransitions.h"
 #include "CharacterVariants/AbilityInputs.h"
 #include "CharacterVariants/Ziplining/ZiplineInterface.h"
-#include "MoverComponent.h"
+#include "DefaultMovementSet/CharacterMoverComponent.h"
 #include "GameFramework/Actor.h"
 #include "Kismet/KismetSystemLibrary.h"
 
@@ -19,7 +19,7 @@ FTransitionEvalResult UZiplineStartTransition::OnEvaluate(const FSimulationTickP
 {
 	FTransitionEvalResult EvalResult = FTransitionEvalResult::NoTransition;
 
-	UMoverComponent* MoverComp = Params.MoverComponent;
+	UCharacterMoverComponent* MoverComp = Cast<UCharacterMoverComponent>(Params.MoverComponent);
 
 	const FMoverSyncState& SyncState = Params.StartState.SyncState;
 
diff --git a/Engine/Plugins/Experimental/MoverExamples/Source/MoverExamples/Private/MoverExamplesCharacter.cpp b/Engine/Plugins/Experimental/MoverExamples/Source/MoverExamples/Private/MoverExamplesCharacter.cpp
--- a/Engine/Plugins/Experimental/MoverExamples/Source/MoverExamples/Private/MoverExamplesCharacter.cpp
+++ b/Engine/Plugins/Experimental/MoverExamples/Source/MoverExamples/Private/MoverExamplesCharacter.cpp
@@ -5,7 +5,7 @@
 #include "Components/InputComponent.h"
 #include "Engine/BlueprintGeneratedClass.h"
 #include "Engine/LocalPlayer.h"
-#include "MoverComponent.h"
+#include "DefaultMovementSet/CharacterMoverComponent.h"
 #include "MoveLibrary/BasedMovementUtils.h"
 #include "CharacterVariants/AbilityInputs.h"
 #include "GameFramework/SpringArmComponent.h"
@@ -21,7 +21,7 @@ static const FName Name_CharacterMotionComponent(TEXT("MoverComponent"));
 AMoverExamplesCharacter::AMoverExamplesCharacter(const FObjectInitializer& ObjectInitializer)
 	: Super(ObjectInitializer)
 {
-	CharacterMotionComponent = CreateDefaultSubobject<UMoverComponent>(Name_CharacterMotionComponent);
+	CharacterMotionComponent = CreateDefaultSubobject<UCharacterMoverComponent>(Name_CharacterMotionComponent);
 	ensure(CharacterMotionComponent);
 
 	PrimaryActorTick.bCanEverTick = true;
@@ -155,7 +155,7 @@ void AMoverExamplesCharacter::OnProduceInput(float DeltaMs, FMoverInputCmdContex
 	{
 		FRotator Rotator = CharacterInputs.ControlRotation;
 		FVector FinalDirectionalIntent;
-		if (const UMoverComponent* MoverComp = GetMoverComponent())
+		if (const UCharacterMoverComponent* MoverComp = GetMoverComponent())
 		{
 			if (MoverComp->IsOnGround() || MoverComp->IsFalling())
 			{
@@ -235,7 +235,7 @@ void AMoverExamplesCharacter::OnProduceInput(float DeltaMs, FMoverInputCmdContex
 
 	if (bUseBaseRelativeMovement)
 	{
-		if (const UMoverComponent* MoverComp = GetComponentByClass<UMoverComponent>())
+		if (const UCharacterMoverComponent* MoverComp = GetComponentByClass<UCharacterMoverComponent>())
 		{
 			if (UPrimitiveComponent* MovementBase = MoverComp->GetMovementBase())
 			{
diff --git a/Engine/Plugins/Experimental/MoverExamples/Source/MoverExamples/Public/MoverExamplesCharacter.h b/Engine/Plugins/Experimental/MoverExamples/Source/MoverExamples/Public/MoverExamplesCharacter.h
--- a/Engine/Plugins/Experimental/MoverExamples/Source/MoverExamples/Public/MoverExamplesCharacter.h
+++ b/Engine/Plugins/Experimental/MoverExamples/Source/MoverExamples/Public/MoverExamplesCharacter.h
@@ -8,6 +8,7 @@
 #include "MoverExamplesCharacter.generated.h"
 
 class UInputAction;
+class UCharacterMoverComponent;
 struct FInputActionValue;
 
 UCLASS()
@@ -31,7 +32,7 @@ public:
 
 	// Accessor for the actor's movement component
 	UFUNCTION(BlueprintPure, Category = Mover)
-	UMoverComponent* GetMoverComponent() const { return CharacterMotionComponent; }
+	UCharacterMoverComponent* GetMoverComponent() const { return CharacterMotionComponent; }
 
 	// Request the character starts moving with an intended directional magnitude. A length of 1 indicates maximum acceleration.
 	UFUNCTION(BlueprintCallable, Category=MoverExamples)
@@ -94,7 +95,7 @@ public:
 
 protected:
 	UPROPERTY(Category = Movement, VisibleAnywhere, BlueprintReadOnly, meta = (AllowPrivateAccess = "true"))
-	TObjectPtr<UMoverComponent> CharacterMotionComponent;
+	TObjectPtr<UCharacterMoverComponent> CharacterMotionComponent;
 
 private:
 	FVector LastAffirmativeMoveInput = FVector::ZeroVector;	// Movement input (intent or velocity) the last time we had one that wasn't zero
