Mover 2.0 adding more context to cached data warnings

#rnx

diff --git a/Engine/Plugins/Experimental/Mover/Source/Mover/Private/MoverComponent.cpp b/Engine/Plugins/Experimental/Mover/Source/Mover/Private/MoverComponent.cpp
--- a/Engine/Plugins/Experimental/Mover/Source/Mover/Private/MoverComponent.cpp
+++ b/Engine/Plugins/Experimental/Mover/Source/Mover/Private/MoverComponent.cpp
@@ -1252,7 +1252,7 @@ const FMoverSyncState& UMoverComponent::GetSyncState() const
 {
 	if (!bHasValidCachedState)
 	{
-		UE_LOG(LogMover, Warning, TEXT("Attempting direct access to the last-cached sync state before one has been set. Results will be unreliable."));
+		UE_LOG(LogMover, Warning, TEXT("Attempting direct access to the last-cached sync state before one has been set. Results will be unreliable. Use the HasValidCachedState function to check if CachedLastSyncState is valid or not."));
 	}
 
 	return CachedLastSyncState;
@@ -1288,7 +1288,7 @@ const FMoverInputCmdContext& UMoverComponent::GetLastInputCmd() const
 {
 	if (!bHasValidCachedUsedInput)
 	{
-		UE_LOG(LogMover, Warning, TEXT("Attempting direct access to the last-cached used input cmd before one has been set. Results will be unreliable."));
+		UE_LOG(LogMover, Warning, TEXT("Attempting direct access to the last-cached used input cmd before one has been set. Results will be unreliable. Use the HasValidCachedInputCmd function to check if CachedLastUsedInputCmd is valid or not."));
 	}
 
 	return CachedLastUsedInputCmd;
