Mover 2.0 making error messages for incorrect use of FindLayeredMove more helpful.

diff --git a/Engine/Plugins/Experimental/Mover/Source/Mover/Private/MoverComponent.cpp b/Engine/Plugins/Experimental/Mover/Source/Mover/Private/MoverComponent.cpp
--- a/Engine/Plugins/Experimental/Mover/Source/Mover/Private/MoverComponent.cpp
+++ b/Engine/Plugins/Experimental/Mover/Source/Mover/Private/MoverComponent.cpp
@@ -889,11 +889,20 @@ DEFINE_FUNCTION(UMoverComponent::execK2_FindActiveLayeredMove)
 
 	DidSucceed = false;
 	
-	if (!MovePtr || !StructProp)
+	if (!MovePtr)
 	{
 		FBlueprintExceptionInfo ExceptionInfo(
 			EBlueprintExceptionType::AbortExecution,
-			LOCTEXT("MoverComponent_GetActiveLayeredMove_UnresolvedTarget", "Failed to resolve the TargetAsRawBytes for GetActiveLayeredMove")
+			LOCTEXT("MoverComponent_GetActiveLayeredMove_UnresolvedTarget", "Failed to resolve the OutLayeredMove for GetActiveLayeredMove")
+		);
+
+		FBlueprintCoreDelegates::ThrowScriptException(P_THIS, Stack, ExceptionInfo);
+	}
+	else if (!StructProp)
+	{
+		FBlueprintExceptionInfo ExceptionInfo(
+			EBlueprintExceptionType::AbortExecution,
+			LOCTEXT("MoverComponent_GetActiveLayeredMove_UnresolvedTarget", "GetActiveLayeredMove: Target for OutLayeredMove is not a valid type. It must be a Struct and a child of FLayeredMoveBase.")
 		);
 
 		FBlueprintCoreDelegates::ThrowScriptException(P_THIS, Stack, ExceptionInfo);
@@ -902,7 +911,7 @@ DEFINE_FUNCTION(UMoverComponent::execK2_FindActiveLayeredMove)
 	{
 		FBlueprintExceptionInfo ExceptionInfo(
 			EBlueprintExceptionType::AbortExecution,
-			LOCTEXT("MoverComponent_GetActiveLayeredMove_BadType", "TargetAsRawBytes is not a valid type. Must be a child of FLayeredMoveBase.")
+			LOCTEXT("MoverComponent_GetActiveLayeredMove_BadType", "GetActiveLayeredMove: Target for OutLayeredMove is not a valid type. Must be a child of FLayeredMoveBase.")
 		);
 
 		FBlueprintCoreDelegates::ThrowScriptException(P_THIS, Stack, ExceptionInfo);
