Mover 2.0 fixing LOCTEXT conflict with GetActiveLayeredMove error message

diff --git a/Engine/Plugins/Experimental/Mover/Source/Mover/Private/MoverComponent.cpp b/Engine/Plugins/Experimental/Mover/Source/Mover/Private/MoverComponent.cpp
--- a/Engine/Plugins/Experimental/Mover/Source/Mover/Private/MoverComponent.cpp
+++ b/Engine/Plugins/Experimental/Mover/Source/Mover/Private/MoverComponent.cpp
@@ -902,7 +902,7 @@ DEFINE_FUNCTION(UMoverComponent::execK2_FindActiveLayeredMove)
 	{
 		FBlueprintExceptionInfo ExceptionInfo(
 			EBlueprintExceptionType::AbortExecution,
-			LOCTEXT("MoverComponent_GetActiveLayeredMove_UnresolvedTarget", "GetActiveLayeredMove: Target for OutLayeredMove is not a valid type. It must be a Struct and a child of FLayeredMoveBase.")
+			LOCTEXT("MoverComponent_GetActiveLayeredMove_TargetNotStruct", "GetActiveLayeredMove: Target for OutLayeredMove is not a valid type. It must be a Struct and a child of FLayeredMoveBase.")
 		);
 
 		FBlueprintCoreDelegates::ThrowScriptException(P_THIS, Stack, ExceptionInfo);
