[Mover] Bug Fix - Rotating Platforms Jitter
Prevent usage of GT State in PT

diff --git a/Engine/Plugins/Experimental/Mover/Source/Mover/Private/PhysicsMover/PhysicsMovementUtils.cpp b/Engine/Plugins/Experimental/Mover/Source/Mover/Private/PhysicsMover/PhysicsMovementUtils.cpp
--- a/Engine/Plugins/Experimental/Mover/Source/Mover/Private/PhysicsMover/PhysicsMovementUtils.cpp
+++ b/Engine/Plugins/Experimental/Mover/Source/Mover/Private/PhysicsMover/PhysicsMovementUtils.cpp
@@ -121,18 +121,7 @@ FVector UPhysicsMovementUtils::ComputeGroundVelocityFromHitResult(const FVector&
 	{
 		FVector Offset = CharacterPosition - Rigid->GetX();
 		Offset -= Offset.ProjectOnToNormal(FloorHit.ImpactNormal);
-
-		if (Rigid->KinematicTarget().IsSet())
-		{
-			const FVector LinearDisplacement = Rigid->KinematicTarget().GetTargetPosition() - Rigid->GetX();
-			const FQuat RelativeQuat = Rigid->GetR().Inverse() * Rigid->KinematicTarget().GetTargetRotation();
-			const FVector AngularDisplacement = RelativeQuat.ToRotationVector();
-			GroundVelocity = (LinearDisplacement + AngularDisplacement.Cross(Offset)) / DeltaSeconds;
-		}
-		else
-		{
-			GroundVelocity = Rigid->GetV() + Rigid->GetW().Cross(Offset);
-		}
+		GroundVelocity = Rigid->GetV() + Rigid->GetW().Cross(Offset);
 	}
 	return GroundVelocity;
 }
@@ -144,18 +133,7 @@ FVector UPhysicsMovementUtils::ComputeIntegratedGroundVelocityFromHitResult(cons
 	{
 		FVector Offset = CharacterPosition - Rigid->GetX();
 		Offset -= Offset.ProjectOnToNormal(FloorHit.ImpactNormal);
-
-		if (Rigid->KinematicTarget().IsSet())
-		{
-			const FVector LinearDisplacement = Rigid->KinematicTarget().GetTargetPosition() - Rigid->GetX();
-			const FQuat RelativeQuat = Rigid->GetR().Inverse() * Rigid->KinematicTarget().GetTargetRotation();
-			const FVector AngularDisplacement = RelativeQuat.ToRotationVector();
-			GroundVelocity = (LinearDisplacement + AngularDisplacement.Cross(Offset)) / DeltaSeconds;
-		}
-		else
-		{
-			GroundVelocity = Rigid->GetV() + Rigid->GetW().Cross(Offset);
-		}
+		GroundVelocity = Rigid->GetV() + Rigid->GetW().Cross(Offset);
 
 		if (Rigid->IsDynamic() && Rigid->GravityEnabled())
 		{
