Physics Mover
- Use internal acceleration structure for ground queries
- Ensure lifting off surface is only true if moving upwards
- Reduce fallback query radius

diff --git a/Engine/Plugins/Experimental/Mover/Source/Mover/Private/Backends/MoverNetworkPhysicsLiaison.cpp b/Engine/Plugins/Experimental/Mover/Source/Mover/Private/Backends/MoverNetworkPhysicsLiaison.cpp
--- a/Engine/Plugins/Experimental/Mover/Source/Mover/Private/Backends/MoverNetworkPhysicsLiaison.cpp
+++ b/Engine/Plugins/Experimental/Mover/Source/Mover/Private/Backends/MoverNetworkPhysicsLiaison.cpp
@@ -876,7 +876,6 @@ void UMoverNetworkPhysicsLiaisonComponent::OnPreSimulate_Internal(const FPhysics
 			if (Blackboard->TryGet(CommonBlackboard::LastFloorResult, LastFloorResult))
 			{
 				LocalGroundVelocity = UPhysicsMovementUtils::ComputeGroundVelocityFromHitResult(CharacterParticle->GetX(), LastFloorResult.HitResult, TickParams.DeltaTimeSeconds);
-				LocalGroundVelocity -= LocalGroundVelocity.ProjectOnToNormal(LastFloorResult.HitResult.ImpactNormal);
 			}
 		}
 	}
diff --git a/Engine/Plugins/Experimental/Mover/Source/Mover/Private/PhysicsMover/Modes/PhysicsDrivenFallingMode.cpp b/Engine/Plugins/Experimental/Mover/Source/Mover/Private/PhysicsMover/Modes/PhysicsDrivenFallingMode.cpp
--- a/Engine/Plugins/Experimental/Mover/Source/Mover/Private/PhysicsMover/Modes/PhysicsDrivenFallingMode.cpp
+++ b/Engine/Plugins/Experimental/Mover/Source/Mover/Private/PhysicsMover/Modes/PhysicsDrivenFallingMode.cpp
@@ -92,11 +92,11 @@ void UPhysicsDrivenFallingMode::OnSimulationTick(const FSimulationTickParams& Pa
 	float PawnRadius;
 	UpdatedPrimitive->CalcBoundingCylinder(PawnRadius, PawnHalfHeight);
 
-	const float QueryDistance = FMath::Max(1.1f * TargetHeight, TargetHeight - UpDir.Dot(ProposedMove.LinearVelocity) * DeltaSeconds);
+	const float QueryDistance = 1.1f * FMath::Max(TargetHeight, TargetHeight - UpDir.Dot(ProposedMove.LinearVelocity) * DeltaSeconds);
 	const float ShrinkRadius = 5.0f; // TODO - Make this a user setting
 	const float QueryRadius = FMath::Max(PawnRadius - ShrinkRadius, 0.0f);
 
-	UPhysicsMovementUtils::FloorSweep(StartingSyncState->GetLocation_WorldSpace(), StartingSyncState->GetVelocity_WorldSpace() * DeltaSeconds,
+	UPhysicsMovementUtils::FloorSweep_Internal(StartingSyncState->GetLocation_WorldSpace(), StartingSyncState->GetVelocity_WorldSpace() * DeltaSeconds,
 		UpdatedPrimitive, UpDir, QueryRadius, QueryDistance, CommonLegacySettings->MaxWalkSlopeCosine, TargetHeight, FloorResult, WaterResult);
 
 	SimBlackboard->Set(CommonBlackboard::LastFloorResult, FloorResult);
@@ -132,7 +132,7 @@ void UPhysicsDrivenFallingMode::OnSimulationTick(const FSimulationTickParams& Pa
 
 	constexpr float FloorDistanceTolerance = 2.0f;
 	const FVector ProjectedGroundVelocity = UPhysicsMovementUtils::ComputeIntegratedGroundVelocityFromHitResult(StartingSyncState->GetLocation_WorldSpace(), FloorResult.HitResult, DeltaSeconds);
-	const float ProjectedRelativeVerticalVelocity = FloorResult.HitResult.ImpactNormal.Dot(ProposedMove.LinearVelocity - ProjectedGroundVelocity);
+	const float ProjectedRelativeVerticalVelocity = UpDir.Dot(ProposedMove.LinearVelocity - ProjectedGroundVelocity);
 	const float ProjectedFloorDistance = FloorResult.FloorDist + ProjectedRelativeVerticalVelocity * DeltaSeconds;
 	const bool bIsFloorWithinReach = ProjectedFloorDistance < TargetHeight + FloorDistanceTolerance;
 	const bool bIsMovingUpRelativeToFloor = ProjectedRelativeVerticalVelocity > UE_KINDA_SMALL_NUMBER;
@@ -140,7 +140,7 @@ void UPhysicsDrivenFallingMode::OnSimulationTick(const FSimulationTickParams& Pa
 	if (FloorResult.IsWalkableFloor() && bIsFloorWithinReach && !bIsMovingUpRelativeToFloor)
 	{
 		OutputState.MovementEndState.NextModeName = DefaultModeNames::Walking;
-		TargetPos -= UpDir * (FloorResult.FloorDist - TargetHeight);
+		TargetPos -= (UpDir.Dot(TargetVel * DeltaSeconds) + (FloorResult.FloorDist - TargetHeight)) * UpDir;
 	}
 	else
 	{
diff --git a/Engine/Plugins/Experimental/Mover/Source/Mover/Private/PhysicsMover/Modes/PhysicsDrivenSwimmingMode.cpp b/Engine/Plugins/Experimental/Mover/Source/Mover/Private/PhysicsMover/Modes/PhysicsDrivenSwimmingMode.cpp
--- a/Engine/Plugins/Experimental/Mover/Source/Mover/Private/PhysicsMover/Modes/PhysicsDrivenSwimmingMode.cpp
+++ b/Engine/Plugins/Experimental/Mover/Source/Mover/Private/PhysicsMover/Modes/PhysicsDrivenSwimmingMode.cpp
@@ -97,7 +97,7 @@ void UPhysicsDrivenSwimmingMode::OnSimulationTick(const FSimulationTickParams& P
 
 	const float QueryDistance= 2.0f * PawnHalfHeight;
 	
-	UPhysicsMovementUtils::FloorSweep(StartingSyncState->GetLocation_WorldSpace(), StartingSyncState->GetVelocity_WorldSpace() * DeltaSeconds,
+	UPhysicsMovementUtils::FloorSweep_Internal(StartingSyncState->GetLocation_WorldSpace(), StartingSyncState->GetVelocity_WorldSpace() * DeltaSeconds,
 		UpdatedPrimitive, UpDir, PawnRadius, QueryDistance, CommonLegacySettings->MaxWalkSlopeCosine, TargetHeight, FloorResult, WaterResult);
 
 	SimBlackboard->Set(CommonBlackboard::LastFloorResult, FloorResult);
diff --git a/Engine/Plugins/Experimental/Mover/Source/Mover/Private/PhysicsMover/Modes/PhysicsDrivenWalkingMode.cpp b/Engine/Plugins/Experimental/Mover/Source/Mover/Private/PhysicsMover/Modes/PhysicsDrivenWalkingMode.cpp
--- a/Engine/Plugins/Experimental/Mover/Source/Mover/Private/PhysicsMover/Modes/PhysicsDrivenWalkingMode.cpp
+++ b/Engine/Plugins/Experimental/Mover/Source/Mover/Private/PhysicsMover/Modes/PhysicsDrivenWalkingMode.cpp
@@ -138,7 +138,7 @@ void UPhysicsDrivenWalkingMode::FloorCheck(const FMoverDefaultSyncState& SyncSta
 	const float FloorSweepDistance = TargetHeight + CommonLegacySettings->MaxStepHeight;
 	const float ShrinkRadius = 1.0f;
 	const float QueryRadius = FMath::Max(PawnRadius - ShrinkRadius, 0.0f);
-	UPhysicsMovementUtils::FloorSweep(SyncState.GetLocation_WorldSpace(), DeltaPos, UpdatedPrimitive, UpDir,
+	UPhysicsMovementUtils::FloorSweep_Internal(SyncState.GetLocation_WorldSpace(), DeltaPos, UpdatedPrimitive, UpDir,
 		QueryRadius, FloorSweepDistance, CommonLegacySettings->MaxWalkSlopeCosine, TargetHeight, OutFloorResult, OutWaterResult);
 
 	if (!OutFloorResult.bBlockingHit)
@@ -163,7 +163,7 @@ void UPhysicsDrivenWalkingMode::FloorCheck(const FMoverDefaultSyncState& SyncSta
 	{
 		// Collision should prevent movement. Just try to find ground at start of movement
 		const float ShrinkMultiplier = 0.75f;
-		UPhysicsMovementUtils::FloorSweep(SyncState.GetLocation_WorldSpace(), DeltaPos, UpdatedPrimitive, UpDir,
+		UPhysicsMovementUtils::FloorSweep_Internal(SyncState.GetLocation_WorldSpace(), FVector::ZeroVector, UpdatedPrimitive, UpDir,
 			ShrinkMultiplier * QueryRadius, FloorSweepDistance, CommonLegacySettings->MaxWalkSlopeCosine, TargetHeight, OutFloorResult, OutWaterResult);
 
 		OutFloorResult.bWalkableFloor = OutFloorResult.bWalkableFloor && CanStepUpOnHitSurface(OutFloorResult);
@@ -212,11 +212,11 @@ void UPhysicsDrivenWalkingMode::FloorCheck(const FMoverDefaultSyncState& SyncSta
 		}
 		else
 		{
-			NewQueryRadius = 0.75f * QueryRadius;
+			NewQueryRadius = 0.25f * QueryRadius;
 			NewDeltaPos = DeltaPos - DP * HorizSurfaceDir;
 		}
 
-		UPhysicsMovementUtils::FloorSweep(SyncState.GetLocation_WorldSpace(), NewDeltaPos, UpdatedPrimitive, UpDir,
+		UPhysicsMovementUtils::FloorSweep_Internal(SyncState.GetLocation_WorldSpace(), NewDeltaPos, UpdatedPrimitive, UpDir,
 			NewQueryRadius, FloorSweepDistance, CommonLegacySettings->MaxWalkSlopeCosine, TargetHeight, OutFloorResult, OutWaterResult);
 
 		OutFloorResult.bWalkableFloor = OutFloorResult.bWalkableFloor && CanStepUpOnHitSurface(OutFloorResult);
@@ -234,9 +234,9 @@ void UPhysicsDrivenWalkingMode::FloorCheck(const FMoverDefaultSyncState& SyncSta
 		// Try a query at the start of the movement to find a walkable surface and prevent movement
 
 		NewDeltaPos = FVector::ZeroVector;
-		NewQueryRadius = 0.75f * QueryRadius;
+		NewQueryRadius = 0.25f * QueryRadius;
 
-		UPhysicsMovementUtils::FloorSweep(SyncState.GetLocation_WorldSpace(), NewDeltaPos, UpdatedPrimitive, UpDir,
+		UPhysicsMovementUtils::FloorSweep_Internal(SyncState.GetLocation_WorldSpace(), NewDeltaPos, UpdatedPrimitive, UpDir,
 			NewQueryRadius, FloorSweepDistance, CommonLegacySettings->MaxWalkSlopeCosine, TargetHeight, OutFloorResult, OutWaterResult);
 
 		OutFloorResult.bWalkableFloor = OutFloorResult.bWalkableFloor && CanStepUpOnHitSurface(OutFloorResult);
@@ -321,7 +321,7 @@ void UPhysicsDrivenWalkingMode::OnSimulationTick(const FSimulationTickParams& Pa
 		const FVector StartGroundVelocity = UPhysicsMovementUtils::ComputeGroundVelocityFromHitResult(StartingSyncState->GetLocation_WorldSpace(), FloorResult.HitResult, DeltaSeconds);
 		
 		FVector TargetVelocity = StartingSyncState->GetVelocity_WorldSpace();
-		FVector TargetPosition = StartingSyncState->GetLocation_WorldSpace();
+		FVector TargetPosition = StartingSyncState->GetLocation_WorldSpace() - UpDir * (FloorResult.FloorDist - TargetHeight);
 		if (FloorResult.bWalkableFloor)
 		{
 			const FVector ProposedMovePlaneVelocity = ProposedMove.LinearVelocity - ProposedMove.LinearVelocity.ProjectOnToNormal(PrevGroundNormal);
@@ -352,7 +352,7 @@ void UPhysicsDrivenWalkingMode::OnSimulationTick(const FSimulationTickParams& Pa
 
 		const float ProjectedRelativeVerticalVelocity = FloorResult.HitResult.ImpactNormal.Dot(ProjectedVelocity - ProjectedGroundVelocity);
 		const float VerticalVelocityLimit = 2.0f / DeltaSeconds;
-		if (ProjectedRelativeVerticalVelocity > VerticalVelocityLimit)
+		if ((ProjectedRelativeVerticalVelocity > VerticalVelocityLimit) && (ProjectedVelocity.Z > VerticalVelocityLimit))
 		{
 			bIsLiftingOffSurface = true;
 		}
diff --git a/Engine/Plugins/Experimental/Mover/Source/Mover/Private/PhysicsMover/PhysicsMovementUtils.cpp b/Engine/Plugins/Experimental/Mover/Source/Mover/Private/PhysicsMover/PhysicsMovementUtils.cpp
--- a/Engine/Plugins/Experimental/Mover/Source/Mover/Private/PhysicsMover/PhysicsMovementUtils.cpp
+++ b/Engine/Plugins/Experimental/Mover/Source/Mover/Private/PhysicsMover/PhysicsMovementUtils.cpp
@@ -5,10 +5,12 @@
 #include "Chaos/PhysicsObjectInternalInterface.h"
 #include "Components/PrimitiveComponent.h"
 #include "Engine/World.h"
+#include "Framework/Threading.h"
 #include "GameFramework/PhysicsVolume.h"
 #include "MoveLibrary/FloorQueryUtils.h"
 #include "MoveLibrary/GroundMovementUtils.h"
 #include "MoveLibrary/WaterMovementUtils.h"
+#include "Physics/GenericPhysicsInterface.h"
 #include "PhysicsMover/PhysicsMoverSimulationTypes.h"
 #include "WaterBodyActor.h"
 
@@ -18,7 +20,7 @@
 
 extern FPhysicsDrivenMotionDebugParams GPhysicsDrivenMotionDebugParams;
 
-void UPhysicsMovementUtils::FloorSweep(
+void UPhysicsMovementUtils::FloorSweep_Internal(
 	const FVector& Location,
 	const FVector& DeltaPos,
 	const UPrimitiveComponent* UpdatedPrimitive,
@@ -31,6 +33,8 @@ void UPhysicsMovementUtils::FloorSweep(
 	FWaterCheckResult& OutWaterResult
 )
 {
+	Chaos::EnsureIsInPhysicsThreadContext();
+
 	if (const UWorld* World = UpdatedPrimitive->GetWorld())
 	{
 		FCollisionQueryParams QueryParams(SCENE_QUERY_STAT(PhysicsFloorTest), false, UpdatedPrimitive->GetOwner());
@@ -55,7 +59,7 @@ void UPhysicsMovementUtils::FloorSweep(
 		FVector Start = Location + DeltaPosHoriz + (QueryRadius + UE_KINDA_SMALL_NUMBER) * UpDir;
 		FVector End = Start - AdjustedQueryDistance * UpDir;
 		FHitResult OutHit;
-		if (World->SweepMultiByChannel(Hits, Start, End, FQuat::Identity, CollisionChannel, FCollisionShape::MakeSphere(QueryRadius), QueryParams, ResponseParams))
+		if (Chaos::Private::FGenericPhysicsInterface_Internal::SpherecastMulti(World, QueryRadius, Hits, Start, End, CollisionChannel, QueryParams, ResponseParams))
 		{
 			OutHit = Hits.Last();
 		}
@@ -144,6 +148,7 @@ FVector UPhysicsMovementUtils::ComputeIntegratedGroundVelocityFromHitResult(cons
 					GroundVelocity += PhysVolume->GetGravityZ() * FVector::UpVector * DeltaSeconds;
 				}
 			}
+
 		}
 	}
 	return GroundVelocity;
diff --git a/Engine/Plugins/Experimental/Mover/Source/Mover/Public/PhysicsMover/PhysicsMovementUtils.h b/Engine/Plugins/Experimental/Mover/Source/Mover/Public/PhysicsMover/PhysicsMovementUtils.h
--- a/Engine/Plugins/Experimental/Mover/Source/Mover/Public/PhysicsMover/PhysicsMovementUtils.h
+++ b/Engine/Plugins/Experimental/Mover/Source/Mover/Public/PhysicsMover/PhysicsMovementUtils.h
@@ -17,7 +17,7 @@ class UPrimitiveComponent;
 class MOVER_API UPhysicsMovementUtils
 {
 public:
-	static void FloorSweep(const FVector& Location, const FVector& DeltaPos, const UPrimitiveComponent* UpdatedPrimitive, const FVector& UpDir,
+	static void FloorSweep_Internal(const FVector& Location, const FVector& DeltaPos, const UPrimitiveComponent* UpdatedPrimitive, const FVector& UpDir,
 		float QueryRadius, float QueryDistance, float MaxWalkSlopeCosine, float TargetHeight, FFloorCheckResult& OutFloorResult, FWaterCheckResult& OutWaterResult);
 
 	// If the hit result hit something, return the particle handle
