Physics Mover - Fix bug with previous checkin that was causing the character to slide up slopes

diff --git a/Engine/Plugins/Experimental/Mover/Source/Mover/Private/PhysicsMover/Modes/PhysicsDrivenWalkingMode.cpp b/Engine/Plugins/Experimental/Mover/Source/Mover/Private/PhysicsMover/Modes/PhysicsDrivenWalkingMode.cpp
--- a/Engine/Plugins/Experimental/Mover/Source/Mover/Private/PhysicsMover/Modes/PhysicsDrivenWalkingMode.cpp
+++ b/Engine/Plugins/Experimental/Mover/Source/Mover/Private/PhysicsMover/Modes/PhysicsDrivenWalkingMode.cpp
@@ -321,7 +321,7 @@ void UPhysicsDrivenWalkingMode::OnSimulationTick(const FSimulationTickParams& Pa
 		const FVector StartGroundVelocity = UPhysicsMovementUtils::ComputeGroundVelocityFromHitResult(StartingSyncState->GetLocation_WorldSpace(), FloorResult.HitResult, DeltaSeconds);
 		
 		FVector TargetVelocity = StartingSyncState->GetVelocity_WorldSpace();
-		FVector TargetPosition = StartingSyncState->GetLocation_WorldSpace() - UpDir * (FloorResult.FloorDist - TargetHeight);
+		FVector TargetPosition = StartingSyncState->GetLocation_WorldSpace();
 		if (FloorResult.bWalkableFloor)
 		{
 			const FVector ProposedMovePlaneVelocity = ProposedMove.LinearVelocity - ProposedMove.LinearVelocity.ProjectOnToNormal(PrevGroundNormal);
