[Mover] - Match Falling Mode's ShrinkRadius size to Walking Mode's

diff --git a/Engine/Plugins/Experimental/Mover/Source/Mover/Private/PhysicsMover/Modes/PhysicsDrivenFallingMode.cpp b/Engine/Plugins/Experimental/Mover/Source/Mover/Private/PhysicsMover/Modes/PhysicsDrivenFallingMode.cpp
--- a/Engine/Plugins/Experimental/Mover/Source/Mover/Private/PhysicsMover/Modes/PhysicsDrivenFallingMode.cpp
+++ b/Engine/Plugins/Experimental/Mover/Source/Mover/Private/PhysicsMover/Modes/PhysicsDrivenFallingMode.cpp
@@ -93,7 +93,7 @@ void UPhysicsDrivenFallingMode::OnSimulationTick(const FSimulationTickParams& Pa
 	UpdatedPrimitive->CalcBoundingCylinder(PawnRadius, PawnHalfHeight);
 
 	const float QueryDistance = 1.1f * FMath::Max(TargetHeight, TargetHeight - UpDir.Dot(ProposedMove.LinearVelocity) * DeltaSeconds);
-	const float ShrinkRadius = 5.0f; // TODO - Make this a user setting
+	const float ShrinkRadius = 1.0f; // TODO - Make this a user setting
 	const float QueryRadius = FMath::Max(PawnRadius - ShrinkRadius, 0.0f);
 
 	UPhysicsMovementUtils::FloorSweep_Internal(StartingSyncState->GetLocation_WorldSpace(), StartingSyncState->GetVelocity_WorldSpace() * DeltaSeconds,
