Mover 2.0 Adding support for instantaneous layered moves as Instant Movement Effects

diff --git a/Engine/Plugins/Experimental/Mover/Source/Mover/Private/DefaultMovementSet/InstantMovementEffects/BasicInstantMovementEffects.cpp b/Engine/Plugins/Experimental/Mover/Source/Mover/Private/DefaultMovementSet/InstantMovementEffects/BasicInstantMovementEffects.cpp
new file mode 100644
--- /dev/null
+++ b/Engine/Plugins/Experimental/Mover/Source/Mover/Private/DefaultMovementSet/InstantMovementEffects/BasicInstantMovementEffects.cpp
@@ -0,0 +1,233 @@
+﻿// Copyright Epic Games, Inc. All Rights Reserved.
+
+
+#include "DefaultMovementSet/InstantMovementEffects/BasicInstantMovementEffects.h"
+
+#include "MoverComponent.h"
+#include "MoverDataModelTypes.h"
+#include "MoverSimulationTypes.h"
+#include "DefaultMovementSet/Settings/CommonLegacyMovementSettings.h"
+#include "MoveLibrary/MovementUtils.h"
+#include "MoveLibrary/MoverBlackboard.h"
+
+
+// -------------------------------------------------------------------
+// FTeleportEffect
+// -------------------------------------------------------------------
+
+FTeleportEffect::FTeleportEffect()
+	: TargetLocation(FVector::ZeroVector)
+{
+}
+
+bool FTeleportEffect::ApplyMovementEffect(FApplyMovementEffectParams& ApplyEffectParams, FMoverSyncState& OutputState)
+{
+	if (ApplyEffectParams.UpdatedComponent->GetOwner()->TeleportTo(TargetLocation, ApplyEffectParams.UpdatedComponent->GetComponentRotation()))
+	{
+		FMoverDefaultSyncState& OutputSyncState = OutputState.SyncStateCollection.FindOrAddMutableDataByType<FMoverDefaultSyncState>();
+		
+		if (const FMoverDefaultSyncState* StartingSyncState = ApplyEffectParams.StartState->SyncState.SyncStateCollection.FindDataByType<FMoverDefaultSyncState>())
+		{
+			OutputSyncState.SetTransforms_WorldSpace( ApplyEffectParams.UpdatedComponent->GetComponentLocation(),
+													  ApplyEffectParams.UpdatedComponent->GetComponentRotation(),
+													  OutputSyncState.GetVelocity_WorldSpace(),
+													  nullptr ); // no movement base
+		
+			// TODO: instead of invalidating it, consider checking for a floor. Possibly a dynamic base?
+			if (UMoverBlackboard* SimBlackboard = ApplyEffectParams.MoverComp->GetSimBlackboard_Mutable())
+			{
+				SimBlackboard->Invalidate(CommonBlackboard::LastFloorResult);
+				SimBlackboard->Invalidate(CommonBlackboard::LastFoundDynamicMovementBase);
+			}
+
+			return true;
+		}
+	}
+
+	return false;
+}
+
+FInstantMovementEffect* FTeleportEffect::Clone() const
+{
+	FTeleportEffect* CopyPtr = new FTeleportEffect(*this);
+	return CopyPtr;
+}
+
+void FTeleportEffect::NetSerialize(FArchive& Ar)
+{
+	Super::NetSerialize(Ar);
+
+	Ar << TargetLocation;
+}
+
+UScriptStruct* FTeleportEffect::GetScriptStruct() const
+{
+	return FTeleportEffect::StaticStruct();
+}
+
+FString FTeleportEffect::ToSimpleString() const
+{
+	return FString::Printf(TEXT("Teleport"));
+}
+
+void FTeleportEffect::AddReferencedObjects(FReferenceCollector& Collector)
+{
+	Super::AddReferencedObjects(Collector);
+}
+
+// -------------------------------------------------------------------
+// FJumpImpulseEffect
+// -------------------------------------------------------------------
+
+FJumpImpulseEffect::FJumpImpulseEffect()
+	: UpwardsSpeed(0.f)
+{
+}
+
+bool FJumpImpulseEffect::ApplyMovementEffect(FApplyMovementEffectParams& ApplyEffectParams, FMoverSyncState& OutputState)
+{
+	if (const FMoverDefaultSyncState* SyncState = ApplyEffectParams.StartState->SyncState.SyncStateCollection.FindDataByType<FMoverDefaultSyncState>())
+	{
+		FMoverDefaultSyncState& OutputSyncState = OutputState.SyncStateCollection.FindOrAddMutableDataByType<FMoverDefaultSyncState>();
+		
+		const FVector UpDir = ApplyEffectParams.MoverComp->GetUpDirection();
+		const FVector ImpulseVelocity = UpDir * UpwardsSpeed;
+	
+		// Jump impulse overrides vertical velocity while maintaining the rest
+		const FVector PriorVelocityWS = SyncState->GetVelocity_WorldSpace();
+		const FVector StartingNonUpwardsVelocity = PriorVelocityWS - PriorVelocityWS.ProjectOnToNormal(UpDir);
+
+		if (const UCommonLegacyMovementSettings* CommonSettings = ApplyEffectParams.MoverComp->FindSharedSettings<UCommonLegacyMovementSettings>())
+		{
+			OutputState.MovementMode = CommonSettings->AirMovementModeName;
+		}
+		
+		FRelativeBaseInfo MovementBaseInfo;
+		if (const UMoverBlackboard* SimBlackboard = ApplyEffectParams.MoverComp->GetSimBlackboard())
+		{
+			SimBlackboard->TryGet(CommonBlackboard::LastFoundDynamicMovementBase, MovementBaseInfo);
+		}
+
+		const FVector FinalVelocity = StartingNonUpwardsVelocity + ImpulseVelocity;
+		OutputSyncState.SetTransforms_WorldSpace( ApplyEffectParams.UpdatedComponent->GetComponentLocation(),
+												  ApplyEffectParams.UpdatedComponent->GetComponentRotation(),
+												  FinalVelocity,
+												  MovementBaseInfo.MovementBase.Get(),
+												  MovementBaseInfo.BoneName);
+		
+		ApplyEffectParams.UpdatedComponent->ComponentVelocity = FinalVelocity;
+		
+		return true;
+	}
+
+	return false;
+}
+
+FInstantMovementEffect* FJumpImpulseEffect::Clone() const
+{
+	FJumpImpulseEffect* CopyPtr = new FJumpImpulseEffect(*this);
+	return CopyPtr;
+}
+
+void FJumpImpulseEffect::NetSerialize(FArchive& Ar)
+{
+	Super::NetSerialize(Ar);
+
+	Ar << UpwardsSpeed;
+}
+
+UScriptStruct* FJumpImpulseEffect::GetScriptStruct() const
+{
+	return FJumpImpulseEffect::StaticStruct();
+}
+
+FString FJumpImpulseEffect::ToSimpleString() const
+{
+	return FString::Printf(TEXT("JumpImpulse"));
+}
+
+void FJumpImpulseEffect::AddReferencedObjects(FReferenceCollector& Collector)
+{
+	Super::AddReferencedObjects(Collector);
+}
+
+// -------------------------------------------------------------------
+// FApplyVelocityEffect
+// -------------------------------------------------------------------
+
+FApplyVelocityEffect::FApplyVelocityEffect()
+	: VelocityToApply(FVector::ZeroVector)
+	, bAdditiveVelocity(false)
+	, ForceMovementMode(NAME_None)
+{
+}
+
+bool FApplyVelocityEffect::ApplyMovementEffect(FApplyMovementEffectParams& ApplyEffectParams, FMoverSyncState& OutputState)
+{
+	FMoverDefaultSyncState& OutputSyncState = OutputState.SyncStateCollection.FindOrAddMutableDataByType<FMoverDefaultSyncState>();
+	
+	OutputState.MovementMode = ForceMovementMode;
+	
+	FRelativeBaseInfo MovementBaseInfo;
+	if (const UMoverBlackboard* SimBlackboard = ApplyEffectParams.MoverComp->GetSimBlackboard())
+	{
+		SimBlackboard->TryGet(CommonBlackboard::LastFoundDynamicMovementBase, MovementBaseInfo);
+	}
+
+	FVector Velocity = VelocityToApply;
+	if (bAdditiveVelocity)
+	{
+		if (const FMoverDefaultSyncState* SyncState = ApplyEffectParams.StartState->SyncState.SyncStateCollection.FindDataByType<FMoverDefaultSyncState>())
+		{
+			Velocity += SyncState->GetVelocity_WorldSpace();
+		}
+	}
+
+	OutputSyncState.SetTransforms_WorldSpace( ApplyEffectParams.UpdatedComponent->GetComponentLocation(),
+											  ApplyEffectParams.UpdatedComponent->GetComponentRotation(),
+											  Velocity,
+											  MovementBaseInfo.MovementBase.Get(),
+											  MovementBaseInfo.BoneName);
+
+	ApplyEffectParams.UpdatedComponent->ComponentVelocity = Velocity;
+	
+	return true;
+}
+
+FInstantMovementEffect* FApplyVelocityEffect::Clone() const
+{
+	FApplyVelocityEffect* CopyPtr = new FApplyVelocityEffect(*this);
+	return CopyPtr;
+}
+
+void FApplyVelocityEffect::NetSerialize(FArchive& Ar)
+{
+	Super::NetSerialize(Ar);
+
+	SerializePackedVector<10, 16>(VelocityToApply, Ar);
+
+	Ar << bAdditiveVelocity;
+	
+	bool bUsingForcedMovementMode = !ForceMovementMode.IsNone();
+	Ar.SerializeBits(&bUsingForcedMovementMode, 1);
+
+	if (bUsingForcedMovementMode)
+	{
+		Ar << ForceMovementMode;
+	}
+}
+
+UScriptStruct* FApplyVelocityEffect::GetScriptStruct() const
+{
+	return FApplyVelocityEffect::StaticStruct();
+}
+
+FString FApplyVelocityEffect::ToSimpleString() const
+{
+	return FString::Printf(TEXT("ApplyVelocity"));
+}
+
+void FApplyVelocityEffect::AddReferencedObjects(FReferenceCollector& Collector)
+{
+	Super::AddReferencedObjects(Collector);
+}
diff --git a/Engine/Plugins/Experimental/Mover/Source/Mover/Private/InstantMovementEffect.cpp b/Engine/Plugins/Experimental/Mover/Source/Mover/Private/InstantMovementEffect.cpp
new file mode 100644
--- /dev/null
+++ b/Engine/Plugins/Experimental/Mover/Source/Mover/Private/InstantMovementEffect.cpp
@@ -0,0 +1,25 @@
+﻿// Copyright Epic Games, Inc. All Rights Reserved.
+
+#include "InstantMovementEffect.h"
+
+FInstantMovementEffect* FInstantMovementEffect::Clone() const
+{
+	// If child classes don't override this, saved moves will not work
+	checkf(false, TEXT("FInstantMovementEffect::Clone() being called erroneously. This should always be overridden in child classes!"));
+	return nullptr;
+}
+
+void FInstantMovementEffect::NetSerialize(FArchive& Ar)
+{
+	
+}
+
+UScriptStruct* FInstantMovementEffect::GetScriptStruct() const
+{
+	return FInstantMovementEffect::StaticStruct();
+}
+
+FString FInstantMovementEffect::ToSimpleString() const
+{
+	return GetScriptStruct()->GetName();
+}
diff --git a/Engine/Plugins/Experimental/Mover/Source/Mover/Private/MovementModeStateMachine.cpp b/Engine/Plugins/Experimental/Mover/Source/Mover/Private/MovementModeStateMachine.cpp
--- a/Engine/Plugins/Experimental/Mover/Source/Mover/Private/MovementModeStateMachine.cpp
+++ b/Engine/Plugins/Experimental/Mover/Source/Mover/Private/MovementModeStateMachine.cpp
@@ -161,6 +161,33 @@ void UMovementModeStateMachine::OnSimulationTick(USceneComponent* UpdatedCompone
 		// Transfer any queued moves into the starting state. They'll be started during the move generation.
 		FlushQueuedMovesToGroup(SubstepStartData.SyncState.LayeredMoves);
 		OutputState.SyncState.LayeredMoves = SubstepStartData.SyncState.LayeredMoves;
+
+		FApplyMovementEffectParams EffectParams;
+		EffectParams.MoverComp = MoverComp;
+		EffectParams.StartState = &SubstepStartData;
+		EffectParams.TimeStep = &SubTimeStep;
+		EffectParams.UpdatedComponent = UpdatedComponent;
+		EffectParams.UpdatedPrimitive = UpdatedPrimitive;
+		
+		// Apply any instant effects that were queued up between ticks
+		if (ApplyInstantEffects(EffectParams, OutputState.SyncState))
+		{
+			// Copying over our sync state collection to SubstepStartData so it is effectively the input sync state later for the movement mode. Doing this makes sure state modification from Instant Effects isn't overridden later by the movement mode
+			for (auto SyncDataIt = OutputState.SyncState.SyncStateCollection.GetCollectionDataIterator(); SyncDataIt; ++SyncDataIt)
+			{
+				if (SyncDataIt->Get())
+				{
+					SubstepStartData.SyncState.SyncStateCollection.AddOrOverwriteData(TSharedPtr<FMoverDataStructBase>(SyncDataIt->Get()->Clone()));
+				}
+			}
+
+			if (CurrentModeName != OutputState.SyncState.MovementMode)
+			{
+				SetModeImmediately(OutputState.SyncState.MovementMode);
+				SubstepStartData.SyncState.MovementMode = CurrentModeName;
+			}
+		}
+
 		FLayeredMoveGroup& CurrentLayeredMoves = OutputState.SyncState.LayeredMoves;
 
 		// Gather any layered move contributions
@@ -304,13 +331,22 @@ void UMovementModeStateMachine::OnSimulationTick(USceneComponent* UpdatedCompone
 		SubstepStartData.AuxState  = OutputState.AuxState;
 	}
 
+	FApplyMovementEffectParams EffectParams;
+	EffectParams.MoverComp = MoverComp;
+	EffectParams.StartState = &SubstepStartData;
+	EffectParams.TimeStep = &SubTimeStep;
+	EffectParams.UpdatedComponent = UpdatedComponent;
+	EffectParams.UpdatedPrimitive = UpdatedPrimitive;
+	
+	// Apply any instant effects that were queued up during this tick and didn't get handled in a substep
+	ApplyInstantEffects(EffectParams, OutputState.SyncState);
 }
 
-
 void UMovementModeStateMachine::OnSimulationRollback(const FMoverSyncState* SyncState, const FMoverAuxStateContext* AuxState)
 {
 	ClearQueuedMode();
 	QueuedLayeredMoves.Empty();
+	QueuedInstantEffects.Empty();
 }
 
 
@@ -339,6 +375,11 @@ void UMovementModeStateMachine::QueueLayeredMove(TSharedPtr<FLayeredMoveBase> Mo
 	QueuedLayeredMoves.Add(Move);
 }
 
+void UMovementModeStateMachine::QueueInstantMovementEffect(TSharedPtr<FInstantMovementEffect> Effect)
+{
+	QueuedInstantEffects.Add(Effect);
+}
+
 
 void UMovementModeStateMachine::ConstructDefaultModes()
 {
@@ -386,6 +427,26 @@ void UMovementModeStateMachine::FlushQueuedMovesToGroup(FLayeredMoveGroup& Group
 	}
 }
 
+bool UMovementModeStateMachine::ApplyInstantEffects(FApplyMovementEffectParams& ApplyEffectParams, FMoverSyncState& OutputState)
+{
+	bool bInstantMovementEffectApplied = false;
+	
+	if (!QueuedInstantEffects.IsEmpty())
+	{
+		for (TSharedPtr<FInstantMovementEffect>& QueuedEffect : QueuedInstantEffects)
+		{
+			if (QueuedEffect->ApplyMovementEffect(ApplyEffectParams, OutputState))
+			{
+				bInstantMovementEffectApplied = true;
+			}
+		}
+		
+		QueuedInstantEffects.Empty();
+	}
+	
+	return bInstantMovementEffectApplied;
+}
+
 AActor* UMovementModeStateMachine::GetOwnerActor() const
 {
 	if (const UActorComponent* OwnerMoverComp = Cast<UActorComponent>(GetOuter()))
diff --git a/Engine/Plugins/Experimental/Mover/Source/Mover/Private/MovementModeStateMachine.h b/Engine/Plugins/Experimental/Mover/Source/Mover/Private/MovementModeStateMachine.h
--- a/Engine/Plugins/Experimental/Mover/Source/Mover/Private/MovementModeStateMachine.h
+++ b/Engine/Plugins/Experimental/Mover/Source/Mover/Private/MovementModeStateMachine.h
@@ -5,6 +5,7 @@
 #include "CoreMinimal.h"
 #include "Templates/SubclassOf.h"
 #include "MovementMode.h"
+#include "InstantMovementEffect.h"
 #include "MovementModeStateMachine.generated.h"
 
 struct FProposedMove;
@@ -59,6 +60,8 @@ public:
 	const UBaseMovementMode* FindMovementMode(FName ModeName) const;
 
 	void QueueLayeredMove(TSharedPtr<FLayeredMoveBase> Move);
+	
+ 	void QueueInstantMovementEffect(TSharedPtr<FInstantMovementEffect> Effect);
 
 protected:
 
@@ -76,10 +79,14 @@ protected:
 	/** Moves that are queued to be added to the simulation at the start of the next sim subtick */
 	TArray<TSharedPtr<FLayeredMoveBase>> QueuedLayeredMoves;
 
+ 	/** Effects that are queued to be applied to the simulation at the start of the next sim subtick or at the end of this tick */
+ 	TArray<TSharedPtr<FInstantMovementEffect>> QueuedInstantEffects;
+ 	
 private:
 	void ConstructDefaultModes();
 	void AdvanceToNextMode();
 	void FlushQueuedMovesToGroup(FLayeredMoveGroup& Group);
+ 	bool ApplyInstantEffects(FApplyMovementEffectParams& ApplyEffectParams, FMoverSyncState& OutputState);
 	AActor* GetOwnerActor() const;
 };
 
diff --git a/Engine/Plugins/Experimental/Mover/Source/Mover/Private/MoverComponent.cpp b/Engine/Plugins/Experimental/Mover/Source/Mover/Private/MoverComponent.cpp
--- a/Engine/Plugins/Experimental/Mover/Source/Mover/Private/MoverComponent.cpp
+++ b/Engine/Plugins/Experimental/Mover/Source/Mover/Private/MoverComponent.cpp
@@ -11,6 +11,7 @@
 #include "MoveLibrary/MovementUtils.h"
 #include "MoveLibrary/FloorQueryUtils.h"
 #include "MoverLog.h"
+#include "InstantMovementEffect.h"
 #include "Backends/MoverNetworkPredictionLiaison.h"
 #include "Components/MeshComponent.h"
 #include "Components/PrimitiveComponent.h"
@@ -868,6 +869,43 @@ void UMoverComponent::QueueLayeredMove(TSharedPtr<FLayeredMoveBase> LayeredMove)
 	ModeFSM->QueueLayeredMove(LayeredMove);
 }
 
+void UMoverComponent::K2_QueueInstantMovementEffect(const int32& EffectAsRawData)
+{
+	// This will never be called, the exec version below will be hit instead
+	checkNoEntry();
+}
+
+DEFINE_FUNCTION(UMoverComponent::execK2_QueueInstantMovementEffect)
+{
+	Stack.StepCompiledIn<FStructProperty>(nullptr);
+	void* EffectPtr = Stack.MostRecentPropertyAddress;
+	FStructProperty* StructProp = CastField<FStructProperty>(Stack.MostRecentProperty);
+
+	P_FINISH;
+
+	P_NATIVE_BEGIN;
+
+	const bool bHasValidStructProp = StructProp && StructProp->Struct && StructProp->Struct->IsChildOf(FInstantMovementEffect::StaticStruct());
+
+	if (ensureMsgf((bHasValidStructProp && EffectPtr), TEXT("An invalid type (%s) was sent to a QueueInstantMovementEffect node. A struct derived from FInstantMovementEffect is required. No Movement Effect will be queued."),
+		StructProp ? *GetNameSafe(StructProp->Struct) : *Stack.MostRecentProperty->GetClass()->GetName()))
+	{
+		// Could we steal this instead of cloning? (move semantics)
+		FInstantMovementEffect* EffectAsBasePtr = reinterpret_cast<FInstantMovementEffect*>(EffectPtr);
+		FInstantMovementEffect* ClonedMove = EffectAsBasePtr->Clone();
+
+		P_THIS->QueueInstantMovementEffect(TSharedPtr<FInstantMovementEffect>(ClonedMove));
+	}
+
+	P_NATIVE_END;
+}
+
+
+void UMoverComponent::QueueInstantMovementEffect(TSharedPtr<FInstantMovementEffect> InstantMovementEffect)
+{	
+	ModeFSM->QueueInstantMovementEffect(InstantMovementEffect);
+}
+
 void UMoverComponent::K2_FindActiveLayeredMove(bool& DidSucceed, int32& TargetAsRawBytes) const
 {
 	// This will never be called, the exec version below will be hit instead
diff --git a/Engine/Plugins/Experimental/Mover/Source/Mover/Private/MoverTypes.cpp b/Engine/Plugins/Experimental/Mover/Source/Mover/Private/MoverTypes.cpp
--- a/Engine/Plugins/Experimental/Mover/Source/Mover/Private/MoverTypes.cpp
+++ b/Engine/Plugins/Experimental/Mover/Source/Mover/Private/MoverTypes.cpp
@@ -198,6 +198,11 @@ void FMoverDataCollection::ToString(FAnsiStringBuilderBase& Out) const
 	}
 }
 
+TArray<TSharedPtr<FMoverDataStructBase>>::TConstIterator FMoverDataCollection::GetCollectionDataIterator() const
+{
+	return DataArray.CreateConstIterator();
+}
+
 struct FMoverDataDeleter
 {
 	FORCEINLINE void operator()(FMoverDataStructBase* Object) const
diff --git a/Engine/Plugins/Experimental/Mover/Source/Mover/Public/DefaultMovementSet/InstantMovementEffects/BasicInstantMovementEffects.h b/Engine/Plugins/Experimental/Mover/Source/Mover/Public/DefaultMovementSet/InstantMovementEffects/BasicInstantMovementEffects.h
new file mode 100644
--- /dev/null
+++ b/Engine/Plugins/Experimental/Mover/Source/Mover/Public/DefaultMovementSet/InstantMovementEffects/BasicInstantMovementEffects.h
@@ -0,0 +1,98 @@
+﻿// Copyright Epic Games, Inc. All Rights Reserved.
+
+
+#pragma once
+
+#include "InstantMovementEffect.h"
+#include "BasicInstantMovementEffects.generated.h"
+
+/** Teleport: instantly moves an actor to a new location */
+USTRUCT(BlueprintType)
+struct MOVER_API FTeleportEffect : public FInstantMovementEffect
+{
+	GENERATED_USTRUCT_BODY()
+
+	FTeleportEffect();
+	virtual ~FTeleportEffect() {}
+
+	// Location to teleport to, in world space
+	UPROPERTY(EditAnywhere, BlueprintReadWrite, Category = Mover)
+	FVector TargetLocation;
+
+	virtual bool ApplyMovementEffect(FApplyMovementEffectParams& ApplyEffectParams, FMoverSyncState& OutputState) override;
+
+	virtual FInstantMovementEffect* Clone() const override;
+
+	virtual void NetSerialize(FArchive& Ar) override;
+
+	virtual UScriptStruct* GetScriptStruct() const override;
+
+	virtual FString ToSimpleString() const override;
+
+	virtual void AddReferencedObjects(class FReferenceCollector& Collector) override;
+};
+
+/** Jump Impulse: introduces an instantaneous upwards change in velocity. This overrides the existing 'up' component of the actor's current velocity
+  * Note: this only applies the impulse for one tick!
+  */
+USTRUCT(BlueprintType)
+struct MOVER_API FJumpImpulseEffect : public FInstantMovementEffect
+{
+	GENERATED_USTRUCT_BODY()
+
+	FJumpImpulseEffect();
+
+	virtual ~FJumpImpulseEffect() {}
+
+	// Units per second, in whatever direction the target actor considers 'up'
+	UPROPERTY(EditAnywhere, BlueprintReadWrite, Category = Mover, meta = (ClampMin = "0", UIMin = "0", ForceUnits = "cm/s"))
+	float UpwardsSpeed;
+
+	virtual bool ApplyMovementEffect(FApplyMovementEffectParams& ApplyEffectParams, FMoverSyncState& OutputState) override;
+
+	virtual FInstantMovementEffect* Clone() const override;
+
+	virtual void NetSerialize(FArchive& Ar) override;
+
+	virtual UScriptStruct* GetScriptStruct() const override;
+
+	virtual FString ToSimpleString() const override;
+
+	virtual void AddReferencedObjects(class FReferenceCollector& Collector) override;
+};
+
+/** Apply Velocity: provides an impulse velocity to the actor after (optionally) forcing them into a particular movement mode
+  * Note: this only applies the impulse for one tick!
+  */
+USTRUCT(BlueprintType)
+struct MOVER_API FApplyVelocityEffect : public FInstantMovementEffect
+{
+	GENERATED_USTRUCT_BODY()
+
+	FApplyVelocityEffect();
+	virtual ~FApplyVelocityEffect() {}
+
+	// Velocity to apply to the actor.
+	UPROPERTY(EditAnywhere, BlueprintReadWrite, Category = Mover, meta=(ForceUnits="cm/s"))
+	FVector VelocityToApply;
+
+	// If true VelocityToApply will be added to current velocity on this actor. If false velocity will be set directly to VelocityToApply
+	UPROPERTY(EditAnywhere, BlueprintReadWrite, Category = Mover)
+	bool bAdditiveVelocity;
+	
+	// Optional movement mode name to force the actor into before applying the impulse velocity.
+	UPROPERTY(EditAnywhere, BlueprintReadWrite, Category = Mover)
+	FName ForceMovementMode;
+
+	virtual bool ApplyMovementEffect(FApplyMovementEffectParams& ApplyEffectParams, FMoverSyncState& OutputState) override;
+	
+	virtual FInstantMovementEffect* Clone() const override;
+
+	virtual void NetSerialize(FArchive& Ar) override;
+
+	virtual UScriptStruct* GetScriptStruct() const override;
+
+	virtual FString ToSimpleString() const override;
+
+	virtual void AddReferencedObjects(class FReferenceCollector& Collector) override;
+};
diff --git a/Engine/Plugins/Experimental/Mover/Source/Mover/Public/InstantMovementEffect.h b/Engine/Plugins/Experimental/Mover/Source/Mover/Public/InstantMovementEffect.h
new file mode 100644
--- /dev/null
+++ b/Engine/Plugins/Experimental/Mover/Source/Mover/Public/InstantMovementEffect.h
@@ -0,0 +1,63 @@
+﻿// Copyright Epic Games, Inc. All Rights Reserved.
+
+#pragma once
+
+#include "MoveLibrary/MovementUtilsTypes.h"
+#include "InstantMovementEffect.generated.h"
+
+class UMoverComponent;
+struct FMoverTimeStep;
+struct FMoverTickStartData;
+struct FMoverSyncState;
+
+struct FApplyMovementEffectParams
+{
+	USceneComponent* UpdatedComponent;
+
+	UPrimitiveComponent* UpdatedPrimitive;
+
+	const UMoverComponent* MoverComp;
+
+	const FMoverTickStartData* StartState;
+
+	const FMoverTimeStep* TimeStep;
+};
+
+/**
+ * Instant Movement Effects are methods of affecting movement state directly on a Mover-based actor for one tick.
+ * Note: This is only applied one tick and then removed
+ * Common uses would be for Teleporting, Changing Movement Modes directly, one time force application, etc.
+ * Multiple Instant Movement Effects can be active at the time
+ */
+USTRUCT(BlueprintType)
+struct MOVER_API FInstantMovementEffect
+{
+	GENERATED_USTRUCT_BODY()
+
+	FInstantMovementEffect() { }
+
+	virtual ~FInstantMovementEffect() { }
+	
+	// @return newly allocated copy of this FInstantMovementEffect. Must be overridden by child classes
+	virtual FInstantMovementEffect* Clone() const;
+
+	virtual void NetSerialize(FArchive& Ar);
+
+	virtual UScriptStruct* GetScriptStruct() const;
+
+	virtual FString ToSimpleString() const;
+
+	virtual void AddReferencedObjects(class FReferenceCollector& Collector) {}
+
+	virtual bool ApplyMovementEffect(FApplyMovementEffectParams& ApplyEffectParams, FMoverSyncState& OutputState) { return false; }
+};
+
+template<>
+struct TStructOpsTypeTraits< FInstantMovementEffect > : public TStructOpsTypeTraitsBase2< FInstantMovementEffect >
+{
+	enum
+	{
+		//WithNetSerializer = true,
+		WithCopy = true
+	};
+};
diff --git a/Engine/Plugins/Experimental/Mover/Source/Mover/Public/MoverComponent.h b/Engine/Plugins/Experimental/Mover/Source/Mover/Public/MoverComponent.h
--- a/Engine/Plugins/Experimental/Mover/Source/Mover/Public/MoverComponent.h
+++ b/Engine/Plugins/Experimental/Mover/Source/Mover/Public/MoverComponent.h
@@ -18,6 +18,7 @@
 #include "MoverComponent.generated.h"
 
 struct FMoverTimeStep;
+struct FInstantMovementEffect;
 class UMovementModeStateMachine;
 class UMovementMixer;
 
@@ -145,6 +146,17 @@ public:
 	// Queue a layered move to start during the next simulation frame
 	void QueueLayeredMove(TSharedPtr<FLayeredMoveBase> Move);
 	
+	/**
+	 * Queue a Instant Movement Effect to start at the end of this frame or start of the next subtick - whichever happens first. This will clone whatever move you pass in, so you'll need to fully set it up before queuing.
+	 * @param InstantMovementEffect			The effect to queue, which must be a FInstantMovementEffect sub-type. 
+	 */
+	UFUNCTION(BlueprintCallable, CustomThunk, Category = Mover, meta = (CustomStructureParam = "EffectAsRawData", AllowAbstract = "false", DisplayName = "Queue Instant Movement Effect"))
+	void K2_QueueInstantMovementEffect(UPARAM(DisplayName="Instant Movement Effect") const int32& EffectAsRawData);
+	DECLARE_FUNCTION(execK2_QueueInstantMovementEffect);
+
+	// Queue a Instant Movement Effect to take place at the end of this frame or start of the next subtick - whichever happens first
+	void QueueInstantMovementEffect(TSharedPtr<FInstantMovementEffect> Move);
+	
 	// Queue a movement mode change to occur during the next simulation frame. If bShouldReenter is true, then a mode change will occur even if already in that mode.
 	UFUNCTION(BlueprintCallable, Category = Mover)
 	void QueueNextMode(FName DesiredModeName, bool bShouldReenter=false);
diff --git a/Engine/Plugins/Experimental/Mover/Source/Mover/Public/MoverTypes.h b/Engine/Plugins/Experimental/Mover/Source/Mover/Public/MoverTypes.h
--- a/Engine/Plugins/Experimental/Mover/Source/Mover/Public/MoverTypes.h
+++ b/Engine/Plugins/Experimental/Mover/Source/Mover/Public/MoverTypes.h
@@ -165,6 +165,9 @@ struct MOVER_API FMoverDataCollection
 	/** Get string representation of all elements in this collection */
 	void ToString(FAnsiStringBuilderBase& Out) const;
 
+	/** Const access to data array of collections */
+	TArray<TSharedPtr<FMoverDataStructBase>>::TConstIterator GetCollectionDataIterator() const;
+	
 	/** Find data of a specific type in the collection (mutable version). If not found, null will be returned. */
 	template <typename T>
 	T* FindMutableDataByType() const
