[Mover] - Pawn jumping on sloped ground
Offer the legacy landing logic to handle vertical landing velocity

diff --git a/Engine/Plugins/Experimental/Mover/Source/Mover/Private/PhysicsMover/Modes/PhysicsDrivenWalkingMode.cpp b/Engine/Plugins/Experimental/Mover/Source/Mover/Private/PhysicsMover/Modes/PhysicsDrivenWalkingMode.cpp
--- a/Engine/Plugins/Experimental/Mover/Source/Mover/Private/PhysicsMover/Modes/PhysicsDrivenWalkingMode.cpp
+++ b/Engine/Plugins/Experimental/Mover/Source/Mover/Private/PhysicsMover/Modes/PhysicsDrivenWalkingMode.cpp
@@ -322,7 +322,12 @@ void UPhysicsDrivenWalkingMode::OnSimulationTick(const FSimulationTickParams& Pa
 		
 		FVector TargetVelocity = StartingSyncState->GetVelocity_WorldSpace();
 		FVector TargetPosition = StartingSyncState->GetLocation_WorldSpace();
-		if (FloorResult.bWalkableFloor)
+		if (bHandleVerticalLandingSeparately && FVector::Parallel(StartingSyncState->GetVelocity_WorldSpace().GetSafeNormal(), UpDir))
+		{
+			TargetVelocity = FMath::Lerp(StartingSyncState->GetVelocity_WorldSpace(), ProposedMove.LinearVelocity.ProjectOnToNormal(UpDir), FractionalVelocityToTarget);
+			TargetPosition = StartingSyncState->GetLocation_WorldSpace() - UpDir * (FloorResult.FloorDist - TargetHeight);
+		}
+		else
 		{
 			const FVector ProposedMovePlaneVelocity = ProposedMove.LinearVelocity - ProposedMove.LinearVelocity.ProjectOnToNormal(PrevGroundNormal);
 			const FVector StartingMovePlaneVelocity = StartingSyncState->GetVelocity_WorldSpace() - StartingSyncState->GetVelocity_WorldSpace().ProjectOnToNormal(PrevGroundNormal);
diff --git a/Engine/Plugins/Experimental/Mover/Source/Mover/Public/PhysicsMover/Modes/PhysicsDrivenWalkingMode.h b/Engine/Plugins/Experimental/Mover/Source/Mover/Public/PhysicsMover/Modes/PhysicsDrivenWalkingMode.h
--- a/Engine/Plugins/Experimental/Mover/Source/Mover/Public/PhysicsMover/Modes/PhysicsDrivenWalkingMode.h
+++ b/Engine/Plugins/Experimental/Mover/Source/Mover/Public/PhysicsMover/Modes/PhysicsDrivenWalkingMode.h
@@ -76,6 +76,11 @@ public:
 	UPROPERTY(EditAnywhere, Category = "Physics Mover", meta = (ClampMin = "0", UIMin = "0", ForceUnits = "s"))
 	float MaxUnsupportedTimeBeforeFalling = 0.06f;
 
+	// This setting is relevant for pawns landing on sloped surfaces. When enabled, no sliding occurs
+	// for vertical landing velocities. When disabled, the sliding is governed by the friction forces.
+	UPROPERTY(EditAnywhere, Category = "Physics Mover")
+	bool bHandleVerticalLandingSeparately = true;
+
 protected:
 	void SwitchToState(const FName& StateName, const FSimulationTickParams& Params, FMoverTickEndData& OutputState);
 
