Added missing copyright boilerplate. No functional changes.

#rnx

diff --git a/Engine/Plugins/Experimental/Mover/Source/Mover/Private/MotionWarpingMoverAdapter.cpp b/Engine/Plugins/Experimental/Mover/Source/Mover/Private/MotionWarpingMoverAdapter.cpp
--- a/Engine/Plugins/Experimental/Mover/Source/Mover/Private/MotionWarpingMoverAdapter.cpp
+++ b/Engine/Plugins/Experimental/Mover/Source/Mover/Private/MotionWarpingMoverAdapter.cpp
@@ -1,3 +1,5 @@
+// Copyright Epic Games, Inc. All Rights Reserved.
+
 #include "MotionWarpingMoverAdapter.h"
 #include "Components/SkeletalMeshComponent.h"
 
