Added speed as a separate display field for Mover gameplay debugger screen output, to make it easier to follow velocity changes.

#rnx

diff --git a/Engine/Plugins/Experimental/Mover/Source/Mover/Private/Debug/GameplayDebuggerCategory_Mover.cpp b/Engine/Plugins/Experimental/Mover/Source/Mover/Private/Debug/GameplayDebuggerCategory_Mover.cpp
--- a/Engine/Plugins/Experimental/Mover/Source/Mover/Private/Debug/GameplayDebuggerCategory_Mover.cpp
+++ b/Engine/Plugins/Experimental/Mover/Source/Mover/Private/Debug/GameplayDebuggerCategory_Mover.cpp
@@ -153,11 +153,12 @@ void FGameplayDebuggerCategory_Mover::DrawData(APlayerController* OwnerPC, FGame
 		DrawInWorldInfo(*FocusedActor, CanvasContext);
 	}
 	
-	CanvasContext.Printf(TEXT("{yellow}%s\n{grey}Local Role: {white}%s\n{grey}Mode: {white}%s\n{grey}Velocity: {white}%s\n{yellow}Active Moves: {white}\n%s\n{yellow}Mode Map: \n{white}%s\n{yellow}Active Transitions: {white}\n%s"),
+	CanvasContext.Printf(TEXT("{yellow}%s\n{grey}Local Role: {white}%s\n{grey}Mode: {white}%s\n{grey}Velocity: {white}%s\n{grey}Speed: {white}%.2f\n{yellow}Active Moves: {white}\n%s\n{yellow}Mode Map: \n{white}%s\n{yellow}Active Transitions: {white}\n%s"),
 		*DataPack.PawnName,
 		*DataPack.LocalRole,
 		*DataPack.MovementModeName,
 		*DataPack.Velocity.ToString(),
+		 DataPack.Velocity.Length(),
 		*FString::JoinBy(DataPack.ActiveLayeredMoves, TEXT("\n"), [](FString MoveAsString) { return MoveAsString; }),
 		*FString::JoinBy(DataPack.ModeMap, TEXT("\n"), [](FString ModeMappingAsString) { return ModeMappingAsString; }),
 		*FString::JoinBy(DataPack.ActiveTransitions, TEXT("\n"), [](FString TransitionAsString) { return TransitionAsString; })
