Mover 2.0 converting Movement modes to use new instant movement effects for jumping and teleporting. This change also moves jumping out to the character movement component rather than specific modes (Modes can still handle jumps specifically if desired)

diff --git a/Engine/Plugins/Experimental/Mover/Source/Mover/Private/DefaultMovementSet/CharacterMoverComponent.cpp b/Engine/Plugins/Experimental/Mover/Source/Mover/Private/DefaultMovementSet/CharacterMoverComponent.cpp
--- a/Engine/Plugins/Experimental/Mover/Source/Mover/Private/DefaultMovementSet/CharacterMoverComponent.cpp
+++ b/Engine/Plugins/Experimental/Mover/Source/Mover/Private/DefaultMovementSet/CharacterMoverComponent.cpp
@@ -2,12 +2,15 @@
 
 #include "DefaultMovementSet/CharacterMoverComponent.h"
 
+#include "DefaultMovementSet/InstantMovementEffects/BasicInstantMovementEffects.h"
 #include "DefaultMovementSet/Modes/FallingMode.h"
 #include "DefaultMovementSet/Modes/FlyingMode.h"
 #include "DefaultMovementSet/Modes/WalkingMode.h"
+#include "DefaultMovementSet/Settings/CommonLegacyMovementSettings.h"
 #include "MoveLibrary/FloorQueryUtils.h"
 
 UCharacterMoverComponent::UCharacterMoverComponent()
+	: bHandleJump(true)
 {
 	// Default movement modes
 	MovementModes.Add(DefaultModeNames::Walking, CreateDefaultSubobject<UWalkingMode>(TEXT("DefaultWalkingMode")));
@@ -17,6 +20,13 @@ UCharacterMoverComponent::UCharacterMoverComponent()
 	StartingMovementMode = DefaultModeNames::Falling;
 }
 
+void UCharacterMoverComponent::BeginPlay()
+{
+	Super::BeginPlay();
+
+	OnPreSimulationTick.AddDynamic(this, &UCharacterMoverComponent::OnMoverPreSimulationTick);
+}
+
 bool UCharacterMoverComponent::IsFalling() const
 {
 	if (bHasValidCachedState)
@@ -47,6 +57,16 @@ bool UCharacterMoverComponent::IsOnGround() const
 	return false;
 }
 
+bool UCharacterMoverComponent::IsSwimming() const
+{
+	if (bHasValidCachedState)
+	{
+		return CachedLastSyncState.MovementMode == DefaultModeNames::Swimming;
+	}
+
+	return false;
+}
+
 bool UCharacterMoverComponent::IsSlopeSliding() const
 {
 	if (IsAirborne())
@@ -61,3 +81,36 @@ bool UCharacterMoverComponent::IsSlopeSliding() const
 
 	return false;
 }
+
+bool UCharacterMoverComponent::CanActorJump() const
+{
+	return IsOnGround();
+}
+
+bool UCharacterMoverComponent::Jump()
+{
+	if (const UCommonLegacyMovementSettings* CommonSettings = FindSharedSettings<UCommonLegacyMovementSettings>())
+	{
+		TSharedPtr<FJumpImpulseEffect> JumpMove = MakeShared<FJumpImpulseEffect>();
+		JumpMove->UpwardsSpeed = CommonSettings->JumpUpwardsSpeed;
+		
+		QueueInstantMovementEffect(JumpMove);
+
+		return true;
+	}
+
+	return false;
+}
+
+void UCharacterMoverComponent::OnMoverPreSimulationTick(const FMoverTimeStep& TimeStep, const FMoverInputCmdContext& InputCmd)
+{
+	if (bHandleJump)
+	{
+		const FCharacterDefaultInputs* CharacterInputs = InputCmd.InputCollection.FindDataByType<FCharacterDefaultInputs>();
+	
+		if (CharacterInputs && CharacterInputs->bIsJumpJustPressed && CanActorJump())
+		{
+			Jump();
+		}
+	}
+}
diff --git a/Engine/Plugins/Experimental/Mover/Source/Mover/Private/DefaultMovementSet/Modes/FallingMode.cpp b/Engine/Plugins/Experimental/Mover/Source/Mover/Private/DefaultMovementSet/Modes/FallingMode.cpp
--- a/Engine/Plugins/Experimental/Mover/Source/Mover/Private/DefaultMovementSet/Modes/FallingMode.cpp
+++ b/Engine/Plugins/Experimental/Mover/Source/Mover/Private/DefaultMovementSet/Modes/FallingMode.cpp
@@ -144,13 +144,6 @@ void UFallingMode::OnSimulationTick(const FSimulationTickParams& Params, FMoverT
 	const float DeltaSeconds = Params.TimeStep.StepMs * 0.001f;
 	float PctTimeApplied = 0.f;
 
-	// Instantaneous movement changes that are executed and we exit before consuming any time
-	if (ProposedMove.bHasTargetLocation && AttemptTeleport(UpdatedComponent, ProposedMove.TargetLocation, UpdatedComponent->GetComponentRotation(), StartingSyncState->GetVelocity_WorldSpace(), OutputSyncState))
-	{
-		OutputState.MovementEndState.RemainingMs = Params.TimeStep.StepMs; 	// Give back all the time
-		return;
-	}
-
 	FMovementRecord MoveRecord;
 	MoveRecord.SetDeltaSeconds(DeltaSeconds);
 	
@@ -250,23 +243,6 @@ void UFallingMode::OnUnregistered()
 	Super::OnUnregistered();
 }
 
-
-bool UFallingMode::AttemptTeleport(USceneComponent* UpdatedComponent, const FVector& TeleportPos, const FRotator& TeleportRot, const FVector& PriorVelocity, FMoverDefaultSyncState& OutputSyncState)
-{
-	if (UpdatedComponent->GetOwner()->TeleportTo(TeleportPos, TeleportRot))
-	{
-		OutputSyncState.SetTransforms_WorldSpace( UpdatedComponent->GetComponentLocation(),
-												  UpdatedComponent->GetComponentRotation(),
-												  PriorVelocity,
-												  nullptr); // no movement base
-
-		UpdatedComponent->ComponentVelocity = PriorVelocity;
-		return true;
-	}
-
-	return false;
-}
-
 void UFallingMode::ProcessLanded(const FFloorCheckResult& FloorResult, FVector& Velocity, FRelativeBaseInfo& BaseInfo, FMoverTickEndData& TickEndData) const
 {
 	UMoverBlackboard* SimBlackboard = GetBlackboard_Mutable();
diff --git a/Engine/Plugins/Experimental/Mover/Source/Mover/Private/DefaultMovementSet/Modes/FlyingMode.cpp b/Engine/Plugins/Experimental/Mover/Source/Mover/Private/DefaultMovementSet/Modes/FlyingMode.cpp
--- a/Engine/Plugins/Experimental/Mover/Source/Mover/Private/DefaultMovementSet/Modes/FlyingMode.cpp
+++ b/Engine/Plugins/Experimental/Mover/Source/Mover/Private/DefaultMovementSet/Modes/FlyingMode.cpp
@@ -73,13 +73,6 @@ void UFlyingMode::OnSimulationTick(const FSimulationTickParams& Params, FMoverTi
 
 	const float DeltaSeconds = Params.TimeStep.StepMs * 0.001f;
 
-	// Instantaneous movement changes that are executed and we exit before consuming any time
-	if (ProposedMove.bHasTargetLocation && AttemptTeleport(UpdatedComponent, ProposedMove.TargetLocation, UpdatedComponent->GetComponentRotation(), *StartingSyncState, OutputState))
-	{
-		OutputState.MovementEndState.RemainingMs = Params.TimeStep.StepMs; 	// Give back all the time
-		return;
-	}
-
 	FMovementRecord MoveRecord;
 	MoveRecord.SetDeltaSeconds(DeltaSeconds);
 
@@ -124,29 +117,6 @@ void UFlyingMode::OnSimulationTick(const FSimulationTickParams& Params, FMoverTi
 	CaptureFinalState(UpdatedComponent, MoveRecord, *StartingSyncState, OutputSyncState, DeltaSeconds);
 }
 
-
-bool UFlyingMode::AttemptTeleport(USceneComponent* UpdatedComponent, const FVector& TeleportPos, const FRotator& TeleportRot, const FMoverDefaultSyncState& StartingSyncState, FMoverTickEndData& Output)
-{
-	if (UpdatedComponent->GetOwner()->TeleportTo(TeleportPos, TeleportRot))
-	{
-		FMoverDefaultSyncState& OutputSyncState = Output.SyncState.SyncStateCollection.FindOrAddMutableDataByType<FMoverDefaultSyncState>();
-
-		OutputSyncState.SetTransforms_WorldSpace( UpdatedComponent->GetComponentLocation(),
-												  UpdatedComponent->GetComponentRotation(),
-												  StartingSyncState.GetVelocity_WorldSpace(),
-												  nullptr); // no movement base
-
-		UpdatedComponent->ComponentVelocity = StartingSyncState.GetVelocity_WorldSpace();
-
-		GetBlackboard_Mutable()->Invalidate(CommonBlackboard::LastFloorResult);
-
-		return true;
-	}
-
-	return false;
-}
-
-
 // TODO: replace this function with simply looking at/collapsing the MovementRecord
 void UFlyingMode::CaptureFinalState(USceneComponent* UpdatedComponent, FMovementRecord& Record, const FMoverDefaultSyncState& StartSyncState, FMoverDefaultSyncState& OutputSyncState, const float DeltaSeconds) const
 {
diff --git a/Engine/Plugins/Experimental/Mover/Source/Mover/Private/DefaultMovementSet/Modes/SwimmingMode.cpp b/Engine/Plugins/Experimental/Mover/Source/Mover/Private/DefaultMovementSet/Modes/SwimmingMode.cpp
--- a/Engine/Plugins/Experimental/Mover/Source/Mover/Private/DefaultMovementSet/Modes/SwimmingMode.cpp
+++ b/Engine/Plugins/Experimental/Mover/Source/Mover/Private/DefaultMovementSet/Modes/SwimmingMode.cpp
@@ -4,7 +4,7 @@
 #include "Components/SkeletalMeshComponent.h"
 #include "GameFramework/Pawn.h"
 #include "MoveLibrary/WaterMovementUtils.h"
-#include "DefaultMovementSet/LayeredMoves/BasicLayeredMoves.h"
+#include "DefaultMovementSet/InstantMovementEffects/BasicInstantMovementEffects.h"
 #include "MoverComponent.h"
 #include "DefaultMovementSet/Settings/CommonLegacyMovementSettings.h"
 #include "MoverLog.h"
@@ -184,34 +184,9 @@ void USwimmingMode::OnUnregistered()
 bool USwimmingMode::AttemptJump(float UpwardsSpeed, FMoverTickEndData& OutputState)
 {
 	// TODO: This should check if a jump is even allowed
-	TSharedPtr<FLayeredMove_JumpImpulse> JumpMove = MakeShared<FLayeredMove_JumpImpulse>();
+	TSharedPtr<FJumpImpulseEffect> JumpMove = MakeShared<FJumpImpulseEffect>();
 	JumpMove->UpwardsSpeed = UpwardsSpeed;
-	OutputState.SyncState.LayeredMoves.QueueLayeredMove(JumpMove);
-	OutputState.MovementEndState.NextModeName = DefaultModeNames::Falling;
+	GetMoverComponent()->QueueInstantMovementEffect(JumpMove);
+	
 	return true;
 }
-
-bool USwimmingMode::AttemptTeleport(USceneComponent* UpdatedComponent, const FVector& TeleportPos, const FRotator& TeleportRot, const FVector& PriorVelocity, FMoverTickEndData& Output)
-{
-	if (UpdatedComponent->GetOwner()->TeleportTo(TeleportPos, TeleportRot))
-	{
-		FMoverDefaultSyncState& OutputSyncState = Output.SyncState.SyncStateCollection.FindOrAddMutableDataByType<FMoverDefaultSyncState>();
-
-		OutputSyncState.SetTransforms_WorldSpace( UpdatedComponent->GetComponentLocation(),
-												  UpdatedComponent->GetComponentRotation(),
-												  PriorVelocity,
-												  nullptr ); // no movement base
-		
-		// TODO: instead of invalidating it, consider checking for a floor. Possibly a dynamic base?
-		if (UMoverBlackboard* SimBlackboard = GetBlackboard_Mutable())
-		{
-			SimBlackboard->Invalidate(CommonBlackboard::LastFloorResult);
-			SimBlackboard->Invalidate(CommonBlackboard::LastWaterResult);
-			SimBlackboard->Invalidate(CommonBlackboard::LastFoundDynamicMovementBase);
-		}
-
-		return true;
-	}
-
-	return false;
-}
\ No newline at end of file
diff --git a/Engine/Plugins/Experimental/Mover/Source/Mover/Private/DefaultMovementSet/Modes/WalkingMode.cpp b/Engine/Plugins/Experimental/Mover/Source/Mover/Private/DefaultMovementSet/Modes/WalkingMode.cpp
--- a/Engine/Plugins/Experimental/Mover/Source/Mover/Private/DefaultMovementSet/Modes/WalkingMode.cpp
+++ b/Engine/Plugins/Experimental/Mover/Source/Mover/Private/DefaultMovementSet/Modes/WalkingMode.cpp
@@ -8,10 +8,10 @@
 #include "MoveLibrary/ModularMovement.h"
 #include "MoveLibrary/FloorQueryUtils.h"
 #include "MoveLibrary/GroundMovementUtils.h"
-#include "DefaultMovementSet/LayeredMoves/BasicLayeredMoves.h"
 #include "MoverComponent.h"
 #include "DefaultMovementSet/Settings/CommonLegacyMovementSettings.h"
 #include "MoverLog.h"
+#include "DefaultMovementSet/InstantMovementEffects/BasicInstantMovementEffects.h"
 
 #include UE_INLINE_GENERATED_CPP_BY_NAME(WalkingMode)
 
@@ -123,15 +123,6 @@ void UWalkingMode::OnSimulationTick(const FSimulationTickParams& Params, FMoverT
 
 	const float DeltaSeconds = Params.TimeStep.StepMs * 0.001f;
 
-	// Instantaneous movement changes that are executed and we exit before consuming any time
-	if ( (ProposedMove.bHasTargetLocation && AttemptTeleport(UpdatedComponent, ProposedMove.TargetLocation, UpdatedComponent->GetComponentRotation(), StartingSyncState->GetVelocity_WorldSpace(), OutputState)) ||	// Teleport
-		 (CharacterInputs && CharacterInputs->bIsJumpJustPressed && AttemptJump(CommonLegacySettings->JumpUpwardsSpeed, OutputState)) )	// Jump
-	{
-		UpdatedComponent->ComponentVelocity = StartingSyncState->GetVelocity_WorldSpace();
-		OutputState.MovementEndState.RemainingMs = Params.TimeStep.StepMs; 	// Give back all the time
-		return;
-	}
-
 	const FVector OrigMoveDelta = ProposedMove.LinearVelocity * DeltaSeconds;
 
 	TObjectPtr<AActor> OwnerActor = UpdatedComponent->GetOwner();
@@ -160,7 +151,7 @@ void UWalkingMode::OnSimulationTick(const FSimulationTickParams& Params, FMoverT
 	bool bIsOrientationChanging = false;
 
 	// Apply orientation changes (if any)
-	if (!ProposedMove.AngularVelocity.IsZero())
+	if (!ProposedMove.AngularVelocity.IsNearlyZero())
 	{ 
 		TargetOrient += (ProposedMove.AngularVelocity * DeltaSeconds);
 		bIsOrientationChanging = (TargetOrient != StartingOrient);
@@ -333,40 +324,6 @@ void UWalkingMode::OnUnregistered()
 	Super::OnUnregistered();
 }
 
-bool UWalkingMode::AttemptJump(float UpwardsSpeed, FMoverTickEndData& OutputState)
-{
-	// TODO: This should check if a jump is even allowed
-	TSharedPtr<FLayeredMove_JumpImpulse> JumpMove = MakeShared<FLayeredMove_JumpImpulse>();
-	JumpMove->UpwardsSpeed = UpwardsSpeed;
-	OutputState.SyncState.LayeredMoves.QueueLayeredMove(JumpMove);
-	OutputState.MovementEndState.NextModeName = CommonLegacySettings->AirMovementModeName;
-	return true;
-}
-
-bool UWalkingMode::AttemptTeleport(USceneComponent* UpdatedComponent, const FVector& TeleportPos, const FRotator& TeleportRot, const FVector& PriorVelocity, FMoverTickEndData& Output)
-{
-	if (UpdatedComponent->GetOwner()->TeleportTo(TeleportPos, TeleportRot))
-	{
-		FMoverDefaultSyncState& OutputSyncState = Output.SyncState.SyncStateCollection.FindOrAddMutableDataByType<FMoverDefaultSyncState>();
-
-		OutputSyncState.SetTransforms_WorldSpace( UpdatedComponent->GetComponentLocation(),
-												  UpdatedComponent->GetComponentRotation(),
-												  PriorVelocity,
-												  nullptr ); // no movement base
-		
-		// TODO: instead of invalidating it, consider checking for a floor. Possibly a dynamic base?
-		if (UMoverBlackboard* SimBlackboard = GetBlackboard_Mutable())
-		{
-			SimBlackboard->Invalidate(CommonBlackboard::LastFloorResult);
-			SimBlackboard->Invalidate(CommonBlackboard::LastFoundDynamicMovementBase);
-		}
-
-		return true;
-	}
-
-	return false;
-}
-
 void UWalkingMode::CaptureFinalState(USceneComponent* UpdatedComponent, bool bDidAttemptMovement, const FFloorCheckResult& FloorResult, const FMovementRecord& Record, FMoverDefaultSyncState& OutputSyncState) const
 {
 	FRelativeBaseInfo PriorBaseInfo;
diff --git a/Engine/Plugins/Experimental/Mover/Source/Mover/Private/PhysicsMover/Modes/PhysicsDrivenFallingMode.cpp b/Engine/Plugins/Experimental/Mover/Source/Mover/Private/PhysicsMover/Modes/PhysicsDrivenFallingMode.cpp
--- a/Engine/Plugins/Experimental/Mover/Source/Mover/Private/PhysicsMover/Modes/PhysicsDrivenFallingMode.cpp
+++ b/Engine/Plugins/Experimental/Mover/Source/Mover/Private/PhysicsMover/Modes/PhysicsDrivenFallingMode.cpp
@@ -16,6 +16,8 @@
 #include "Misc/DataValidation.h"
 #endif // WITH_EDITOR
 
+#include "MoveLibrary/MovementUtils.h"
+
 #include UE_INLINE_GENERATED_CPP_BY_NAME(PhysicsDrivenFallingMode)
 
 #define LOCTEXT_NAMESPACE "PhysicsDrivenFallingMode"
@@ -64,13 +66,6 @@ void UPhysicsDrivenFallingMode::OnSimulationTick(const FSimulationTickParams& Pa
 	const float DeltaSeconds = Params.TimeStep.StepMs * 0.001f;
 	const FVector UpDir = GetMoverComponent()->GetUpDirection();
 
-	// Instantaneous movement changes that are executed and we exit before consuming any time
-	if (ProposedMove.bHasTargetLocation && AttemptTeleport(UpdatedComponent, ProposedMove.TargetLocation, UpdatedComponent->GetComponentRotation(), StartingSyncState->GetVelocity_WorldSpace(), OutputSyncState))
-	{
-		OutputState.MovementEndState.RemainingMs = Params.TimeStep.StepMs; 	// Give back all the time
-		return;
-	}
-
 	// Floor query
 
 	UMoverBlackboard* SimBlackboard = GetBlackboard_Mutable();
@@ -140,6 +135,9 @@ void UPhysicsDrivenFallingMode::OnSimulationTick(const FSimulationTickParams& Pa
 	if (FloorResult.IsWalkableFloor() && bIsFloorWithinReach && !bIsMovingUpRelativeToFloor)
 	{
 		OutputState.MovementEndState.NextModeName = DefaultModeNames::Walking;
+		//TargetVel = FVector::VectorPlaneProject(TargetVel, FloorResult.HitResult.Normal);
+		const FPlane MovementPlane(FVector::ZeroVector, UpDir);
+		TargetVel = UMovementUtils::ConstrainToPlane(TargetVel, MovementPlane, false);
 		TargetPos -= (UpDir.Dot(TargetVel * DeltaSeconds) + (FloorResult.FloorDist - TargetHeight)) * UpDir;
 	}
 	else
diff --git a/Engine/Plugins/Experimental/Mover/Source/Mover/Private/PhysicsMover/Modes/PhysicsDrivenFlyingMode.cpp b/Engine/Plugins/Experimental/Mover/Source/Mover/Private/PhysicsMover/Modes/PhysicsDrivenFlyingMode.cpp
--- a/Engine/Plugins/Experimental/Mover/Source/Mover/Private/PhysicsMover/Modes/PhysicsDrivenFlyingMode.cpp
+++ b/Engine/Plugins/Experimental/Mover/Source/Mover/Private/PhysicsMover/Modes/PhysicsDrivenFlyingMode.cpp
@@ -62,13 +62,6 @@ void UPhysicsDrivenFlyingMode::OnSimulationTick(const FSimulationTickParams& Par
 	const float DeltaSeconds = Params.TimeStep.StepMs * 0.001f;
 	const FVector UpDir = GetMoverComponent()->GetUpDirection();
 
-	// Instantaneous movement changes that are executed and we exit before consuming any time
-	if (ProposedMove.bHasTargetLocation && AttemptTeleport(UpdatedComponent, ProposedMove.TargetLocation, UpdatedComponent->GetComponentRotation(), *StartingSyncState, OutputState))
-	{
-		OutputState.MovementEndState.RemainingMs = Params.TimeStep.StepMs; 	// Give back all the time
-		return;
-	}
-
 	// Don't need a floor query - just invalidate the blackboard to ensure we don't use an old result elsewhere
 
 	if (UMoverBlackboard* SimBlackboard = GetBlackboard_Mutable())
diff --git a/Engine/Plugins/Experimental/Mover/Source/Mover/Private/PhysicsMover/Modes/PhysicsDrivenSwimmingMode.cpp b/Engine/Plugins/Experimental/Mover/Source/Mover/Private/PhysicsMover/Modes/PhysicsDrivenSwimmingMode.cpp
--- a/Engine/Plugins/Experimental/Mover/Source/Mover/Private/PhysicsMover/Modes/PhysicsDrivenSwimmingMode.cpp
+++ b/Engine/Plugins/Experimental/Mover/Source/Mover/Private/PhysicsMover/Modes/PhysicsDrivenSwimmingMode.cpp
@@ -3,7 +3,7 @@
 #include "PhysicsMover/Modes/PhysicsDrivenSwimmingMode.h"
 
 #include "Chaos/Character/CharacterGroundConstraint.h"
-#include "DefaultMovementSet/LayeredMoves/BasicLayeredMoves.h"
+#include "DefaultMovementSet/InstantMovementEffects/BasicInstantMovementEffects.h"
 #include "DefaultMovementSet/Settings/CommonLegacyMovementSettings.h"
 #include "GameFramework/PhysicsVolume.h"
 #include "Math/UnitConversion.h"
@@ -69,8 +69,7 @@ void UPhysicsDrivenSwimmingMode::OnSimulationTick(const FSimulationTickParams& P
 
 	const float DeltaSeconds = Params.TimeStep.StepMs * 0.001f;
 
-	if ((ProposedMove.bHasTargetLocation && AttemptTeleport(UpdatedComponent, ProposedMove.TargetLocation, UpdatedComponent->GetComponentRotation(), StartingSyncState->GetVelocity_WorldSpace(), OutputState)) ||	// Teleport
-		(CharacterInputs->bIsJumpJustPressed && AttemptJump(SurfaceSwimmingWaterControlSettings.JumpMultiplier*CommonLegacySettings->JumpUpwardsSpeed, OutputState)))
+	if (CharacterInputs->bIsJumpJustPressed && AttemptJump(SurfaceSwimmingWaterControlSettings.JumpMultiplier*CommonLegacySettings->JumpUpwardsSpeed, OutputState))
 	{
 		OutputState.MovementEndState.RemainingMs = Params.TimeStep.StepMs;
 		return;
@@ -154,29 +153,10 @@ void UPhysicsDrivenSwimmingMode::OnSimulationTick(const FSimulationTickParams& P
 bool UPhysicsDrivenSwimmingMode::AttemptJump(float UpwardsSpeed, FMoverTickEndData& OutputState)
 {
 	// TODO: This should check if a jump is even allowed
- 	TSharedPtr<FLayeredMove_JumpImpulse> JumpMove = MakeShared<FLayeredMove_JumpImpulse>();
+ 	TSharedPtr<FJumpImpulseEffect> JumpMove = MakeShared<FJumpImpulseEffect>();
 	JumpMove->UpwardsSpeed = UpwardsSpeed;
-	OutputState.SyncState.LayeredMoves.QueueLayeredMove(JumpMove);
-	OutputState.MovementEndState.NextModeName = DefaultModeNames::Falling;
-	return true;
-}
-
-bool UPhysicsDrivenSwimmingMode::AttemptTeleport(USceneComponent* UpdatedComponent, const FVector& TeleportPos, const FRotator& TeleportRot, const FVector& PriorVelocity, FMoverTickEndData& Output)
-{
-	FMoverDefaultSyncState& OutputSyncState = Output.SyncState.SyncStateCollection.FindOrAddMutableDataByType<FMoverDefaultSyncState>();
-
-	OutputSyncState.SetTransforms_WorldSpace(TeleportPos,
-		TeleportRot,
-		PriorVelocity,
-		nullptr); // no movement base
-
-	// TODO: instead of invalidating it, consider checking for a floor. Possibly a dynamic base?
-	if (UMoverBlackboard* SimBlackboard = GetBlackboard_Mutable())
-	{
-		SimBlackboard->Invalidate(CommonBlackboard::LastFloorResult);
-		SimBlackboard->Invalidate(CommonBlackboard::LastFoundDynamicMovementBase);
-	}
-
+	GetMoverComponent()->QueueInstantMovementEffect(JumpMove);
+	
 	return true;
 }
 
diff --git a/Engine/Plugins/Experimental/Mover/Source/Mover/Private/PhysicsMover/Modes/PhysicsDrivenWalkingMode.cpp b/Engine/Plugins/Experimental/Mover/Source/Mover/Private/PhysicsMover/Modes/PhysicsDrivenWalkingMode.cpp
--- a/Engine/Plugins/Experimental/Mover/Source/Mover/Private/PhysicsMover/Modes/PhysicsDrivenWalkingMode.cpp
+++ b/Engine/Plugins/Experimental/Mover/Source/Mover/Private/PhysicsMover/Modes/PhysicsDrivenWalkingMode.cpp
@@ -276,13 +276,6 @@ void UPhysicsDrivenWalkingMode::OnSimulationTick(const FSimulationTickParams& Pa
 
 	const float DeltaSeconds = Params.TimeStep.StepMs * 0.001f;
 
-	// Instantaneous movement changes that are executed and we exit before consuming any time
-	if ((ProposedMove.bHasTargetLocation && AttemptTeleport(UpdatedComponent, ProposedMove.TargetLocation, UpdatedComponent->GetComponentRotation(), StartingSyncState->GetVelocity_WorldSpace(), OutputState)))
-	{
-		OutputState.MovementEndState.RemainingMs = Params.TimeStep.StepMs; 	// Give back all the time
-		return;
-	}
-
 	UMoverBlackboard* SimBlackboard = GetBlackboard_Mutable();
 	if (!SimBlackboard)
 	{
@@ -426,25 +419,6 @@ void UPhysicsDrivenWalkingMode::OnSimulationTick(const FSimulationTickParams& Pa
 	}
 }
 
-bool UPhysicsDrivenWalkingMode::AttemptTeleport(USceneComponent* UpdatedComponent, const FVector& TeleportPos, const FRotator& TeleportRot, const FVector& PriorVelocity, FMoverTickEndData& Output)
-{
-	FMoverDefaultSyncState& OutputSyncState = Output.SyncState.SyncStateCollection.FindOrAddMutableDataByType<FMoverDefaultSyncState>();
-
-	OutputSyncState.SetTransforms_WorldSpace(TeleportPos,
-		TeleportRot,
-		PriorVelocity,
-		nullptr); // no movement base
-
-	// TODO: instead of invalidating it, consider checking for a floor. Possibly a dynamic base?
-	if (UMoverBlackboard* SimBlackboard = GetBlackboard_Mutable())
-	{
-		SimBlackboard->Invalidate(CommonBlackboard::LastFloorResult);
-		SimBlackboard->Invalidate(CommonBlackboard::LastFoundDynamicMovementBase);
-	}
-
-	return true;
-}
-
 void UPhysicsDrivenWalkingMode::SwitchToState(const FName& StateName, const FSimulationTickParams& Params, FMoverTickEndData& OutputState)
 {
 	OutputState.MovementEndState.RemainingMs = Params.TimeStep.StepMs;
diff --git a/Engine/Plugins/Experimental/Mover/Source/Mover/Private/PhysicsMover/Transitions/PhysicsJumpCheck.cpp b/Engine/Plugins/Experimental/Mover/Source/Mover/Private/PhysicsMover/Transitions/PhysicsJumpCheck.cpp
--- a/Engine/Plugins/Experimental/Mover/Source/Mover/Private/PhysicsMover/Transitions/PhysicsJumpCheck.cpp
+++ b/Engine/Plugins/Experimental/Mover/Source/Mover/Private/PhysicsMover/Transitions/PhysicsJumpCheck.cpp
@@ -1,7 +1,7 @@
 // Copyright Epic Games, Inc. All Rights Reserved.
 
 #include "PhysicsMover/Transitions/PhysicsJumpCheck.h"
-#include "DefaultMovementSet/LayeredMoves/BasicLayeredMoves.h"
+#include "DefaultMovementSet/InstantMovementEffects/BasicInstantMovementEffects.h"
 #include "MoverComponent.h"
 
 #if WITH_EDITOR
@@ -32,11 +32,10 @@ FTransitionEvalResult UPhysicsJumpCheck::OnEvaluate(const FSimulationTickParams&
 
 void UPhysicsJumpCheck::OnTrigger(const FSimulationTickParams& Params)
 {
-	TSharedPtr<FLayeredMove_JumpImpulse> JumpMove = MakeShared<FLayeredMove_JumpImpulse>();
+	TSharedPtr<FJumpImpulseEffect> JumpMove = MakeShared<FJumpImpulseEffect>();
 	JumpMove->UpwardsSpeed = JumpUpwardsSpeed;
 
-	Params.MoverComponent->QueueLayeredMove(JumpMove);
-
+	Params.MoverComponent->QueueInstantMovementEffect(JumpMove);
 }
 
 #if WITH_EDITOR
diff --git a/Engine/Plugins/Experimental/Mover/Source/Mover/Public/DefaultMovementSet/CharacterMoverComponent.h b/Engine/Plugins/Experimental/Mover/Source/Mover/Public/DefaultMovementSet/CharacterMoverComponent.h
--- a/Engine/Plugins/Experimental/Mover/Source/Mover/Public/DefaultMovementSet/CharacterMoverComponent.h
+++ b/Engine/Plugins/Experimental/Mover/Source/Mover/Public/DefaultMovementSet/CharacterMoverComponent.h
@@ -13,6 +13,8 @@ class MOVER_API UCharacterMoverComponent : public UMoverComponent
 public:
 	UCharacterMoverComponent();
 	
+	virtual void BeginPlay() override;
+	
 	// Is this actor in a falling state? Note that this includes upwards motion induced by jumping.
 	UFUNCTION(BlueprintPure, Category = Mover)
 	virtual bool IsFalling() const;
@@ -25,7 +27,27 @@ public:
 	UFUNCTION(BlueprintPure, Category = Mover)
 	virtual bool IsOnGround() const;
 
+	// Is this actor in a Swimming state? (e.g. Swimming)
+	UFUNCTION(BlueprintPure, Category = Mover)
+	virtual bool IsSwimming() const;
+	
 	// Is this actor sliding on an unwalkable slope
 	UFUNCTION(BlueprintPure, Category = Mover)
 	virtual bool IsSlopeSliding() const;
+
+	// Can this Actor jump?
+	UFUNCTION(BlueprintPure, Category = Mover)
+	virtual bool CanActorJump() const;
+
+	// Perform jump on actor
+	UFUNCTION(BlueprintCallable, Category = Mover)
+	virtual bool Jump();
+
+	// Whether this component should directly handle jumping or not 
+	UPROPERTY(EditAnywhere, BlueprintReadWrite, Category=Mover)
+    bool bHandleJump;
+	
+protected:
+	UFUNCTION()
+	virtual void OnMoverPreSimulationTick(const FMoverTimeStep& TimeStep, const FMoverInputCmdContext& InputCmd);
 };
diff --git a/Engine/Plugins/Experimental/Mover/Source/Mover/Public/DefaultMovementSet/Modes/FallingMode.h b/Engine/Plugins/Experimental/Mover/Source/Mover/Public/DefaultMovementSet/Modes/FallingMode.h
--- a/Engine/Plugins/Experimental/Mover/Source/Mover/Public/DefaultMovementSet/Modes/FallingMode.h
+++ b/Engine/Plugins/Experimental/Mover/Source/Mover/Public/DefaultMovementSet/Modes/FallingMode.h
@@ -86,8 +86,6 @@ protected:
 	virtual void OnRegistered(const FName ModeName) override;
 	virtual void OnUnregistered() override;
 
-	virtual bool AttemptTeleport(USceneComponent* UpdatedComponent, const FVector& TeleportPos, const FRotator& TeleportRot, const FVector& PriorVelocity, FMoverDefaultSyncState& OutputSyncState);
-
 	void CaptureFinalState(USceneComponent* UpdatedComponent, const FMoverDefaultSyncState& StartSyncState, const FFloorCheckResult& FloorResult, float DeltaSeconds, float DeltaSecondsUsed, FMoverDefaultSyncState& OutputSyncState, FMoverTickEndData& TickEndData, FMovementRecord& Record) const;
 
 	TObjectPtr<const UCommonLegacyMovementSettings> CommonLegacySettings;
diff --git a/Engine/Plugins/Experimental/Mover/Source/Mover/Public/DefaultMovementSet/Modes/FlyingMode.h b/Engine/Plugins/Experimental/Mover/Source/Mover/Public/DefaultMovementSet/Modes/FlyingMode.h
--- a/Engine/Plugins/Experimental/Mover/Source/Mover/Public/DefaultMovementSet/Modes/FlyingMode.h
+++ b/Engine/Plugins/Experimental/Mover/Source/Mover/Public/DefaultMovementSet/Modes/FlyingMode.h
@@ -30,8 +30,6 @@ protected:
 	virtual void OnRegistered(const FName ModeName) override;
 	virtual void OnUnregistered() override;
 
-	virtual bool AttemptTeleport(USceneComponent* UpdatedComponent, const FVector& TeleportPos, const FRotator& TeleportRot, const FMoverDefaultSyncState& StartingSyncState, FMoverTickEndData& Output);
-
 	void CaptureFinalState(USceneComponent* UpdatedComponent, FMovementRecord& Record, const FMoverDefaultSyncState& StartSyncState, FMoverDefaultSyncState& OutputSyncState, const float DeltaSeconds) const;
 
 	TObjectPtr<const UCommonLegacyMovementSettings> CommonLegacySettings;
diff --git a/Engine/Plugins/Experimental/Mover/Source/Mover/Public/DefaultMovementSet/Modes/SwimmingMode.h b/Engine/Plugins/Experimental/Mover/Source/Mover/Public/DefaultMovementSet/Modes/SwimmingMode.h
--- a/Engine/Plugins/Experimental/Mover/Source/Mover/Public/DefaultMovementSet/Modes/SwimmingMode.h
+++ b/Engine/Plugins/Experimental/Mover/Source/Mover/Public/DefaultMovementSet/Modes/SwimmingMode.h
@@ -210,8 +210,6 @@ protected:
 	virtual void OnUnregistered() override;
 
 	virtual bool AttemptJump(float UpwardsSpeed, FMoverTickEndData& Output);
-	virtual bool AttemptTeleport(USceneComponent* UpdatedComponent, const FVector& TeleportPos, const FRotator& TeleportRot, const FVector& PriorVelocity, FMoverTickEndData& Output);
-
 
 	void CaptureFinalState(USceneComponent* UpdatedComponent, const FFloorCheckResult& FloorResult, const FMovementRecord& Record, FMoverDefaultSyncState& OutputSyncState) const;
 
diff --git a/Engine/Plugins/Experimental/Mover/Source/Mover/Public/DefaultMovementSet/Modes/WalkingMode.h b/Engine/Plugins/Experimental/Mover/Source/Mover/Public/DefaultMovementSet/Modes/WalkingMode.h
--- a/Engine/Plugins/Experimental/Mover/Source/Mover/Public/DefaultMovementSet/Modes/WalkingMode.h
+++ b/Engine/Plugins/Experimental/Mover/Source/Mover/Public/DefaultMovementSet/Modes/WalkingMode.h
@@ -44,10 +44,6 @@ protected:
 	virtual void OnRegistered(const FName ModeName) override; 
 	virtual void OnUnregistered() override;
 
-	virtual bool AttemptJump(float UpwardsSpeed, FMoverTickEndData& Output);
-	virtual bool AttemptTeleport(USceneComponent* UpdatedComponent, const FVector& TeleportPos, const FRotator& TeleportRot, const FVector& PriorVelocity, FMoverTickEndData& Output);
-
-
 	void CaptureFinalState(USceneComponent* UpdatedComponent, bool bDidAttemptMovement, const FFloorCheckResult& FloorResult, const FMovementRecord& Record, FMoverDefaultSyncState& OutputSyncState) const;
 
 	FRelativeBaseInfo UpdateFloorAndBaseInfo(const FFloorCheckResult& FloorResult) const;
diff --git a/Engine/Plugins/Experimental/Mover/Source/Mover/Public/PhysicsMover/Modes/PhysicsDrivenSwimmingMode.h b/Engine/Plugins/Experimental/Mover/Source/Mover/Public/PhysicsMover/Modes/PhysicsDrivenSwimmingMode.h
--- a/Engine/Plugins/Experimental/Mover/Source/Mover/Public/PhysicsMover/Modes/PhysicsDrivenSwimmingMode.h
+++ b/Engine/Plugins/Experimental/Mover/Source/Mover/Public/PhysicsMover/Modes/PhysicsDrivenSwimmingMode.h
@@ -22,8 +22,7 @@ class MOVER_API UPhysicsDrivenSwimmingMode : public USwimmingMode, public IPhysi
 public:
 
 	virtual void OnSimulationTick(const FSimulationTickParams& Params, FMoverTickEndData& OutputState) override;
-	virtual bool AttemptJump(float UpwardsSpeed, FMoverTickEndData& Output);
-	virtual bool AttemptTeleport(USceneComponent* UpdatedComponent, const FVector& TeleportPos, const FRotator& TeleportRot, const FVector& PriorVelocity, FMoverTickEndData& Output);
+	virtual bool AttemptJump(float UpwardsSpeed, FMoverTickEndData& Output) override;
 
 	virtual void UpdateConstraintSettings(Chaos::FCharacterGroundConstraint& Constraint) const override;
 
diff --git a/Engine/Plugins/Experimental/Mover/Source/Mover/Public/PhysicsMover/Modes/PhysicsDrivenWalkingMode.h b/Engine/Plugins/Experimental/Mover/Source/Mover/Public/PhysicsMover/Modes/PhysicsDrivenWalkingMode.h
--- a/Engine/Plugins/Experimental/Mover/Source/Mover/Public/PhysicsMover/Modes/PhysicsDrivenWalkingMode.h
+++ b/Engine/Plugins/Experimental/Mover/Source/Mover/Public/PhysicsMover/Modes/PhysicsDrivenWalkingMode.h
@@ -23,7 +23,6 @@ class MOVER_API UPhysicsDrivenWalkingMode : public UWalkingMode, public IPhysics
 public:
 
 	virtual void OnSimulationTick(const FSimulationTickParams& Params, FMoverTickEndData& OutputState) override;
-	virtual bool AttemptTeleport(USceneComponent* UpdatedComponent, const FVector& TeleportPos, const FRotator& TeleportRot, const FVector& PriorVelocity, FMoverTickEndData& Output);
 
 	virtual void UpdateConstraintSettings(Chaos::FCharacterGroundConstraint& Constraint) const override;
 	virtual void OnContactModification_Internal(const FPhysicsMoverSimulationContactModifierParams& Params, Chaos::FCollisionContactModifier& Modifier) const override;
