Mover 2.0 removing teleport layered move as it's functionality now lives in the teleport instant movement effect. This allows us to get rid of some data in FProposedMove so we're cleaning that up as well.

diff --git a/Engine/Plugins/Experimental/Mover/Config/BaseMover.ini b/Engine/Plugins/Experimental/Mover/Config/BaseMover.ini
--- a/Engine/Plugins/Experimental/Mover/Config/BaseMover.ini
+++ b/Engine/Plugins/Experimental/Mover/Config/BaseMover.ini
@@ -6,7 +6,7 @@
 +StructRedirects=(OldName="/Script/Mover.KinematicMoverTickStartData",NewName="/Script/Mover.MoverTickStartData")
 +StructRedirects=(OldName="/Script/Mover.KinematicMoverTickEndData",NewName="/Script/Mover.MoverTickEndData")
 +StructRedirects=(OldName="/Script/Mover.KinematicDefaultInputs",NewName="/Script/Mover.CharacterDefaultInputs")
-
++StructRedirects=(OldName="/Script/Mover.LayeredMove_JumpImpulse",NewName="/Script/Mover.LayeredMove_JumpImpulseOverDuration")
 
 
 
diff --git a/Engine/Plugins/Experimental/Mover/Source/Mover/Private/DefaultMovementSet/LayeredMoves/AnimRootMotionLayeredMove.cpp b/Engine/Plugins/Experimental/Mover/Source/Mover/Private/DefaultMovementSet/LayeredMoves/AnimRootMotionLayeredMove.cpp
--- a/Engine/Plugins/Experimental/Mover/Source/Mover/Private/DefaultMovementSet/LayeredMoves/AnimRootMotionLayeredMove.cpp
+++ b/Engine/Plugins/Experimental/Mover/Source/Mover/Private/DefaultMovementSet/LayeredMoves/AnimRootMotionLayeredMove.cpp
@@ -60,10 +60,7 @@ bool FLayeredMove_AnimRootMotion::GenerateMove(const FMoverTickStartData& SimSta
 	OutProposedMove.MixMode = MixMode;
 
 	// Convert the transform into linear and angular velocities
-	const FPlane MovementPlane(FVector::ZeroVector, MoverComp->GetUpDirection());
-
 	OutProposedMove.LinearVelocity    = WorldSpaceRootMotion.GetTranslation() / DeltaSeconds;
-	OutProposedMove.MovePlaneVelocity = UMovementUtils::ConstrainToPlane(OutProposedMove.LinearVelocity, MovementPlane);
 	OutProposedMove.AngularVelocity   = WorldSpaceRootMotion.GetRotation().Rotator() * (1.f / DeltaSeconds);
 
 	return true;
diff --git a/Engine/Plugins/Experimental/Mover/Source/Mover/Private/DefaultMovementSet/LayeredMoves/BasicLayeredMoves.cpp b/Engine/Plugins/Experimental/Mover/Source/Mover/Private/DefaultMovementSet/LayeredMoves/BasicLayeredMoves.cpp
--- a/Engine/Plugins/Experimental/Mover/Source/Mover/Private/DefaultMovementSet/LayeredMoves/BasicLayeredMoves.cpp
+++ b/Engine/Plugins/Experimental/Mover/Source/Mover/Private/DefaultMovementSet/LayeredMoves/BasicLayeredMoves.cpp
@@ -90,18 +90,18 @@ void FLayeredMove_LinearVelocity::AddReferencedObjects(class FReferenceCollector
 
 
 // -------------------------------------------------------------------
-// FLayeredMove_JumpImpulse
+// FLayeredMove_JumpImpulseOverDuration
 // -------------------------------------------------------------------
 
-FLayeredMove_JumpImpulse::FLayeredMove_JumpImpulse() 
+FLayeredMove_JumpImpulseOverDuration::FLayeredMove_JumpImpulseOverDuration() 
 	: UpwardsSpeed(0.f)
 {
-	DurationMs = 0.f;
+	DurationMs = 100.f;
 	MixMode = EMoveMixMode::OverrideVelocity;
 }
 
 
-bool FLayeredMove_JumpImpulse::GenerateMove(const FMoverTickStartData& SimState, const FMoverTimeStep& TimeStep, const UMoverComponent* MoverComp, UMoverBlackboard* SimBlackboard, FProposedMove& OutProposedMove)
+bool FLayeredMove_JumpImpulseOverDuration::GenerateMove(const FMoverTickStartData& SimState, const FMoverTimeStep& TimeStep, const UMoverComponent* MoverComp, UMoverBlackboard* SimBlackboard, FProposedMove& OutProposedMove)
 {	
 	const FMoverDefaultSyncState* SyncState = SimState.SyncState.SyncStateCollection.FindDataByType<FMoverDefaultSyncState>();
 	check(SyncState);
@@ -127,30 +127,30 @@ bool FLayeredMove_JumpImpulse::GenerateMove(const FMoverTickStartData& SimState,
 	return true;
 }
 
-FLayeredMoveBase* FLayeredMove_JumpImpulse::Clone() const
+FLayeredMoveBase* FLayeredMove_JumpImpulseOverDuration::Clone() const
 {
-	FLayeredMove_JumpImpulse* CopyPtr = new FLayeredMove_JumpImpulse(*this);
+	FLayeredMove_JumpImpulseOverDuration* CopyPtr = new FLayeredMove_JumpImpulseOverDuration(*this);
 	return CopyPtr;
 }
 
-void FLayeredMove_JumpImpulse::NetSerialize(FArchive& Ar)
+void FLayeredMove_JumpImpulseOverDuration::NetSerialize(FArchive& Ar)
 {
 	Super::NetSerialize(Ar);
 
 	Ar << UpwardsSpeed;
 }
 
-UScriptStruct* FLayeredMove_JumpImpulse::GetScriptStruct() const
+UScriptStruct* FLayeredMove_JumpImpulseOverDuration::GetScriptStruct() const
 {
-	return FLayeredMove_JumpImpulse::StaticStruct();
+	return FLayeredMove_JumpImpulseOverDuration::StaticStruct();
 }
 
-FString FLayeredMove_JumpImpulse::ToSimpleString() const
+FString FLayeredMove_JumpImpulseOverDuration::ToSimpleString() const
 {
-	return FString::Printf(TEXT("JumpImpulse"));
+	return FString::Printf(TEXT("JumpImpulseOverDuration"));
 }
 
-void FLayeredMove_JumpImpulse::AddReferencedObjects(class FReferenceCollector& Collector)
+void FLayeredMove_JumpImpulseOverDuration::AddReferencedObjects(class FReferenceCollector& Collector)
 {
 	Super::AddReferencedObjects(Collector);
 }
@@ -306,56 +306,6 @@ FVector FLayeredMove_JumpTo::GetRelativeLocation(float MoveFraction, const FRota
 	return FacingRotation.RotateVector(RelativeLocationFacingSpace);
 }
 
-// -------------------------------------------------------------------
-// FLayeredMove_Teleport
-// -------------------------------------------------------------------
-
-FLayeredMove_Teleport::FLayeredMove_Teleport()
-	: TargetLocation(FVector::ZeroVector)
-{
-	DurationMs = 0.f;
-	MixMode = EMoveMixMode::OverrideAll;
-	Priority = 10;
-}
-
-bool FLayeredMove_Teleport::GenerateMove(const FMoverTickStartData& SimState, const FMoverTimeStep& TimeStep, const UMoverComponent* MoverComp, UMoverBlackboard* SimBlackboard, FProposedMove& OutProposedMove)
-{
-	OutProposedMove = FProposedMove();
-	OutProposedMove.MixMode = EMoveMixMode::OverrideAll;
-	OutProposedMove.TargetLocation = TargetLocation;
-	OutProposedMove.bHasTargetLocation = true;
-
-	return true;
-}
-
-FLayeredMoveBase* FLayeredMove_Teleport::Clone() const
-{
-	FLayeredMove_Teleport* CopyPtr = new FLayeredMove_Teleport(*this);
-	return CopyPtr;
-}
-
-void FLayeredMove_Teleport::NetSerialize(FArchive& Ar)
-{
-	Super::NetSerialize(Ar);
-
-	Ar << TargetLocation;
-}
-
-UScriptStruct* FLayeredMove_Teleport::GetScriptStruct() const
-{
-	return FLayeredMove_Teleport::StaticStruct();
-}
-
-FString FLayeredMove_Teleport::ToSimpleString() const
-{
-	return FString::Printf(TEXT("Teleport"));
-}
-
-void FLayeredMove_Teleport::AddReferencedObjects(class FReferenceCollector& Collector)
-{
-	Super::AddReferencedObjects(Collector);
-}
-
 // -------------------------------------------------------------------
 // FLayeredMove_MoveTo
 // -------------------------------------------------------------------
@@ -402,8 +352,6 @@ float FLayeredMove_MoveTo::EvaluateFloatCurveAtFraction(const UCurveFloat& Curve
 bool FLayeredMove_MoveTo::GenerateMove(const FMoverTickStartData& StartState, const FMoverTimeStep& TimeStep, const UMoverComponent* MoverComp, UMoverBlackboard* SimBlackboard, FProposedMove& OutProposedMove)
 {
 	OutProposedMove.MixMode = MixMode;
-	OutProposedMove.TargetLocation = TargetLocation;
-	OutProposedMove.bHasTargetLocation = true;
 
 	const float DeltaSeconds = TimeStep.StepMs / 1000.f;
 	
diff --git a/Engine/Plugins/Experimental/Mover/Source/Mover/Private/MoveLibrary/AirMovementUtils.cpp b/Engine/Plugins/Experimental/Mover/Source/Mover/Private/MoveLibrary/AirMovementUtils.cpp
--- a/Engine/Plugins/Experimental/Mover/Source/Mover/Private/MoveLibrary/AirMovementUtils.cpp
+++ b/Engine/Plugins/Experimental/Mover/Source/Mover/Private/MoveLibrary/AirMovementUtils.cpp
@@ -25,7 +25,6 @@ FProposedMove UAirMovementUtils::ComputeControlledFreeMove(const FFreeMoveParams
 	ComputeVelocityParams.Acceleration = InParams.Acceleration;
 	
 	OutMove.LinearVelocity = UMovementUtils::ComputeVelocity(ComputeVelocityParams);
-	OutMove.MovePlaneVelocity = UMovementUtils::ConstrainToPlane(OutMove.LinearVelocity, MovementPlane, false);
 
 	// JAH TODO: this is where we can perform turning, based on aux settings. For now, just snap to the intended final orientation.
 	FVector IntendedFacingDir = InParams.OrientationIntent.RotateVector(FVector::ForwardVector).GetSafeNormal();
diff --git a/Engine/Plugins/Experimental/Mover/Source/Mover/Private/MoveLibrary/GroundMovementUtils.cpp b/Engine/Plugins/Experimental/Mover/Source/Mover/Private/MoveLibrary/GroundMovementUtils.cpp
--- a/Engine/Plugins/Experimental/Mover/Source/Mover/Private/MoveLibrary/GroundMovementUtils.cpp
+++ b/Engine/Plugins/Experimental/Mover/Source/Mover/Private/MoveLibrary/GroundMovementUtils.cpp
@@ -27,7 +27,7 @@ FProposedMove UGroundMovementUtils::ComputeControlledGroundMove(const FGroundMov
 	FComputeVelocityParams ComputeVelocityParams;
 	ComputeVelocityParams.DeltaSeconds = InParams.DeltaSeconds;
 	ComputeVelocityParams.InitialVelocity = InParams.PriorVelocity;
-	ComputeVelocityParams.MoveDirectionIntent = InParams.MoveInput;
+	ComputeVelocityParams.MoveDirectionIntent = MoveDirIntentInMovementPlane;
 	ComputeVelocityParams.MaxSpeed = InParams.MaxSpeed;
 	ComputeVelocityParams.TurningBoost = InParams.TurningBoost;
 	ComputeVelocityParams.Deceleration = InParams.Deceleration;
@@ -35,8 +35,8 @@ FProposedMove UGroundMovementUtils::ComputeControlledGroundMove(const FGroundMov
 	ComputeVelocityParams.Friction = InParams.Friction;
 	
 	// Figure out linear velocity
-	OutMove.MovePlaneVelocity = UMovementUtils::ComputeVelocity(ComputeVelocityParams);
-	OutMove.LinearVelocity = UMovementUtils::ConstrainToPlane(OutMove.MovePlaneVelocity, GroundSurfacePlane, true);
+	const FVector Velocity = UMovementUtils::ComputeVelocity(ComputeVelocityParams);
+	OutMove.LinearVelocity = UMovementUtils::ConstrainToPlane(Velocity, GroundSurfacePlane, true);
 
 	// Linearly rotate in place
 	OutMove.AngularVelocity = UMovementUtils::ComputeAngularVelocity(InParams.PriorOrientation, InParams.OrientationIntent, InParams.DeltaSeconds, InParams.TurningRate);
diff --git a/Engine/Plugins/Experimental/Mover/Source/Mover/Public/DefaultMovementSet/LayeredMoves/BasicLayeredMoves.h b/Engine/Plugins/Experimental/Mover/Source/Mover/Public/DefaultMovementSet/LayeredMoves/BasicLayeredMoves.h
--- a/Engine/Plugins/Experimental/Mover/Source/Mover/Public/DefaultMovementSet/LayeredMoves/BasicLayeredMoves.h
+++ b/Engine/Plugins/Experimental/Mover/Source/Mover/Public/DefaultMovementSet/LayeredMoves/BasicLayeredMoves.h
@@ -68,15 +68,18 @@ struct TStructOpsTypeTraits< FLayeredMove_LinearVelocity > : public TStructOpsTy
 };
 
 
-/** Jump Impulse: introduces an instantaneous upwards change in velocity. This overrides the existing 'up' component of the actor's current velocity */
+/** Jump Impulse Over Duration: introduces an instantaneous upwards change in velocity over a duration(ms).
+  * This overrides the existing 'up' component of the actor's current velocity
+  * Note: For a one time Jump Impulse consider using a FJumpImpulseEffect Instant Effect
+  */
 USTRUCT(BlueprintType)
-struct MOVER_API FLayeredMove_JumpImpulse : public FLayeredMoveBase
+struct MOVER_API FLayeredMove_JumpImpulseOverDuration : public FLayeredMoveBase
 {
 	GENERATED_USTRUCT_BODY()
 
-	FLayeredMove_JumpImpulse();
+	FLayeredMove_JumpImpulseOverDuration();
 
-	virtual ~FLayeredMove_JumpImpulse() {}
+	virtual ~FLayeredMove_JumpImpulseOverDuration() {}
 
 	// Units per second, in whatever direction the target actor considers 'up'
 	UPROPERTY(EditAnywhere, BlueprintReadWrite, Category = Mover)
@@ -97,7 +100,7 @@ struct MOVER_API FLayeredMove_JumpImpulse : public FLayeredMoveBase
 };
 
 template<>
-struct TStructOpsTypeTraits< FLayeredMove_JumpImpulse > : public TStructOpsTypeTraitsBase2< FLayeredMove_JumpImpulse >
+struct TStructOpsTypeTraits< FLayeredMove_JumpImpulseOverDuration > : public TStructOpsTypeTraitsBase2< FLayeredMove_JumpImpulseOverDuration >
 {
 	enum
 	{
@@ -174,44 +177,6 @@ struct TStructOpsTypeTraits< FLayeredMove_JumpTo > : public TStructOpsTypeTraits
 	};
 };
 
-/** Teleport: instantly moves an actor to a new location */
-USTRUCT(BlueprintType)
-struct MOVER_API FLayeredMove_Teleport : public FLayeredMoveBase
-{
-	GENERATED_USTRUCT_BODY()
-
-	FLayeredMove_Teleport();
-	virtual ~FLayeredMove_Teleport() {}
-
-	// Location to teleport to, in world space
-	UPROPERTY(EditAnywhere, BlueprintReadWrite, Category = Mover)
-	FVector TargetLocation;
-
-	// Generate a movement 
-	virtual bool GenerateMove(const FMoverTickStartData& StartState, const FMoverTimeStep& TimeStep, const UMoverComponent* MoverComp, UMoverBlackboard* SimBlackboard, FProposedMove& OutProposedMove) override;
-
-	virtual FLayeredMoveBase* Clone() const override;
-
-	virtual void NetSerialize(FArchive& Ar) override;
-
-	virtual UScriptStruct* GetScriptStruct() const override;
-
-	virtual FString ToSimpleString() const override;
-
-	virtual void AddReferencedObjects(class FReferenceCollector& Collector) override;
-};
-
-template<>
-struct TStructOpsTypeTraits< FLayeredMove_Teleport > : public TStructOpsTypeTraitsBase2< FLayeredMove_Teleport >
-{
-	enum
-	{
-		//WithNetSerializer = true,
-		WithCopy = true
-	};
-};
-
-
 /** MoveTo: Move Actor from the starting location to the target location over a duration of time.*/
 USTRUCT(BlueprintType)
 struct MOVER_API FLayeredMove_MoveTo : public FLayeredMoveBase
diff --git a/Engine/Plugins/Experimental/Mover/Source/Mover/Public/MoveLibrary/MovementUtilsTypes.h b/Engine/Plugins/Experimental/Mover/Source/Mover/Public/MoveLibrary/MovementUtilsTypes.h
--- a/Engine/Plugins/Experimental/Mover/Source/Mover/Public/MoveLibrary/MovementUtilsTypes.h
+++ b/Engine/Plugins/Experimental/Mover/Source/Mover/Public/MoveLibrary/MovementUtilsTypes.h
@@ -25,8 +25,7 @@ struct MOVER_API FProposedMove
 	GENERATED_USTRUCT_BODY()
 
 	FProposedMove() : 
-		bHasDirIntent(false),
-		bHasTargetLocation(false)
+		bHasDirIntent(false)
 	{}
 
 	UPROPERTY(EditAnywhere, BlueprintReadWrite, Category = Mover)
@@ -37,8 +36,6 @@ struct MOVER_API FProposedMove
 
 	UPROPERTY(EditAnywhere, BlueprintReadWrite, Category = Mover)
 	uint8	 bHasDirIntent : 1;							// Signals whether there was any directional intent specified
-	UPROPERTY(EditAnywhere, BlueprintReadWrite, Category = Mover)
-	uint8	 bHasTargetLocation : 1;					// Signals whether the proposed move should move to a target location, regardless of other fields
 
 	UPROPERTY(EditAnywhere, BlueprintReadWrite, Category = Mover)
 	FVector  DirectionIntent = FVector::ZeroVector;		// Directional, per-axis magnitude [-1, 1] in world space (length of 1 indicates max speed intent). Only valid if bHasDirIntent is set.
@@ -47,10 +44,4 @@ struct MOVER_API FProposedMove
 	FVector  LinearVelocity = FVector::ZeroVector;		// Units per second, world space, possibly mapped onto walking surface
 	UPROPERTY(EditAnywhere, BlueprintReadWrite, Category = Mover)
 	FRotator AngularVelocity = FRotator::ZeroRotator;	// Degrees per second, local space
-
-	UPROPERTY(EditAnywhere, BlueprintReadWrite, Category = Mover)
-	FVector  MovePlaneVelocity = FVector::ZeroVector;	// Units per second, world space, always along the movement plane
-
-	UPROPERTY(EditAnywhere, BlueprintReadWrite, Category = Mover)
-	FVector  TargetLocation = FVector::ZeroVector;		// World space go-to position. Only valid if bHasTargetLocation is set. Used for movement like teleportation.
 };
\ No newline at end of file
