Add uproperty markup to adapter object ptrs, to allow GC to work correctly

#rnx

diff --git a/Engine/Plugins/Animation/MotionWarping/Source/MotionWarping/Public/MotionWarpingCharacterAdapter.h b/Engine/Plugins/Animation/MotionWarping/Source/MotionWarping/Public/MotionWarpingCharacterAdapter.h
--- a/Engine/Plugins/Animation/MotionWarping/Source/MotionWarping/Public/MotionWarpingCharacterAdapter.h
+++ b/Engine/Plugins/Animation/MotionWarping/Source/MotionWarping/Public/MotionWarpingCharacterAdapter.h
@@ -28,5 +28,6 @@ private:
 	FTransform WarpLocalRootMotionOnCharacter(const FTransform& LocalRootMotionTransform, UCharacterMovementComponent* TargetMoveComp, float DeltaSeconds);
 
 	/** The associated character */
+	UPROPERTY(Transient, DuplicateTransient)
 	TObjectPtr<ACharacter> TargetCharacter;
 };
\ No newline at end of file
diff --git a/Engine/Plugins/Experimental/Mover/Source/Mover/Public/MotionWarpingMoverAdapter.h b/Engine/Plugins/Experimental/Mover/Source/Mover/Public/MotionWarpingMoverAdapter.h
--- a/Engine/Plugins/Experimental/Mover/Source/Mover/Public/MotionWarpingMoverAdapter.h
+++ b/Engine/Plugins/Experimental/Mover/Source/Mover/Public/MotionWarpingMoverAdapter.h
@@ -28,5 +28,6 @@ private:
 	// This is called when our Mover actor wants to warp local motion, and passes the responsibility onto the warping component
 	FTransform WarpLocalRootMotionOnMoverComp(const FTransform& LocalRootMotionTransform, float DeltaSeconds, const FMotionWarpingUpdateContext* OptionalWarpingContext);
 
+	UPROPERTY(Transient, DuplicateTransient)
 	TObjectPtr<UMoverComponent> TargetMoverComp;
 };
\ No newline at end of file
