Mover: added templated convenience function to access the typed MoverComponent from the movement mode

#11795

diff --git a/Engine/Plugins/Experimental/Mover/Source/Mover/Public/MovementMode.h b/Engine/Plugins/Experimental/Mover/Source/Mover/Public/MovementMode.h
--- a/Engine/Plugins/Experimental/Mover/Source/Mover/Public/MovementMode.h
+++ b/Engine/Plugins/Experimental/Mover/Source/Mover/Public/MovementMode.h
@@ -49,6 +49,14 @@ public:
 	UFUNCTION(BlueprintCallable, Category=Mover, meta=(DisplayName="Get Mover Component"))
 	UMoverComponent* GetMoverComponent() const;
 
+	/** Templated convenience version of GetMoverComponent which checks the type is as presumed. */
+	template<class T>
+	T* GetMoverComponentChecked() const
+	{
+		static_assert(TPointerIsConvertibleFromTo<T, const UMoverComponent>::Value, "'T' template parameter to GetMoverComponentChecked must be derived from UMoverComponent");
+		return CastChecked<T>(GetOuter());
+	}
+
 	UFUNCTION(BlueprintCallable, Category=Mover, meta=(DisplayName="Get Blackboard"))
 	const UMoverBlackboard* GetBlackboard() const;
 
