Mover - Add mode activate and deactivate functions

diff --git a/Engine/Plugins/Experimental/Mover/Source/Mover/Private/MovementMode.cpp b/Engine/Plugins/Experimental/Mover/Source/Mover/Private/MovementMode.cpp
--- a/Engine/Plugins/Experimental/Mover/Source/Mover/Private/MovementMode.cpp
+++ b/Engine/Plugins/Experimental/Mover/Source/Mover/Private/MovementMode.cpp
@@ -65,6 +65,30 @@ void UBaseMovementMode::DoSimulationTick(const FSimulationTickParams& Params, FM
 	}
 }
 
+void UBaseMovementMode::DoActivate()
+{
+	if (bHasBlueprintOnActivate)
+	{
+		K2_OnActivate();
+	}
+	else
+	{
+		OnActivate();
+	}
+}
+
+void UBaseMovementMode::DoDeactivate()
+{
+	if (bHasBlueprintOnDeactivate)
+	{
+		K2_OnDeactivate();
+	}
+	else
+	{
+		OnDeactivate();
+	}
+}
+
 
 UMoverComponent* UBaseMovementMode::GetMoverComponent() const
 {
@@ -129,4 +153,12 @@ void UBaseMovementMode::OnSimulationTick(const FSimulationTickParams& Params, FM
 {
 }
 
+void UBaseMovementMode::OnActivate()
+{
+}
+
+void UBaseMovementMode::OnDeactivate()
+{
+}
+
 #undef LOCTEXT_NAMESPACE
diff --git a/Engine/Plugins/Experimental/Mover/Source/Mover/Private/MovementModeStateMachine.cpp b/Engine/Plugins/Experimental/Mover/Source/Mover/Private/MovementModeStateMachine.cpp
--- a/Engine/Plugins/Experimental/Mover/Source/Mover/Private/MovementModeStateMachine.cpp
+++ b/Engine/Plugins/Experimental/Mover/Source/Mover/Private/MovementModeStateMachine.cpp
@@ -406,6 +406,13 @@ void UMovementModeStateMachine::AdvanceToNextMode()
 		const FName PreviousModeName = CurrentModeName;
 		CurrentModeName = NextModeName;
 
+		if (PreviousModeName != NAME_None && Modes.Contains(PreviousModeName))
+		{
+			Modes[PreviousModeName]->DoDeactivate();
+		}
+
+		Modes[CurrentModeName]->DoActivate();
+
 		// signal movement mode change event
 		const UMoverComponent* MoverComp = CastChecked<UMoverComponent>(GetOuter());
 		MoverComp->OnMovementModeChanged.Broadcast(PreviousModeName, NextModeName);
diff --git a/Engine/Plugins/Experimental/Mover/Source/Mover/Public/MovementMode.h b/Engine/Plugins/Experimental/Mover/Source/Mover/Public/MovementMode.h
--- a/Engine/Plugins/Experimental/Mover/Source/Mover/Public/MovementMode.h
+++ b/Engine/Plugins/Experimental/Mover/Source/Mover/Public/MovementMode.h
@@ -29,8 +29,6 @@ public:
 	virtual FString GetDisplayName() const = 0;
 };
 
-
-
 /**
  * Base class for all movement modes, exposing simulation update methods for both C++ and blueprint extension
  */
@@ -45,6 +43,8 @@ public:
 	void DoUnregister();
 	void DoGenerateMove(const FMoverTickStartData& StartState, const FMoverTimeStep& TimeStep, FProposedMove& OutProposedMove) const;
 	void DoSimulationTick(const FSimulationTickParams& Params, FMoverTickEndData& OutputState);
+	void DoActivate();
+	void DoDeactivate();
 
 	UFUNCTION(BlueprintCallable, Category=Mover, meta=(DisplayName="Get Mover Component"))
 	UMoverComponent* GetMoverComponent() const;
@@ -100,9 +100,19 @@ protected:
 	UFUNCTION(BlueprintImplementableEvent, DisplayName = "OnSimulationTick", meta = (ScriptName = "OnSimulationTick"))
 	FMoverTickEndData K2_OnSimulationTick(const FSimulationTickParams& Params);
 
+	virtual void OnActivate();
+
+	UFUNCTION(BlueprintImplementableEvent, DisplayName = "OnActivate", meta = (ScriptName = "OnActivate"))
+	void K2_OnActivate();
+
+	virtual void OnDeactivate();
 
+	UFUNCTION(BlueprintImplementableEvent, DisplayName = "OnDeactivate", meta = (ScriptName = "OnDeactivate"))
+	void K2_OnDeactivate();
 
 private:
 	bool bHasBlueprintGenerateMove = false;
 	bool bHasBlueprintSimulationTick = false;
+	bool bHasBlueprintOnActivate = false;
+	bool bHasBlueprintOnDeactivate = false;
 };
\ No newline at end of file
