Mover: fixed issue where root motion in async simulations could be exaggerated or spit the actor out of the world

diff --git a/Engine/Plugins/Experimental/Mover/Source/Mover/Private/MoverComponent.cpp b/Engine/Plugins/Experimental/Mover/Source/Mover/Private/MoverComponent.cpp
--- a/Engine/Plugins/Experimental/Mover/Source/Mover/Private/MoverComponent.cpp
+++ b/Engine/Plugins/Experimental/Mover/Source/Mover/Private/MoverComponent.cpp
@@ -1111,11 +1111,14 @@ FTransform UMoverComponent::ConvertLocalRootMotionToWorld(const FTransform& Loca
 	// Optionally convert this to be relative to a different space
 	if (AlternateActorToWorld)
 	{
-		const FTransform WorldToActor = GetOwner()->GetTransform().Inverse();
-		const FTransform ActorSpaceRootMotion = WorldSpaceRootMotion * WorldToActor;
-		WorldSpaceRootMotion = ActorSpaceRootMotion * *AlternateActorToWorld;
-	}
+		const FTransform AlternateActorToWorldNoTrans(AlternateActorToWorld->GetRotation(), FVector::ZeroVector, AlternateActorToWorld->GetScale3D());
+
+		FTransform WorldToActorNoTrans(GetOwner()->GetTransform().Inverse());
+		WorldToActorNoTrans.SetTranslation(FVector::ZeroVector);
 
+		const FTransform ActorSpaceRootMotion = WorldSpaceRootMotion * WorldToActorNoTrans;
+		WorldSpaceRootMotion = ActorSpaceRootMotion * AlternateActorToWorldNoTrans;
+	}
 	
 	// Optionally process/warp worldspace root motion
 	return ProcessWorldRootMotionDelegate.IsBound()
