Mover: adding cvar-gated option for logging out detailed anim root motion layered moves

diff --git a/Engine/Plugins/Experimental/Mover/Source/Mover/Private/DefaultMovementSet/LayeredMoves/AnimRootMotionLayeredMove.cpp b/Engine/Plugins/Experimental/Mover/Source/Mover/Private/DefaultMovementSet/LayeredMoves/AnimRootMotionLayeredMove.cpp
--- a/Engine/Plugins/Experimental/Mover/Source/Mover/Private/DefaultMovementSet/LayeredMoves/AnimRootMotionLayeredMove.cpp
+++ b/Engine/Plugins/Experimental/Mover/Source/Mover/Private/DefaultMovementSet/LayeredMoves/AnimRootMotionLayeredMove.cpp
@@ -10,6 +10,13 @@
 #include "MoverLog.h"
 #include "MotionWarpingComponent.h"
 
+#if !UE_BUILD_SHIPPING
+FAutoConsoleVariable CVarLogAnimRootMotionSteps(
+	TEXT("mover.debug.LogAnimRootMotionSteps"),
+	false,
+	TEXT("Whether to log detailed information about anim root motion layered moves. 0: Disable, 1: Enable"),
+	ECVF_Cheat);
+#endif	// !UE_BUILD_SHIPPING
 
 FLayeredMove_AnimRootMotion::FLayeredMove_AnimRootMotion()
 {
@@ -63,6 +70,12 @@ bool FLayeredMove_AnimRootMotion::GenerateMove(const FMoverTickStartData& SimSta
 	OutProposedMove.LinearVelocity    = WorldSpaceRootMotion.GetTranslation() / DeltaSeconds;
 	OutProposedMove.AngularVelocity   = WorldSpaceRootMotion.GetRotation().Rotator() * (1.f / DeltaSeconds);
 
+#if !UE_BUILD_SHIPPING
+	UE_CLOG(CVarLogAnimRootMotionSteps->GetBool(), LogMover, Log, TEXT("AnimRootMotion. SimF %i (dt %.3f) Range [%.3f, %.3f] => LocalT: %s (WST: %s)  Vel: %.3f"),
+	        TimeStep.ServerFrame, DeltaSeconds, ExtractionStartPosition, ExtractionEndPosition, 
+	        *LocalRootMotion.GetTranslation().ToString(), *WorldSpaceRootMotion.GetTranslation().ToString(), OutProposedMove.LinearVelocity.Length());
+#endif // !UE_BUILD_SHIPPING
+
 	return true;
 }
 
