Mover 2.0 adding a bool to Falling Mode to customize landing - actors can now land and cancel vertical velocity or have the vertical velocity redirected

diff --git a/Engine/Plugins/Experimental/Mover/Source/Mover/Private/DefaultMovementSet/Modes/FallingMode.cpp b/Engine/Plugins/Experimental/Mover/Source/Mover/Private/DefaultMovementSet/Modes/FallingMode.cpp
--- a/Engine/Plugins/Experimental/Mover/Source/Mover/Private/DefaultMovementSet/Modes/FallingMode.cpp
+++ b/Engine/Plugins/Experimental/Mover/Source/Mover/Private/DefaultMovementSet/Modes/FallingMode.cpp
@@ -13,14 +13,15 @@
 #include UE_INLINE_GENERATED_CPP_BY_NAME(FallingMode)
 
 UFallingMode::UFallingMode(const FObjectInitializer& ObjectInitializer)
-	: Super(ObjectInitializer),
-	  AirControlPercentage(0.4f),
-	  FallingDeceleration(200.0f),
-	  OverTerminalSpeedFallingDeceleration(800.0f),
-	  TerminalMovementPlaneSpeed(1500.0f),
-	  bShouldClampTerminalVerticalSpeed(true),
-	  VerticalFallingDeceleration(4000.0f),
-	  TerminalVerticalSpeed(2000.0f)
+	: Super(ObjectInitializer)
+	, bCancelVerticalSpeedOnLanding(true)
+	, AirControlPercentage(0.4f)
+	, FallingDeceleration(200.0f)
+	, OverTerminalSpeedFallingDeceleration(800.0f)
+	, TerminalMovementPlaneSpeed(1500.0f)
+	, bShouldClampTerminalVerticalSpeed(true)
+	, VerticalFallingDeceleration(4000.0f)
+	, TerminalVerticalSpeed(2000.0f)
 {
 	SharedSettingsClass = UCommonLegacyMovementSettings::StaticClass();
 }
@@ -251,8 +252,17 @@ void UFallingMode::ProcessLanded(const FFloorCheckResult& FloorResult, FVector&
 	// if we can walk on the floor we landed on
 	if (FloorResult.IsWalkableFloor())
 	{
+		if (bCancelVerticalSpeedOnLanding)
+		{
+			const FPlane MovementPlane(FVector::ZeroVector, GetMoverComponent()->GetUpDirection());
+			Velocity = UMovementUtils::ConstrainToPlane(Velocity, MovementPlane, false);
+		}
+		else
+		{
+			Velocity = FVector::VectorPlaneProject(Velocity, FloorResult.HitResult.Normal);
+		}
+		
 		// Transfer to LandingMovementMode (usually walking), and cache any floor / movement base info
-		Velocity.Z = 0.0;
 		NextMovementMode = CommonLegacySettings->GroundMovementModeName;
 
 		SimBlackboard->Set(CommonBlackboard::LastFloorResult, FloorResult);
diff --git a/Engine/Plugins/Experimental/Mover/Source/Mover/Public/DefaultMovementSet/Modes/FallingMode.h b/Engine/Plugins/Experimental/Mover/Source/Mover/Public/DefaultMovementSet/Modes/FallingMode.h
--- a/Engine/Plugins/Experimental/Mover/Source/Mover/Public/DefaultMovementSet/Modes/FallingMode.h
+++ b/Engine/Plugins/Experimental/Mover/Source/Mover/Public/DefaultMovementSet/Modes/FallingMode.h
@@ -32,6 +32,13 @@ public:
 	UPROPERTY(BlueprintAssignable, Category = Mover)
 	FMover_OnLanded OnLanded;
 	
+	/**
+	 * If true, actor will land and lose all speed in the vertical direction upon landing. If false, actor's vertical speed will be redirected based on the surface normal it hit.
+	 * Note: Actor's horizontal speed will not be affected if true. If false, horizontal speed may be increased on landing.
+	 */
+	UPROPERTY(Category = Mover, EditAnywhere, BlueprintReadWrite)
+	bool bCancelVerticalSpeedOnLanding;
+	
 	/**
 	 * When falling, amount of movement control available to the actor.
 	 * 0 = no control, 1 = full control
