Mover
- Activate network physics component to avoid assert in physics callback
- Refactor network physics component creation/destruction to try and avoid crash in live coding

diff --git a/Engine/Plugins/Experimental/Mover/Source/Mover/Private/Backends/MoverNetworkPhysicsLiaison.cpp b/Engine/Plugins/Experimental/Mover/Source/Mover/Private/Backends/MoverNetworkPhysicsLiaison.cpp
--- a/Engine/Plugins/Experimental/Mover/Source/Mover/Private/Backends/MoverNetworkPhysicsLiaison.cpp
+++ b/Engine/Plugins/Experimental/Mover/Source/Mover/Private/Backends/MoverNetworkPhysicsLiaison.cpp
@@ -234,12 +234,6 @@ UMoverNetworkPhysicsLiaisonComponent::UMoverNetworkPhysicsLiaisonComponent()
 			MyActor->SetReplicatingMovement(true);
 			MyActor->SetReplicateMovement(true);
 		}
-
-		NetworkPhysicsComponent = CreateDefaultSubobject<UNetworkPhysicsComponent>(TEXT("PhysMover_NetworkPhysicsComponent"));
-		NetworkPhysicsComponent->SetNetAddressable(); // Make DSO components net addressable
-		NetworkPhysicsComponent->SetIsReplicated(true);
-		NetworkPhysicsComponent->RegisterComponent();
-		NetworkPhysicsComponent->InitializeComponent();
 	}
 }
 
@@ -386,10 +380,20 @@ void UMoverNetworkPhysicsLiaisonComponent::InitializeComponent()
 		MoverComp->ModeFSM->SetModeImmediately(MoverComp->StartingMovementMode);
 	}
 
-	// Register network data for recording and rewind/resim
-	if (NetworkPhysicsComponent)
+	if (Chaos::FPhysicsSolverBase::IsNetworkPhysicsPredictionEnabled())
 	{
-		NetworkPhysicsComponent->CreateDataHistory<FNetworkPhysicsMoverTraits>(this);
+		NetworkPhysicsComponent = NewObject<UNetworkPhysicsComponent>(GetOwner(), TEXT("PhysMover_NetworkPhysicsComponent"));
+		if (NetworkPhysicsComponent)
+		{
+			NetworkPhysicsComponent->SetNetAddressable(); // Make DSO components net addressable
+			NetworkPhysicsComponent->SetIsReplicated(true);
+			NetworkPhysicsComponent->RegisterComponent();
+			NetworkPhysicsComponent->InitializeComponent();
+			NetworkPhysicsComponent->Activate(true);
+
+			// Register network data for recording and rewind/resim
+			NetworkPhysicsComponent->CreateDataHistory<FNetworkPhysicsMoverTraits>(this);
+		}
 	}
 }
 
@@ -398,6 +402,7 @@ void UMoverNetworkPhysicsLiaisonComponent::UninitializeComponent()
 	if (NetworkPhysicsComponent)
 	{
 		NetworkPhysicsComponent->RemoveDataHistory();
+		NetworkPhysicsComponent->DestroyComponent();
 	}
 
 	Super::UninitializeComponent();
