Mover - Add gameplay tags to movement modes to indicate state

diff --git a/Engine/Plugins/Experimental/Mover/Source/Mover/Mover.Build.cs b/Engine/Plugins/Experimental/Mover/Source/Mover/Mover.Build.cs
--- a/Engine/Plugins/Experimental/Mover/Source/Mover/Mover.Build.cs
+++ b/Engine/Plugins/Experimental/Mover/Source/Mover/Mover.Build.cs
@@ -16,6 +16,7 @@ public class Mover : ModuleRules
 				"AnimGraphRuntime",
 				"MotionWarping",
 				"Water",
+				"GameplayTags"
 			}
 			);
 			
diff --git a/Engine/Plugins/Experimental/Mover/Source/Mover/Private/DefaultMovementSet/CharacterMoverComponent.cpp b/Engine/Plugins/Experimental/Mover/Source/Mover/Private/DefaultMovementSet/CharacterMoverComponent.cpp
--- a/Engine/Plugins/Experimental/Mover/Source/Mover/Private/DefaultMovementSet/CharacterMoverComponent.cpp
+++ b/Engine/Plugins/Experimental/Mover/Source/Mover/Private/DefaultMovementSet/CharacterMoverComponent.cpp
@@ -29,9 +29,9 @@ void UCharacterMoverComponent::BeginPlay()
 
 bool UCharacterMoverComponent::IsFalling() const
 {
-	if (bHasValidCachedState)
+	if (const UBaseMovementMode* Mode = GetMovementMode())
 	{
-		return CachedLastSyncState.MovementMode == DefaultModeNames::Falling;
+		return Mode->GameplayTags.HasTag(Mover_IsFalling);
 	}
 
 	return false;
@@ -39,9 +39,9 @@ bool UCharacterMoverComponent::IsFalling() const
 
 bool UCharacterMoverComponent::IsAirborne() const
 {
-	if (bHasValidCachedState)
+	if (const UBaseMovementMode* Mode = GetMovementMode())
 	{
-		return CachedLastSyncState.MovementMode == DefaultModeNames::Flying || CachedLastSyncState.MovementMode == DefaultModeNames::Falling;
+		return Mode->GameplayTags.HasTag(Mover_IsInAir);
 	}
 
 	return false;
@@ -49,9 +49,9 @@ bool UCharacterMoverComponent::IsAirborne() const
 
 bool UCharacterMoverComponent::IsOnGround() const
 {
-	if (bHasValidCachedState)
+	if (const UBaseMovementMode* Mode = GetMovementMode())
 	{
-		return CachedLastSyncState.MovementMode == DefaultModeNames::Walking;
+		return Mode->GameplayTags.HasTag(Mover_IsOnGround);
 	}
 
 	return false;
@@ -59,9 +59,9 @@ bool UCharacterMoverComponent::IsOnGround() const
 
 bool UCharacterMoverComponent::IsSwimming() const
 {
-	if (bHasValidCachedState)
+	if (const UBaseMovementMode* Mode = GetMovementMode())
 	{
-		return CachedLastSyncState.MovementMode == DefaultModeNames::Swimming;
+		return Mode->GameplayTags.HasTag(Mover_IsSwimming);
 	}
 
 	return false;
diff --git a/Engine/Plugins/Experimental/Mover/Source/Mover/Private/DefaultMovementSet/Modes/FallingMode.cpp b/Engine/Plugins/Experimental/Mover/Source/Mover/Private/DefaultMovementSet/Modes/FallingMode.cpp
--- a/Engine/Plugins/Experimental/Mover/Source/Mover/Private/DefaultMovementSet/Modes/FallingMode.cpp
+++ b/Engine/Plugins/Experimental/Mover/Source/Mover/Private/DefaultMovementSet/Modes/FallingMode.cpp
@@ -24,6 +24,9 @@ UFallingMode::UFallingMode(const FObjectInitializer& ObjectInitializer)
 	, TerminalVerticalSpeed(2000.0f)
 {
 	SharedSettingsClass = UCommonLegacyMovementSettings::StaticClass();
+
+	GameplayTags.AddTag(Mover_IsInAir);
+	GameplayTags.AddTag(Mover_IsFalling);
 }
 
 constexpr float VERTICAL_SLOPE_NORMAL_Z = 0.001f; // Slope is vertical if Abs(Normal.Z) <= this threshold. Accounts for precision problems that sometimes angle normals slightly off horizontal for vertical surface.
diff --git a/Engine/Plugins/Experimental/Mover/Source/Mover/Private/DefaultMovementSet/Modes/FlyingMode.cpp b/Engine/Plugins/Experimental/Mover/Source/Mover/Private/DefaultMovementSet/Modes/FlyingMode.cpp
--- a/Engine/Plugins/Experimental/Mover/Source/Mover/Private/DefaultMovementSet/Modes/FlyingMode.cpp
+++ b/Engine/Plugins/Experimental/Mover/Source/Mover/Private/DefaultMovementSet/Modes/FlyingMode.cpp
@@ -12,6 +12,7 @@
 UFlyingMode::UFlyingMode(const FObjectInitializer& ObjectInitializer)
 	: Super(ObjectInitializer)
 {
+	GameplayTags.AddTag(Mover_IsInAir);
 }
 
 void UFlyingMode::OnGenerateMove(const FMoverTickStartData& StartState, const FMoverTimeStep& TimeStep, FProposedMove& OutProposedMove) const
diff --git a/Engine/Plugins/Experimental/Mover/Source/Mover/Private/DefaultMovementSet/Modes/SwimmingMode.cpp b/Engine/Plugins/Experimental/Mover/Source/Mover/Private/DefaultMovementSet/Modes/SwimmingMode.cpp
--- a/Engine/Plugins/Experimental/Mover/Source/Mover/Private/DefaultMovementSet/Modes/SwimmingMode.cpp
+++ b/Engine/Plugins/Experimental/Mover/Source/Mover/Private/DefaultMovementSet/Modes/SwimmingMode.cpp
@@ -16,6 +16,8 @@ USwimmingMode::USwimmingMode(const FObjectInitializer& ObjectInitializer)
 	: Super(ObjectInitializer)
 {
 	SharedSettingsClass = UCommonLegacyMovementSettings::StaticClass();
+
+	GameplayTags.AddTag(Mover_IsSwimming);
 }
 
 void USwimmingMode::OnGenerateMove(const FMoverTickStartData& StartState, const FMoverTimeStep& TimeStep, FProposedMove& OutProposedMove) const
diff --git a/Engine/Plugins/Experimental/Mover/Source/Mover/Private/DefaultMovementSet/Modes/WalkingMode.cpp b/Engine/Plugins/Experimental/Mover/Source/Mover/Private/DefaultMovementSet/Modes/WalkingMode.cpp
--- a/Engine/Plugins/Experimental/Mover/Source/Mover/Private/DefaultMovementSet/Modes/WalkingMode.cpp
+++ b/Engine/Plugins/Experimental/Mover/Source/Mover/Private/DefaultMovementSet/Modes/WalkingMode.cpp
@@ -20,6 +20,8 @@ UWalkingMode::UWalkingMode(const FObjectInitializer& ObjectInitializer)
 	: Super(ObjectInitializer)
 {
 	SharedSettingsClass = UCommonLegacyMovementSettings::StaticClass();
+
+	GameplayTags.AddTag(Mover_IsOnGround);
 }
 
 void UWalkingMode::OnGenerateMove(const FMoverTickStartData& StartState, const FMoverTimeStep& TimeStep, FProposedMove& OutProposedMove) const
diff --git a/Engine/Plugins/Experimental/Mover/Source/Mover/Private/MoverComponent.cpp b/Engine/Plugins/Experimental/Mover/Source/Mover/Private/MoverComponent.cpp
--- a/Engine/Plugins/Experimental/Mover/Source/Mover/Private/MoverComponent.cpp
+++ b/Engine/Plugins/Experimental/Mover/Source/Mover/Private/MoverComponent.cpp
@@ -1326,6 +1326,18 @@ FName UMoverComponent::GetMovementModeName() const
 	return NAME_None;
 }
 
+const UBaseMovementMode* UMoverComponent::GetMovementMode() const
+{
+	if (bHasValidCachedState)
+	{
+		if (const TObjectPtr<UBaseMovementMode>* CurrentMode = MovementModes.Find(CachedLastSyncState.MovementMode))
+		{
+			return CurrentMode->Get();
+		}
+	}
+
+	return nullptr;
+}
 
 UPrimitiveComponent* UMoverComponent::GetMovementBase() const
 {
diff --git a/Engine/Plugins/Experimental/Mover/Source/Mover/Private/MoverTypes.cpp b/Engine/Plugins/Experimental/Mover/Source/Mover/Private/MoverTypes.cpp
--- a/Engine/Plugins/Experimental/Mover/Source/Mover/Private/MoverTypes.cpp
+++ b/Engine/Plugins/Experimental/Mover/Source/Mover/Private/MoverTypes.cpp
@@ -7,6 +7,11 @@
 
 #define LOCTEXT_NAMESPACE "MoverData"
 
+UE_DEFINE_GAMEPLAY_TAG_COMMENT(Mover_IsOnGround, "Mover.IsOnGround", "Default Mover state flag indicating character is on the ground.")
+UE_DEFINE_GAMEPLAY_TAG_COMMENT(Mover_IsInAir, "Mover.IsInAir", "Default Mover state flag indicating character is in the air.")
+UE_DEFINE_GAMEPLAY_TAG_COMMENT(Mover_IsFalling, "Mover.IsFalling", "Default Mover state flag indicating character is falling.")
+UE_DEFINE_GAMEPLAY_TAG_COMMENT(Mover_IsSwimming, "Mover.IsSwimming", "Default Mover state flag indicating character is swimming.")
+
 FMoverOnImpactParams::FMoverOnImpactParams() 
 	: AttemptedMoveDelta(0) 
 {
diff --git a/Engine/Plugins/Experimental/Mover/Source/Mover/Public/MovementMode.h b/Engine/Plugins/Experimental/Mover/Source/Mover/Public/MovementMode.h
--- a/Engine/Plugins/Experimental/Mover/Source/Mover/Public/MovementMode.h
+++ b/Engine/Plugins/Experimental/Mover/Source/Mover/Public/MovementMode.h
@@ -3,6 +3,7 @@
 #pragma once
 
 #include "CoreMinimal.h"
+#include "GameplayTagContainer.h"
 #include "MoverSimulationTypes.h"
 #include "MoverTypes.h"
 #include "MoveLibrary/MoverBlackboard.h"
@@ -76,6 +77,10 @@ public:
 	UPROPERTY(EditAnywhere, BlueprintReadOnly, Instanced, Category = Mover, meta = (FullyExpand = true))
 	TArray<TObjectPtr<UBaseMovementModeTransition>> Transitions;
 
+	/** A list of gameplay tags associated with this movement mode */
+	UPROPERTY(EditDefaultsOnly, Category = Mover)
+	FGameplayTagContainer GameplayTags;
+
 protected:
 
 	virtual void OnRegistered(const FName ModeName);
diff --git a/Engine/Plugins/Experimental/Mover/Source/Mover/Public/MoverComponent.h b/Engine/Plugins/Experimental/Mover/Source/Mover/Public/MoverComponent.h
--- a/Engine/Plugins/Experimental/Mover/Source/Mover/Public/MoverComponent.h
+++ b/Engine/Plugins/Experimental/Mover/Source/Mover/Public/MoverComponent.h
@@ -29,7 +29,6 @@ namespace MoverComponentConstants
 	extern const FVector DefaultUpDir;				// Fallback up direction if not determined by the component or world (normalized)
 }
 
-
 // Fired just before a simulation tick, regardless of being a re-simulated frame or not.
 DECLARE_DYNAMIC_MULTICAST_DELEGATE_TwoParams(FMover_OnPreSimTick, const FMoverTimeStep&, TimeStep, const FMoverInputCmdContext&, InputCmd);
 
@@ -246,6 +245,10 @@ public:	// Queries
 	UFUNCTION(BlueprintPure, Category = Mover)
 	FName GetMovementModeName() const;
 
+	// Get the current movement mode 
+	UFUNCTION(BlueprintPure, Category = Mover)
+	const UBaseMovementMode* GetMovementMode() const;
+
 	// Get the current movement base. Null if there isn't one.
 	UFUNCTION(BlueprintPure, Category = Mover)
 	UPrimitiveComponent* GetMovementBase() const;
@@ -329,7 +332,7 @@ public:	// Queries
 
 		return nullptr;
 	}
-	
+
 protected:
 
 	/** Makes this component and owner actor reflect the state of a particular frame snapshot. This occurs after simulation ticking, as well as during a rollback before we resimulate forward.
@@ -442,7 +445,6 @@ private:
 	UPROPERTY(Transient)
 	TObjectPtr<UMoverBlackboard> SimBlackboard;
 
-
 	friend class UBaseMovementMode;
 	friend class UMoverNetworkPhysicsLiaisonComponent;
 	friend class UMoverDebugComponent;
diff --git a/Engine/Plugins/Experimental/Mover/Source/Mover/Public/MoverTypes.h b/Engine/Plugins/Experimental/Mover/Source/Mover/Public/MoverTypes.h
--- a/Engine/Plugins/Experimental/Mover/Source/Mover/Public/MoverTypes.h
+++ b/Engine/Plugins/Experimental/Mover/Source/Mover/Public/MoverTypes.h
@@ -7,8 +7,15 @@
 #include "NetworkPredictionTickState.h"
 #include "MoverLog.h"
 #include "Kismet/BlueprintFunctionLibrary.h"
+#include "NativeGameplayTags.h"
 #include "MoverTypes.generated.h"
 
+// Gameplay tags
+UE_DECLARE_GAMEPLAY_TAG_EXTERN(Mover_IsOnGround)
+UE_DECLARE_GAMEPLAY_TAG_EXTERN(Mover_IsInAir)
+UE_DECLARE_GAMEPLAY_TAG_EXTERN(Mover_IsFalling)
+UE_DECLARE_GAMEPLAY_TAG_EXTERN(Mover_IsSwimming)
+
 // Struct to hold params for when an impact happens. This contains all of the data for impacts including what gets passed to the FMover_OnImpact delegate
 USTRUCT(BlueprintType, meta = (DisplayName = "Impact Data"))
 struct MOVER_API FMoverOnImpactParams
