Mover: introduces planar constraint options to limit movement. Includes some changes to movement libraries APIs of how we pass affected components down through various movement functions.

diff --git a/Engine/Plugins/Experimental/Mover/Source/Mover/Private/DefaultMovementSet/Modes/FallingMode.cpp b/Engine/Plugins/Experimental/Mover/Source/Mover/Private/DefaultMovementSet/Modes/FallingMode.cpp
--- a/Engine/Plugins/Experimental/Mover/Source/Mover/Private/DefaultMovementSet/Modes/FallingMode.cpp
+++ b/Engine/Plugins/Experimental/Mover/Source/Mover/Private/DefaultMovementSet/Modes/FallingMode.cpp
@@ -33,6 +33,7 @@ constexpr float VERTICAL_SLOPE_NORMAL_Z = 0.001f; // Slope is vertical if Abs(No
 
 void UFallingMode::OnGenerateMove(const FMoverTickStartData& StartState, const FMoverTimeStep& TimeStep, FProposedMove& OutProposedMove) const
 {
+	const UMoverComponent* MoverComp = GetMoverComponent();
 	const FCharacterDefaultInputs* CharacterInputs = StartState.InputCmd.InputCollection.FindDataByType<FCharacterDefaultInputs>();
 	const FMoverDefaultSyncState* StartingSyncState = StartState.SyncState.SyncStateCollection.FindDataByType<FMoverDefaultSyncState>();
 	check(StartingSyncState);
@@ -48,7 +49,8 @@ void UFallingMode::OnGenerateMove(const FMoverTickStartData& StartState, const F
 	if (CharacterInputs)
 	{
 		Params.MoveInputType = CharacterInputs->GetMoveInputType();
-		Params.MoveInput = CharacterInputs->GetMoveInput();
+		const bool bMaintainInputMagnitude = true;
+		Params.MoveInput = UPlanarConstraintUtils::ConstrainDirectionToPlane(MoverComp->GetPlanarConstraint(), CharacterInputs->GetMoveInput_WorldSpace(), bMaintainInputMagnitude);
 	}
 	else
 	{
@@ -134,8 +136,7 @@ void UFallingMode::OnGenerateMove(const FMoverTickStartData& StartState, const F
 void UFallingMode::OnSimulationTick(const FSimulationTickParams& Params, FMoverTickEndData& OutputState)
 {
 	const FMoverTickStartData& StartState = Params.StartState;
-	USceneComponent* UpdatedComponent = Params.UpdatedComponent;
-	UPrimitiveComponent* UpdatedPrimitive = Params.UpdatedPrimitive;
+	USceneComponent* UpdatedComponent = Params.MovingComps.UpdatedComponent.Get();
 	FProposedMove ProposedMove = Params.ProposedMove;
 
 	const FCharacterDefaultInputs* CharacterInputs = StartState.InputCmd.InputCollection.FindDataByType<FCharacterDefaultInputs>();
@@ -176,7 +177,7 @@ void UFallingMode::OnSimulationTick(const FSimulationTickParams& Params, FMoverT
 	FHitResult Hit(1.f);
 	const FQuat TargetOrientQuat = TargetOrient.Quaternion();
 
-	UMovementUtils::TrySafeMoveUpdatedComponent(UpdatedComponent, UpdatedPrimitive, MoveDelta, TargetOrientQuat, true, Hit, ETeleportType::None, MoveRecord);
+	UMovementUtils::TrySafeMoveUpdatedComponent(Params.MovingComps, MoveDelta, TargetOrientQuat, true, Hit, ETeleportType::None, MoveRecord);
 
 	// Compute final velocity based on how long we actually go until we get a hit.
 	FVector NewFallingVelocity = StartingSyncState->GetVelocity_WorldSpace();
@@ -186,7 +187,7 @@ void UFallingMode::OnSimulationTick(const FSimulationTickParams& Params, FMoverT
 	FFloorCheckResult LandingFloor;
 
 	// Handle impact, whether it's a landing surface or something to slide on
-	if (Hit.IsValidBlockingHit() && UpdatedPrimitive)
+	if (Hit.IsValidBlockingHit() && UpdatedComponent)
 	{
 		float LastMoveTimeSlice = DeltaSeconds;
 		float SubTimeTickRemaining = LastMoveTimeSlice * (1.f - Hit.Time);
@@ -194,7 +195,7 @@ void UFallingMode::OnSimulationTick(const FSimulationTickParams& Params, FMoverT
 		PctTimeApplied += Hit.Time * (1.f - PctTimeApplied);
 
 		// Check for hitting a landing surface
-		if (UAirMovementUtils::IsValidLandingSpot(UpdatedComponent, UpdatedPrimitive, UpdatedPrimitive->GetComponentLocation(), 
+		if (UAirMovementUtils::IsValidLandingSpot(Params.MovingComps, UpdatedComponent->GetComponentLocation(),
 			Hit, CommonLegacySettings->FloorSweepDistance, CommonLegacySettings->MaxWalkSlopeCosine, OUT LandingFloor))
 		{
 			CaptureFinalState(UpdatedComponent, *StartingSyncState, LandingFloor, DeltaSeconds, DeltaSeconds * PctTimeApplied, OutputSyncState, OutputState, MoveRecord);
@@ -209,7 +210,7 @@ void UFallingMode::OnSimulationTick(const FSimulationTickParams& Params, FMoverT
 		MoverComponent->HandleImpact(ImpactParams);
 
 		// We didn't land on a walkable surface, so let's try to slide along it
-		UAirMovementUtils::TryMoveToFallAlongSurface(UpdatedComponent, UpdatedPrimitive, MoverComponent, MoveDelta, 
+		UAirMovementUtils::TryMoveToFallAlongSurface(Params.MovingComps, MoveDelta,
 			(1.f - Hit.Time), TargetOrientQuat, Hit.Normal, Hit, true,
 			CommonLegacySettings->FloorSweepDistance, CommonLegacySettings->MaxWalkSlopeCosine, LandingFloor, MoveRecord);
 
diff --git a/Engine/Plugins/Experimental/Mover/Source/Mover/Private/DefaultMovementSet/Modes/FlyingMode.cpp b/Engine/Plugins/Experimental/Mover/Source/Mover/Private/DefaultMovementSet/Modes/FlyingMode.cpp
--- a/Engine/Plugins/Experimental/Mover/Source/Mover/Private/DefaultMovementSet/Modes/FlyingMode.cpp
+++ b/Engine/Plugins/Experimental/Mover/Source/Mover/Private/DefaultMovementSet/Modes/FlyingMode.cpp
@@ -17,6 +17,7 @@ UFlyingMode::UFlyingMode(const FObjectInitializer& ObjectInitializer)
 
 void UFlyingMode::OnGenerateMove(const FMoverTickStartData& StartState, const FMoverTimeStep& TimeStep, FProposedMove& OutProposedMove) const
 {
+	const UMoverComponent* MoverComp = GetMoverComponent();
 	const FCharacterDefaultInputs* CharacterInputs = StartState.InputCmd.InputCollection.FindDataByType<FCharacterDefaultInputs>();
 	const FMoverDefaultSyncState* StartingSyncState = StartState.SyncState.SyncStateCollection.FindDataByType<FMoverDefaultSyncState>();
 	check(StartingSyncState);
@@ -27,7 +28,8 @@ void UFlyingMode::OnGenerateMove(const FMoverTickStartData& StartState, const FM
 	if (CharacterInputs)
 	{
 		Params.MoveInputType = CharacterInputs->GetMoveInputType();
-		Params.MoveInput = CharacterInputs->GetMoveInput();
+		const bool bMaintainInputMagnitude = true;
+		Params.MoveInput = UPlanarConstraintUtils::ConstrainDirectionToPlane(MoverComp->GetPlanarConstraint(), CharacterInputs->GetMoveInput_WorldSpace(), bMaintainInputMagnitude);
 	}
 	else
 	{
@@ -62,8 +64,7 @@ void UFlyingMode::OnGenerateMove(const FMoverTickStartData& StartState, const FM
 void UFlyingMode::OnSimulationTick(const FSimulationTickParams& Params, FMoverTickEndData& OutputState)
 {
 	const FMoverTickStartData& StartState = Params.StartState;
-	USceneComponent* UpdatedComponent = Params.UpdatedComponent;
-	UPrimitiveComponent* UpdatedPrimitive = Params.UpdatedPrimitive;
+	USceneComponent* UpdatedComponent = Params.MovingComps.UpdatedComponent.Get();
 	FProposedMove ProposedMove = Params.ProposedMove;
 
 	const FCharacterDefaultInputs* CharacterInputs = StartState.InputCmd.InputCollection.FindDataByType<FCharacterDefaultInputs>();
@@ -103,7 +104,7 @@ void UFlyingMode::OnSimulationTick(const FSimulationTickParams& Params, FMoverTi
 
 	if (!MoveDelta.IsNearlyZero() || bIsOrientationChanging)
 	{
-		UMovementUtils::TrySafeMoveUpdatedComponent(UpdatedComponent, UpdatedPrimitive, MoveDelta, OrientQuat, true, Hit, ETeleportType::None, MoveRecord);
+		UMovementUtils::TrySafeMoveUpdatedComponent(Params.MovingComps, MoveDelta, OrientQuat, true, Hit, ETeleportType::None, MoveRecord);
 	}
 
 	if (Hit.IsValidBlockingHit())
@@ -112,7 +113,7 @@ void UFlyingMode::OnSimulationTick(const FSimulationTickParams& Params, FMoverTi
 		FMoverOnImpactParams ImpactParams(DefaultModeNames::Flying, Hit, MoveDelta);
 		MoverComponent->HandleImpact(ImpactParams);
 		// Try to slide the remaining distance along the surface.
-		UMovementUtils::TryMoveToSlideAlongSurface(UpdatedComponent, UpdatedPrimitive, MoverComponent, MoveDelta, 1.f - Hit.Time, OrientQuat, Hit.Normal, Hit, true, MoveRecord);
+		UMovementUtils::TryMoveToSlideAlongSurface(FMovingComponentSet(MoverComponent), MoveDelta, 1.f - Hit.Time, OrientQuat, Hit.Normal, Hit, true, MoveRecord);
 	}
 
 	CaptureFinalState(UpdatedComponent, MoveRecord, *StartingSyncState, OutputSyncState, DeltaSeconds);
diff --git a/Engine/Plugins/Experimental/Mover/Source/Mover/Private/DefaultMovementSet/Modes/SwimmingMode.cpp b/Engine/Plugins/Experimental/Mover/Source/Mover/Private/DefaultMovementSet/Modes/SwimmingMode.cpp
--- a/Engine/Plugins/Experimental/Mover/Source/Mover/Private/DefaultMovementSet/Modes/SwimmingMode.cpp
+++ b/Engine/Plugins/Experimental/Mover/Source/Mover/Private/DefaultMovementSet/Modes/SwimmingMode.cpp
@@ -136,7 +136,9 @@ void USwimmingMode::OnGenerateMove(const FMoverTickStartData& StartState, const
 		if (CharacterInputs)
 		{
 			Params.MoveInputType = CharacterInputs->GetMoveInputType();
-			Params.MoveInput = CharacterInputs->GetMoveInput_WorldSpace();
+
+			const bool bMaintainInputMagnitude = true;
+			Params.MoveInput = UPlanarConstraintUtils::ConstrainDirectionToPlane(MoverComp->GetPlanarConstraint(), CharacterInputs->GetMoveInput_WorldSpace(), bMaintainInputMagnitude);
 		}
 		else
 		{
diff --git a/Engine/Plugins/Experimental/Mover/Source/Mover/Private/DefaultMovementSet/Modes/WalkingMode.cpp b/Engine/Plugins/Experimental/Mover/Source/Mover/Private/DefaultMovementSet/Modes/WalkingMode.cpp
--- a/Engine/Plugins/Experimental/Mover/Source/Mover/Private/DefaultMovementSet/Modes/WalkingMode.cpp
+++ b/Engine/Plugins/Experimental/Mover/Source/Mover/Private/DefaultMovementSet/Modes/WalkingMode.cpp
@@ -66,7 +66,9 @@ void UWalkingMode::OnGenerateMove(const FMoverTickStartData& StartState, const F
 	if (CharacterInputs)
 	{
 		Params.MoveInputType = CharacterInputs->GetMoveInputType();
-		Params.MoveInput = CharacterInputs->GetMoveInput_WorldSpace();
+
+		const bool bMaintainInputMagnitude = true;
+		Params.MoveInput = UPlanarConstraintUtils::ConstrainDirectionToPlane(MoverComp->GetPlanarConstraint(), CharacterInputs->GetMoveInput_WorldSpace(), bMaintainInputMagnitude);
 	}
 	else
 	{
@@ -105,10 +107,10 @@ void UWalkingMode::OnGenerateMove(const FMoverTickStartData& StartState, const F
 
 void UWalkingMode::OnSimulationTick(const FSimulationTickParams& Params, FMoverTickEndData& OutputState)
 {
-	const UMoverComponent* MoverComp = GetMoverComponent();
+	UMoverComponent* MoverComp = GetMoverComponent();
 	const FMoverTickStartData& StartState = Params.StartState;
-	USceneComponent* UpdatedComponent = Params.UpdatedComponent;
-	UPrimitiveComponent* UpdatedPrimitive = Params.UpdatedPrimitive;
+	USceneComponent* UpdatedComponent = Params.MovingComps.UpdatedComponent.Get();
+	UPrimitiveComponent* UpdatedPrimitive = Params.MovingComps.UpdatedPrimitive.Get();
 	FProposedMove ProposedMove = Params.ProposedMove;
 
 	if (!UpdatedComponent || !UpdatedPrimitive)
@@ -175,7 +177,7 @@ void UWalkingMode::OnSimulationTick(const FSimulationTickParams& Params, FMoverT
 	{
 		// Attempt to move the full amount first
 		bDidAttemptMovement = true;
-		bool bMoved = UMovementUtils::TrySafeMoveUpdatedComponent(UpdatedComponent, UpdatedPrimitive, CurMoveDelta, TargetOrientQuat, true, MoveHitResult, ETeleportType::None, MoveRecord);
+		bool bMoved = UMovementUtils::TrySafeMoveUpdatedComponent(Params.MovingComps, CurMoveDelta, TargetOrientQuat, true, MoveHitResult, ETeleportType::None, MoveRecord);
 		float LastMoveSeconds = DeltaSeconds;
 
 		if (MoveHitResult.bStartPenetrating)
@@ -198,7 +200,7 @@ void UWalkingMode::OnSimulationTick(const FSimulationTickParams& Params, FMoverT
 				// It's a walkable ramp, so cut up the move and attempt to move the remainder of it along the ramp's surface, possibly generating another hit
 				const float PercentTimeRemaining = 1.f - PercentTimeAppliedSoFar;
 				CurMoveDelta = UGroundMovementUtils::ComputeDeflectedMoveOntoRamp(CurMoveDelta * PercentTimeRemaining, MoveHitResult, CommonLegacySettings->MaxWalkSlopeCosine, CurrentFloor.bLineTrace);
-				UMovementUtils::TrySafeMoveUpdatedComponent(UpdatedComponent, UpdatedPrimitive, CurMoveDelta, TargetOrientQuat, true, MoveHitResult, ETeleportType::None, MoveRecord);
+				UMovementUtils::TrySafeMoveUpdatedComponent(Params.MovingComps, CurMoveDelta, TargetOrientQuat, true, MoveHitResult, ETeleportType::None, MoveRecord);
 				LastMoveSeconds = PercentTimeRemaining * LastMoveSeconds;
 
 				const float SecondHitPercent = MoveHitResult.Time * PercentTimeRemaining;
@@ -209,7 +211,6 @@ void UWalkingMode::OnSimulationTick(const FSimulationTickParams& Params, FMoverT
 			{
 				// If still blocked, try to step up onto the blocking object OR slide along it
 
-				UMoverComponent* MoverComponent = GetMoverComponent();
 				// JAH TODO: Take movement bases into account
 				if (UGroundMovementUtils::CanStepUpOnHitSurface(MoveHitResult)) // || (CharacterOwner->GetMovementBase() != nullptr && Hit.HitObjectHandle == CharacterOwner->GetMovementBase()->GetOwner()))
 				{
@@ -217,21 +218,21 @@ void UWalkingMode::OnSimulationTick(const FSimulationTickParams& Params, FMoverT
 					const FVector PreStepUpLocation = UpdatedComponent->GetComponentLocation();
 					const FVector DownwardDir = -MoverComp->GetUpDirection();
 
-					if (!UGroundMovementUtils::TryMoveToStepUp(UpdatedComponent, UpdatedPrimitive, MoverComponent, DownwardDir, CommonLegacySettings->MaxStepHeight, CommonLegacySettings->MaxWalkSlopeCosine, CommonLegacySettings->FloorSweepDistance, OrigMoveDelta * (1.f - PercentTimeAppliedSoFar), MoveHitResult, CurrentFloor, false, &StepUpFloorResult, MoveRecord))
+					if (!UGroundMovementUtils::TryMoveToStepUp(Params.MovingComps, DownwardDir, CommonLegacySettings->MaxStepHeight, CommonLegacySettings->MaxWalkSlopeCosine, CommonLegacySettings->FloorSweepDistance, OrigMoveDelta * (1.f - PercentTimeAppliedSoFar), MoveHitResult, CurrentFloor, false, &StepUpFloorResult, MoveRecord))
 					{
 						FMoverOnImpactParams ImpactParams(DefaultModeNames::Walking, MoveHitResult, OrigMoveDelta);
-						MoverComponent->HandleImpact(ImpactParams);
+						MoverComp->HandleImpact(ImpactParams);
 						float PercentAvailableToSlide = 1.f - PercentTimeAppliedSoFar;
-						float SlideAmount = UGroundMovementUtils::TryWalkToSlideAlongSurface(UpdatedComponent, UpdatedPrimitive, MoverComponent, OrigMoveDelta, PercentAvailableToSlide, TargetOrientQuat, MoveHitResult.Normal, MoveHitResult, true, MoveRecord, CommonLegacySettings->MaxWalkSlopeCosine, CommonLegacySettings->MaxStepHeight);
+						float SlideAmount = UGroundMovementUtils::TryWalkToSlideAlongSurface(Params.MovingComps, OrigMoveDelta, PercentAvailableToSlide, TargetOrientQuat, MoveHitResult.Normal, MoveHitResult, true, MoveRecord, CommonLegacySettings->MaxWalkSlopeCosine, CommonLegacySettings->MaxStepHeight);
 						PercentTimeAppliedSoFar += PercentAvailableToSlide * SlideAmount;
 					}
 				}
 				else if (MoveHitResult.Component.IsValid() && !MoveHitResult.Component.Get()->CanCharacterStepUp(Cast<APawn>(MoveHitResult.GetActor())))
 				{
 					FMoverOnImpactParams ImpactParams(DefaultModeNames::Walking, MoveHitResult, OrigMoveDelta);
-					MoverComponent->HandleImpact(ImpactParams);
+					MoverComp->HandleImpact(ImpactParams);
 					float PercentAvailableToSlide = 1.f - PercentTimeAppliedSoFar;
-					float SlideAmount = UGroundMovementUtils::TryWalkToSlideAlongSurface(UpdatedComponent, UpdatedPrimitive, MoverComponent, OrigMoveDelta, 1.f - PercentTimeAppliedSoFar, TargetOrientQuat, MoveHitResult.Normal, MoveHitResult, true, MoveRecord, CommonLegacySettings->MaxWalkSlopeCosine, CommonLegacySettings->MaxStepHeight);
+					float SlideAmount = UGroundMovementUtils::TryWalkToSlideAlongSurface(Params.MovingComps, OrigMoveDelta, 1.f - PercentTimeAppliedSoFar, TargetOrientQuat, MoveHitResult.Normal, MoveHitResult, true, MoveRecord, CommonLegacySettings->MaxWalkSlopeCosine, CommonLegacySettings->MaxStepHeight);
 					PercentTimeAppliedSoFar += PercentAvailableToSlide * SlideAmount;
 				}
 			}
@@ -244,7 +245,7 @@ void UWalkingMode::OnSimulationTick(const FSimulationTickParams& Params, FMoverT
 
 		if (CurrentFloor.IsWalkableFloor())
 		{
-			UGroundMovementUtils::TryMoveToAdjustHeightAboveFloor(UpdatedComponent, UpdatedPrimitive, CurrentFloor, CommonLegacySettings->MaxWalkSlopeCosine, MoveRecord);
+			UGroundMovementUtils::TryMoveToAdjustHeightAboveFloor(MoverComp, CurrentFloor, CommonLegacySettings->MaxWalkSlopeCosine, MoveRecord);
 		}
     
 		if (!CurrentFloor.IsWalkableFloor() && !CurrentFloor.HitResult.bStartPenetrating)
@@ -275,7 +276,7 @@ void UWalkingMode::OnSimulationTick(const FSimulationTickParams& Params, FMoverT
 			const EMoveComponentFlags IncludeBlockingOverlapsWithoutEvents = (MOVECOMP_NeverIgnoreBlockingOverlaps | MOVECOMP_DisableBlockingOverlapDispatch);
 			EMoveComponentFlags MoveComponentFlags = MOVECOMP_NoFlags;
 			MoveComponentFlags = (MoveComponentFlags | IncludeBlockingOverlapsWithoutEvents);
-			UMovementUtils::TryMoveToResolvePenetration(UpdatedComponent, UpdatedPrimitive, MoveComponentFlags, RequestedAdjustment, Hit, UpdatedComponent->GetComponentQuat(), MoveRecord);
+			UMovementUtils::TryMoveToResolvePenetration(Params.MovingComps, MoveComponentFlags, RequestedAdjustment, Hit, UpdatedComponent->GetComponentQuat(), MoveRecord);
 		}
 		
 		if (!CurrentFloor.IsWalkableFloor() && !Hit.bStartPenetrating)
diff --git a/Engine/Plugins/Experimental/Mover/Source/Mover/Private/MoveLibrary/AirMovementUtils.cpp b/Engine/Plugins/Experimental/Mover/Source/Mover/Private/MoveLibrary/AirMovementUtils.cpp
--- a/Engine/Plugins/Experimental/Mover/Source/Mover/Private/MoveLibrary/AirMovementUtils.cpp
+++ b/Engine/Plugins/Experimental/Mover/Source/Mover/Private/MoveLibrary/AirMovementUtils.cpp
@@ -33,7 +33,7 @@ FProposedMove UAirMovementUtils::ComputeControlledFreeMove(const FFreeMoveParams
 	return OutMove;
 }
 
-bool UAirMovementUtils::IsValidLandingSpot(USceneComponent* UpdatedComponent, UPrimitiveComponent* UpdatedPrimitive, const FVector& Location, const FHitResult& Hit, float FloorSweepDistance, float WalkableFloorZ, FFloorCheckResult& OutFloorResult)
+bool UAirMovementUtils::IsValidLandingSpot(const FMovingComponentSet& MovingComps, const FVector& Location, const FHitResult& Hit, float FloorSweepDistance, float WalkableFloorZ, FFloorCheckResult& OutFloorResult)
 {
 	OutFloorResult.Clear();
 
@@ -54,7 +54,7 @@ bool UAirMovementUtils::IsValidLandingSpot(USceneComponent* UpdatedComponent, UP
 	}
 
 	// Make sure floor test passes here.
-	UFloorQueryUtils::FindFloor(UpdatedComponent, UpdatedPrimitive, 
+	UFloorQueryUtils::FindFloor(MovingComps.UpdatedComponent.Get(), MovingComps.UpdatedPrimitive.Get(),
 		FloorSweepDistance, WalkableFloorZ,
 		Location, OutFloorResult);
 
@@ -66,7 +66,7 @@ bool UAirMovementUtils::IsValidLandingSpot(USceneComponent* UpdatedComponent, UP
 	return true;
 }
 
-float UAirMovementUtils::TryMoveToFallAlongSurface(USceneComponent* UpdatedComponent, UPrimitiveComponent* UpdatedPrimitive, UMoverComponent* MoverComponent, const FVector& Delta, float PctOfDeltaToMove, const FQuat Rotation, const FVector& Normal, FHitResult& Hit, bool bHandleImpact, float FloorSweepDistance, float MaxWalkSlopeCosine, FFloorCheckResult& OutFloorResult, FMovementRecord& MoveRecord)
+float UAirMovementUtils::TryMoveToFallAlongSurface(const FMovingComponentSet& MovingComps, const FVector& Delta, float PctOfDeltaToMove, const FQuat Rotation, const FVector& Normal, FHitResult& Hit, bool bHandleImpact, float FloorSweepDistance, float MaxWalkSlopeCosine, FFloorCheckResult& OutFloorResult, FMovementRecord& MoveRecord)
 {
 	OutFloorResult.Clear();
 
@@ -78,16 +78,19 @@ float UAirMovementUtils::TryMoveToFallAlongSurface(USceneComponent* UpdatedCompo
 	float PctOfTimeUsed = 0.f;
 	const FVector OldHitNormal = Normal;
 
-	FVector SlideDelta = UMovementUtils::ComputeSlideDelta(Delta, PctOfDeltaToMove, Normal, Hit);
+	FVector SlideDelta = UMovementUtils::ComputeSlideDelta(MovingComps, Delta, PctOfDeltaToMove, Normal, Hit);
 
 	if ((SlideDelta | Delta) > 0.f)
 	{
 		// First sliding attempt along surface
-		UMovementUtils::TrySafeMoveUpdatedComponent(UpdatedComponent, UpdatedPrimitive, SlideDelta, Rotation, true, Hit, ETeleportType::None, MoveRecord);
+		UMovementUtils::TrySafeMoveUpdatedComponent(MovingComps, SlideDelta, Rotation, true, Hit, ETeleportType::None, MoveRecord);
 
 		PctOfTimeUsed = Hit.Time;
 		if (Hit.IsValidBlockingHit())
 		{
+			UMoverComponent* MoverComponent = MovingComps.MoverComponent.Get();
+			UPrimitiveComponent* UpdatedPrimitive = MovingComps.UpdatedPrimitive.Get();
+
 			// Notify first impact
 			if (MoverComponent && bHandleImpact)
 			{
@@ -96,19 +99,19 @@ float UAirMovementUtils::TryMoveToFallAlongSurface(USceneComponent* UpdatedCompo
 			}
 
 			// Check if we landed
-			if (!IsValidLandingSpot(UpdatedComponent, UpdatedPrimitive, UpdatedPrimitive->GetComponentLocation(),
+			if (!IsValidLandingSpot(MovingComps, UpdatedPrimitive->GetComponentLocation(),
 				Hit, FloorSweepDistance, MaxWalkSlopeCosine, OutFloorResult))
 			{
 				// We've hit another surface during our first move, so let's try to slide along both of them together
 
 				// Compute new slide normal when hitting multiple surfaces.
-				SlideDelta = UMovementUtils::ComputeTwoWallAdjustedDelta(SlideDelta, Hit, OldHitNormal);
+				SlideDelta = UMovementUtils::ComputeTwoWallAdjustedDelta(MovingComps, SlideDelta, Hit, OldHitNormal);
 
 				// Only proceed if the new direction is of significant length and not in reverse of original attempted move.
 				if (!SlideDelta.IsNearlyZero(UE::MoverUtils::SMALL_MOVE_DISTANCE) && (SlideDelta | Delta) > 0.f)
 				{
 					// Perform second move, taking 2 walls into account
-					UMovementUtils::TrySafeMoveUpdatedComponent(UpdatedComponent, UpdatedPrimitive, SlideDelta, Rotation, true, Hit, ETeleportType::None, MoveRecord);
+					UMovementUtils::TrySafeMoveUpdatedComponent(MovingComps, SlideDelta, Rotation, true, Hit, ETeleportType::None, MoveRecord);
 					PctOfTimeUsed += (Hit.Time * (1.f - PctOfTimeUsed));
 
 					// Notify second impact
@@ -119,7 +122,7 @@ float UAirMovementUtils::TryMoveToFallAlongSurface(USceneComponent* UpdatedCompo
 					}
 
 					// Check if we've landed, to acquire floor result
-					IsValidLandingSpot(UpdatedComponent, UpdatedPrimitive, UpdatedPrimitive->GetComponentLocation(),
+					IsValidLandingSpot(MovingComps, UpdatedPrimitive->GetComponentLocation(),
 						Hit, FloorSweepDistance, MaxWalkSlopeCosine, OutFloorResult);
 				}
 			}
diff --git a/Engine/Plugins/Experimental/Mover/Source/Mover/Private/MoveLibrary/BasedMovementUtils.cpp b/Engine/Plugins/Experimental/Mover/Source/Mover/Private/MoveLibrary/BasedMovementUtils.cpp
--- a/Engine/Plugins/Experimental/Mover/Source/Mover/Private/MoveLibrary/BasedMovementUtils.cpp
+++ b/Engine/Plugins/Experimental/Mover/Source/Mover/Private/MoveLibrary/BasedMovementUtils.cpp
@@ -408,7 +408,7 @@ void UBasedMovementUtils::UpdateSimpleBasedMovement(UMoverComponent* TargetMover
 					const bool bSweep = true;
 					FHitResult MoveHitResult;
 
-					bool bDidMove = UMovementUtils::TryMoveUpdatedComponent_Internal(UpdatedComponent, WorldDeltaLocation, WorldTargetQuat, bSweep, MoveComponentFlags, &MoveHitResult, ETeleportType::None);
+					bool bDidMove = UMovementUtils::TryMoveUpdatedComponent_Internal(FMovingComponentSet(TargetMoverComp), WorldDeltaLocation, WorldTargetQuat, bSweep, MoveComponentFlags, &MoveHitResult, ETeleportType::None);
 					
 					const FVector NewWorldLocation = UpdatedComponent->GetComponentLocation();
 
diff --git a/Engine/Plugins/Experimental/Mover/Source/Mover/Private/MoveLibrary/ConstrainedMoveUtils.cpp b/Engine/Plugins/Experimental/Mover/Source/Mover/Private/MoveLibrary/ConstrainedMoveUtils.cpp
new file mode 100644
--- /dev/null
+++ b/Engine/Plugins/Experimental/Mover/Source/Mover/Private/MoveLibrary/ConstrainedMoveUtils.cpp
@@ -0,0 +1,69 @@
+// Copyright Epic Games, Inc. All Rights Reserved.
+
+#include "MoveLibrary/ConstrainedMoveUtils.h"
+#include "MoverLog.h"
+
+
+void UPlanarConstraintUtils::SetPlanarConstraintEnabled(UPARAM(ref) FPlanarConstraint& Constraint, bool bEnabled)
+{
+	Constraint.bConstrainToPlane = bEnabled;
+}
+
+
+void UPlanarConstraintUtils::SetPlanarConstraintNormal(UPARAM(ref) FPlanarConstraint& Constraint, FVector PlaneNormal)
+{
+	PlaneNormal = PlaneNormal.GetSafeNormal();
+
+	if (PlaneNormal.IsNearlyZero())
+	{
+		UE_LOG(LogMover, Warning, TEXT("Can't use SetPlanarConstraintNormal with a zero-length normal. Leaving normal as %s"), *Constraint.PlaneConstraintNormal.ToCompactString());
+	}
+	else
+	{
+		Constraint.PlaneConstraintNormal = PlaneNormal;
+	}
+}
+
+
+void UPlanarConstraintUtils::SetPlaneConstraintOrigin(UPARAM(ref) FPlanarConstraint& Constraint, FVector PlaneOrigin)
+{
+	Constraint.PlaneConstraintOrigin = PlaneOrigin;
+}
+
+
+FVector UPlanarConstraintUtils::ConstrainDirectionToPlane(const FPlanarConstraint& Constraint, FVector Direction, bool bMaintainMagnitude)
+{
+	if (Constraint.bConstrainToPlane)
+	{
+		float OrigSize = Direction.Size();
+
+		Direction = FVector::VectorPlaneProject(Direction, Constraint.PlaneConstraintNormal);
+
+		if (bMaintainMagnitude)
+		{
+			Direction = Direction.GetSafeNormal() * OrigSize;
+		}
+	}
+
+	return Direction;
+}
+
+FVector UPlanarConstraintUtils::ConstrainLocationToPlane(const FPlanarConstraint& Constraint, FVector Location)
+{
+	if (Constraint.bConstrainToPlane)
+	{
+		Location = FVector::PointPlaneProject(Location, Constraint.PlaneConstraintOrigin, Constraint.PlaneConstraintNormal);
+	}
+
+	return Location;
+}
+
+FVector UPlanarConstraintUtils::ConstrainNormalToPlane(const FPlanarConstraint& Constraint, FVector Normal)
+{
+	if (Constraint.bConstrainToPlane)
+	{
+		Normal = FVector::VectorPlaneProject(Normal, Constraint.PlaneConstraintNormal).GetSafeNormal();
+	}
+
+	return Normal;
+}
\ No newline at end of file
diff --git a/Engine/Plugins/Experimental/Mover/Source/Mover/Private/MoveLibrary/GroundMovementUtils.cpp b/Engine/Plugins/Experimental/Mover/Source/Mover/Private/MoveLibrary/GroundMovementUtils.cpp
--- a/Engine/Plugins/Experimental/Mover/Source/Mover/Private/MoveLibrary/GroundMovementUtils.cpp
+++ b/Engine/Plugins/Experimental/Mover/Source/Mover/Private/MoveLibrary/GroundMovementUtils.cpp
@@ -49,9 +49,9 @@ static const FName StepFwdSubstepName = "StepFwd";
 static const FName StepDownSubstepName = "StepDown";
 static const FName SlideSubstepName = "SlideFromStep";
 
-bool UGroundMovementUtils::TryMoveToStepUp(USceneComponent* UpdatedComponent, UPrimitiveComponent* UpdatedPrimitive, UMoverComponent* MoverComponent, const FVector& GravDir, float MaxStepHeight, float MaxWalkSlopeCosine, float FloorSweepDistance, const FVector& MoveDelta, const FHitResult& MoveHitResult, const FFloorCheckResult& CurrentFloor, bool bIsFalling, FOptionalFloorCheckResult* OutFloorTestResult, FMovementRecord& MoveRecord)
+bool UGroundMovementUtils::TryMoveToStepUp(const FMovingComponentSet& MovingComps, const FVector& GravDir, float MaxStepHeight, float MaxWalkSlopeCosine, float FloorSweepDistance, const FVector& MoveDelta, const FHitResult& MoveHitResult, const FFloorCheckResult& CurrentFloor, bool bIsFalling, FOptionalFloorCheckResult* OutFloorTestResult, FMovementRecord& MoveRecord)
 {
-	UCapsuleComponent* CapsuleComponent = Cast<UCapsuleComponent>(UpdatedPrimitive);
+	UCapsuleComponent* CapsuleComponent = Cast<UCapsuleComponent>(MovingComps.UpdatedPrimitive.Get());
 
 	if (CapsuleComponent == nullptr || !CanStepUpOnHitSurface(MoveHitResult) || MaxStepHeight <= 0.f)
 	{
@@ -61,7 +61,7 @@ bool UGroundMovementUtils::TryMoveToStepUp(USceneComponent* UpdatedComponent, UP
 	TArray<FMovementSubstep> QueuedSubsteps;	// keeping track of substeps before committing, because some moves can be backed out
 
 
-	const FVector OldLocation = UpdatedPrimitive->GetComponentLocation();
+	const FVector OldLocation = CapsuleComponent->GetComponentLocation();
 	FVector LastComponentLocation = OldLocation;
 
 	float PawnRadius, PawnHalfHeight;
@@ -121,17 +121,17 @@ bool UGroundMovementUtils::TryMoveToStepUp(USceneComponent* UpdatedComponent, UP
 	}
 
 	// Scope our movement updates, and do not apply them until all intermediate moves are completed.
-	FScopedMovementUpdate ScopedStepUpMovement(UpdatedComponent, EScopedUpdate::DeferredUpdates);
+	FScopedMovementUpdate ScopedStepUpMovement(CapsuleComponent, EScopedUpdate::DeferredUpdates);
 
 	// step up - treat as vertical wall
 	FHitResult SweepUpHit(1.f);
-	const FQuat PawnRotation = UpdatedComponent->GetComponentQuat();
+	const FQuat PawnRotation = CapsuleComponent->GetComponentQuat();
 
 	const FVector UpAdjustment = -GravDir * StepTravelUpHeight;
-	const bool bDidStepUp = UMovementUtils::TryMoveUpdatedComponent_Internal(UpdatedComponent, UpAdjustment, PawnRotation, true, MOVECOMP_NoFlags, &SweepUpHit, ETeleportType::None);
+	const bool bDidStepUp = UMovementUtils::TryMoveUpdatedComponent_Internal(MovingComps, UpAdjustment, PawnRotation, true, MOVECOMP_NoFlags, &SweepUpHit, ETeleportType::None);
 
 	UE_LOG(LogMover, VeryVerbose, TEXT("TryMoveToStepUp Up: %s (role %i) UpAdjustment=%s DidMove=%i"),
-		*GetNameSafe(UpdatedComponent->GetOwner()), UpdatedComponent->GetOwnerRole(), *UpAdjustment.ToCompactString(), bDidStepUp);
+		*GetNameSafe(CapsuleComponent->GetOwner()), CapsuleComponent->GetOwnerRole(), *UpAdjustment.ToCompactString(), bDidStepUp);
 
 	if (SweepUpHit.bStartPenetrating)
 	{
@@ -142,15 +142,15 @@ bool UGroundMovementUtils::TryMoveToStepUp(USceneComponent* UpdatedComponent, UP
 	}
 
 	// Cache upwards substep
-	QueuedSubsteps.Add(FMovementSubstep(StepUpSubstepName, UpdatedPrimitive->GetComponentLocation()-LastComponentLocation, false));
-	LastComponentLocation = UpdatedPrimitive->GetComponentLocation();
+	QueuedSubsteps.Add(FMovementSubstep(StepUpSubstepName, CapsuleComponent->GetComponentLocation()-LastComponentLocation, false));
+	LastComponentLocation = CapsuleComponent->GetComponentLocation();
 
 	// step fwd
 	FHitResult StepFwdHit(1.f);
-	const bool bDidStepFwd = UMovementUtils::TryMoveUpdatedComponent_Internal(UpdatedComponent, MoveDelta, PawnRotation, true, MOVECOMP_NoFlags, &StepFwdHit, ETeleportType::None);
+	const bool bDidStepFwd = UMovementUtils::TryMoveUpdatedComponent_Internal(MovingComps, MoveDelta, PawnRotation, true, MOVECOMP_NoFlags, &StepFwdHit, ETeleportType::None);
 
 	UE_LOG(LogMover, VeryVerbose, TEXT("TryMoveToStepUp Fwd: %s (role %i) MoveDelta=%s DidMove=%i"),
-		*GetNameSafe(UpdatedComponent->GetOwner()), UpdatedComponent->GetOwnerRole(), *MoveDelta.ToCompactString(), bDidStepFwd);
+		*GetNameSafe(CapsuleComponent->GetOwner()), CapsuleComponent->GetOwnerRole(), *MoveDelta.ToCompactString(), bDidStepFwd);
 
 	// Check result of forward movement
 	if (StepFwdHit.bBlockingHit)
@@ -163,25 +163,26 @@ bool UGroundMovementUtils::TryMoveToStepUp(USceneComponent* UpdatedComponent, UP
 			return false;
 		}
 
+		UMoverComponent* MoverComp = MovingComps.MoverComponent.Get();
 		// If we hit something above us and also something ahead of us, we should notify about the upward hit as well.
 		// The forward hit will be handled later (in the bSteppedOver case below).
 		// In the case of hitting something above but not forward, we are not blocked from moving so we don't need the notification.
-		if (MoverComponent && SweepUpHit.bBlockingHit && StepFwdHit.bBlockingHit)
+		if (MoverComp && SweepUpHit.bBlockingHit && StepFwdHit.bBlockingHit)
 		{
 			FMoverOnImpactParams ImpactParams(NAME_None, SweepUpHit, MoveDelta);
-			MoverComponent->HandleImpact(ImpactParams);
+			MoverComp->HandleImpact(ImpactParams);
 		}
 
 		// pawn ran into a wall
-		if (MoverComponent)
+		if (MoverComp)
 		{
 			FMoverOnImpactParams ImpactParams(NAME_None, StepFwdHit, MoveDelta);
-			MoverComponent->HandleImpact(ImpactParams);
+			MoverComp->HandleImpact(ImpactParams);
 		}
 		
 		if (bIsFalling)
 		{
-			QueuedSubsteps.Add( FMovementSubstep(StepFwdSubstepName, UpdatedComponent->GetComponentLocation()-LastComponentLocation, true) );
+			QueuedSubsteps.Add( FMovementSubstep(StepFwdSubstepName, CapsuleComponent->GetComponentLocation()-LastComponentLocation, true) );
 
 			// Commit queued substeps to movement record
 			for (FMovementSubstep Substep : QueuedSubsteps)
@@ -193,17 +194,17 @@ bool UGroundMovementUtils::TryMoveToStepUp(USceneComponent* UpdatedComponent, UP
 		}
 
 		// Cache forwards substep before the slide attempt
-		QueuedSubsteps.Add(FMovementSubstep(StepFwdSubstepName, UpdatedPrimitive->GetComponentLocation() - LastComponentLocation, true));
-		LastComponentLocation = UpdatedPrimitive->GetComponentLocation();
+		QueuedSubsteps.Add(FMovementSubstep(StepFwdSubstepName, CapsuleComponent->GetComponentLocation() - LastComponentLocation, true));
+		LastComponentLocation = CapsuleComponent->GetComponentLocation();
 
 		// adjust and try again
 		const float ForwardHitTime = StepFwdHit.Time;
 
 		// locking relevancy so velocity isn't added until it is needed to (adding it to the QueuedSubsteps so it can get added later)
 		MoveRecord.LockRelevancy(false);
-		const float ForwardSlideAmount = TryWalkToSlideAlongSurface(UpdatedComponent, UpdatedPrimitive, MoverComponent, MoveDelta, 1.f - StepFwdHit.Time, PawnRotation, StepFwdHit.Normal, StepFwdHit, true, MoveRecord, MaxWalkSlopeCosine, MaxStepHeight);
-		QueuedSubsteps.Add( FMovementSubstep(SlideSubstepName, UpdatedComponent->GetComponentLocation()-LastComponentLocation, true) );
-		LastComponentLocation = UpdatedPrimitive->GetComponentLocation();
+		const float ForwardSlideAmount = TryWalkToSlideAlongSurface(MovingComps, MoveDelta, 1.f - StepFwdHit.Time, PawnRotation, StepFwdHit.Normal, StepFwdHit, true, MoveRecord, MaxWalkSlopeCosine, MaxStepHeight);
+		QueuedSubsteps.Add( FMovementSubstep(SlideSubstepName, CapsuleComponent->GetComponentLocation()-LastComponentLocation, true) );
+		LastComponentLocation = CapsuleComponent->GetComponentLocation();
 		MoveRecord.UnlockRelevancy();
 
 		if (bIsFalling)
@@ -224,17 +225,17 @@ bool UGroundMovementUtils::TryMoveToStepUp(USceneComponent* UpdatedComponent, UP
 	else
 	{
 		// Our forward move attempt was unobstructed - cache it
-		QueuedSubsteps.Add(FMovementSubstep(StepFwdSubstepName, UpdatedPrimitive->GetComponentLocation() - LastComponentLocation, true));
-		LastComponentLocation = UpdatedPrimitive->GetComponentLocation();
+		QueuedSubsteps.Add(FMovementSubstep(StepFwdSubstepName, CapsuleComponent->GetComponentLocation() - LastComponentLocation, true));
+		LastComponentLocation = CapsuleComponent->GetComponentLocation();
 	}
 
 
 	// Step down
 	const FVector StepDownAdjustment = GravDir * StepTravelDownHeight;
-	const bool bDidStepDown = UMovementUtils::TryMoveUpdatedComponent_Internal(UpdatedComponent, StepDownAdjustment, UpdatedComponent->GetComponentQuat(), true, MOVECOMP_NoFlags, &StepFwdHit, ETeleportType::None);
+	const bool bDidStepDown = UMovementUtils::TryMoveUpdatedComponent_Internal(MovingComps, StepDownAdjustment, CapsuleComponent->GetComponentQuat(), true, MOVECOMP_NoFlags, &StepFwdHit, ETeleportType::None);
 
 	UE_LOG(LogMover, VeryVerbose, TEXT("TryMoveToStepUp Down: %s (role %i) StepDownAdjustment=%s DidMove=%i"),
-		*GetNameSafe(UpdatedComponent->GetOwner()), UpdatedComponent->GetOwnerRole(), *StepDownAdjustment.ToCompactString(), bDidStepDown);
+		*GetNameSafe(CapsuleComponent->GetOwner()), CapsuleComponent->GetOwnerRole(), *StepDownAdjustment.ToCompactString(), bDidStepDown);
 
 
 	// If step down was initially penetrating abort the step up
@@ -299,9 +300,9 @@ bool UGroundMovementUtils::TryMoveToStepUp(USceneComponent* UpdatedComponent, UP
 		if (OutFloorTestResult != NULL)
 		{
 
-			UFloorQueryUtils::FindFloor(UpdatedComponent, UpdatedPrimitive,
+			UFloorQueryUtils::FindFloor(CapsuleComponent, CapsuleComponent,
 				FloorSweepDistance, MaxWalkSlopeCosine,
-				UpdatedComponent->GetComponentLocation(), StepDownResult.FloorTestResult);
+				CapsuleComponent->GetComponentLocation(), StepDownResult.FloorTestResult);
 
 			// Reject unwalkable normals if we end up higher than our initial height.
 			// It's fine to walk down onto an unwalkable surface, don't reject those moves.
@@ -323,8 +324,8 @@ bool UGroundMovementUtils::TryMoveToStepUp(USceneComponent* UpdatedComponent, UP
 	}
 
 	// Cache downwards substep
-	QueuedSubsteps.Add(FMovementSubstep(StepDownSubstepName, UpdatedPrimitive->GetComponentLocation() - LastComponentLocation, false));
-	LastComponentLocation = UpdatedPrimitive->GetComponentLocation();
+	QueuedSubsteps.Add(FMovementSubstep(StepDownSubstepName, CapsuleComponent->GetComponentLocation() - LastComponentLocation, false));
+	LastComponentLocation = CapsuleComponent->GetComponentLocation();
 
 	// Copy step down result.
 	if (OutFloorTestResult != NULL)
@@ -345,7 +346,7 @@ bool UGroundMovementUtils::TryMoveToStepUp(USceneComponent* UpdatedComponent, UP
 
 }
 
-bool UGroundMovementUtils::TryMoveToAdjustHeightAboveFloor(USceneComponent* UpdatedComponent, UPrimitiveComponent* UpdatedPrimitive, FFloorCheckResult& CurrentFloor, float MaxWalkSlopeCosine, FMovementRecord& MoveRecord)
+bool UGroundMovementUtils::TryMoveToAdjustHeightAboveFloor(const FMovingComponentSet& MovingComps, FFloorCheckResult& CurrentFloor, float MaxWalkSlopeCosine, FMovementRecord& MoveRecord)
 {
 	// If we have a floor check that hasn't hit anything, don't adjust height.
 	if (!CurrentFloor.IsWalkableFloor())
@@ -365,6 +366,8 @@ bool UGroundMovementUtils::TryMoveToAdjustHeightAboveFloor(USceneComponent* Upda
 		OldFloorDist = CurrentFloor.LineDist;
 	}
 
+	USceneComponent* UpdatedComponent = MovingComps.UpdatedComponent.Get();
+
 	// Move up or down to maintain floor height.
 	if (OldFloorDist < UE::FloorQueryUtility::MIN_FLOOR_DIST || OldFloorDist > UE::FloorQueryUtility::MAX_FLOOR_DIST)
 	{
@@ -375,7 +378,7 @@ bool UGroundMovementUtils::TryMoveToAdjustHeightAboveFloor(USceneComponent* Upda
 
 		MoveRecord.LockRelevancy(false);
 
-		UMovementUtils::TrySafeMoveUpdatedComponent(UpdatedComponent, UpdatedPrimitive,
+		UMovementUtils::TrySafeMoveUpdatedComponent(MovingComps,
 			FVector(0.f, 0.f, MoveDist), UpdatedComponent->GetComponentQuat(), 
 			true, AdjustHit, ETeleportType::None, MoveRecord);
 
@@ -407,7 +410,7 @@ bool UGroundMovementUtils::TryMoveToAdjustHeightAboveFloor(USceneComponent* Upda
 	return false;
 }
 
-float UGroundMovementUtils::TryWalkToSlideAlongSurface(USceneComponent* UpdatedComponent, UPrimitiveComponent* UpdatedPrimitive, UMoverComponent* MovementComponent, const FVector& Delta, float PctOfDeltaToMove, const FQuat Rotation, const FVector& Normal, FHitResult& Hit, bool bHandleImpact, FMovementRecord& MoveRecord, float MaxWalkSlopeCosine, float MaxStepHeight)
+float UGroundMovementUtils::TryWalkToSlideAlongSurface(const FMovingComponentSet& MovingComps, const FVector& Delta, float PctOfDeltaToMove, const FQuat Rotation, const FVector& Normal, FHitResult& Hit, bool bHandleImpact, FMovementRecord& MoveRecord, float MaxWalkSlopeCosine, float MaxStepHeight)
 {
 	if (!Hit.bBlockingHit)
 	{
@@ -415,6 +418,7 @@ float UGroundMovementUtils::TryWalkToSlideAlongSurface(USceneComponent* UpdatedC
 	}
 	
 	FVector SafeWalkNormal(Normal);
+	UMoverComponent* MoverComponent = MovingComps.MoverComponent.Get();
 
 	// We don't want to be pushed up an unwalkable surface.
 	if (SafeWalkNormal.Z > 0.f && !UFloorQueryUtils::IsHitSurfaceWalkable(Hit, MaxWalkSlopeCosine))
@@ -425,25 +429,25 @@ float UGroundMovementUtils::TryWalkToSlideAlongSurface(USceneComponent* UpdatedC
 	float PctOfTimeUsed = 0.f;
 	const FVector OldSafeHitNormal = SafeWalkNormal;
 
-	FVector SlideDelta = UMovementUtils::ComputeSlideDelta(Delta, PctOfDeltaToMove, SafeWalkNormal, Hit);
+	FVector SlideDelta = UMovementUtils::ComputeSlideDelta(MovingComps, Delta, PctOfDeltaToMove, SafeWalkNormal, Hit);
 
 	if (SlideDelta.Dot(Delta) > 0.f)
 	{
-		UMovementUtils::TrySafeMoveUpdatedComponent(UpdatedComponent, UpdatedPrimitive, SlideDelta, Rotation, true, Hit, ETeleportType::None, MoveRecord);
+		UMovementUtils::TrySafeMoveUpdatedComponent(MovingComps, SlideDelta, Rotation, true, Hit, ETeleportType::None, MoveRecord);
 
 		PctOfTimeUsed = Hit.Time;
 
 		if (Hit.IsValidBlockingHit())
 		{
 			// Notify first impact
-			if (MovementComponent && bHandleImpact)
+			if (MoverComponent && bHandleImpact)
 			{
 				FMoverOnImpactParams ImpactParams(NAME_None, Hit, SlideDelta);
-				MovementComponent->HandleImpact(ImpactParams);
+				MoverComponent->HandleImpact(ImpactParams);
 			}
 
 			// Compute new slide normal when hitting multiple surfaces.
-			SlideDelta = UMovementUtils::ComputeTwoWallAdjustedDelta(SlideDelta, Hit, OldSafeHitNormal);
+			SlideDelta = UMovementUtils::ComputeTwoWallAdjustedDelta(MovingComps, SlideDelta, Hit, OldSafeHitNormal);
 			if (SlideDelta.Z > 0.f && UFloorQueryUtils::IsHitSurfaceWalkable(Hit, MaxWalkSlopeCosine) && Hit.Normal.Z > UE_KINDA_SMALL_NUMBER)
 			{
 				// Maintain horizontal velocity
@@ -467,14 +471,14 @@ float UGroundMovementUtils::TryWalkToSlideAlongSurface(USceneComponent* UpdatedC
 			if (!SlideDelta.IsNearlyZero(UE::MoverUtils::SMALL_MOVE_DISTANCE) && (SlideDelta | Delta) > 0.f)
 			{
 				// Perform second move
-				UMovementUtils::TrySafeMoveUpdatedComponent(UpdatedComponent, UpdatedPrimitive, SlideDelta, Rotation, true, Hit, ETeleportType::None, MoveRecord);
+				UMovementUtils::TrySafeMoveUpdatedComponent(MovingComps, SlideDelta, Rotation, true, Hit, ETeleportType::None, MoveRecord);
 				PctOfTimeUsed += (Hit.Time * (1.f - PctOfTimeUsed));
 
 				// Notify second impact
-				if (MovementComponent && bHandleImpact && Hit.bBlockingHit)
+				if (MoverComponent && bHandleImpact && Hit.bBlockingHit)
 				{
 					FMoverOnImpactParams ImpactParams(NAME_None, Hit, SlideDelta);
-					MovementComponent->HandleImpact(ImpactParams);
+					MoverComponent->HandleImpact(ImpactParams);
 				}
 			}
 		}
diff --git a/Engine/Plugins/Experimental/Mover/Source/Mover/Private/MoveLibrary/MovementUtils.cpp b/Engine/Plugins/Experimental/Mover/Source/Mover/Private/MoveLibrary/MovementUtils.cpp
--- a/Engine/Plugins/Experimental/Mover/Source/Mover/Private/MoveLibrary/MovementUtils.cpp
+++ b/Engine/Plugins/Experimental/Mover/Source/Mover/Private/MoveLibrary/MovementUtils.cpp
@@ -154,12 +154,19 @@ FVector UMovementUtils::ConstrainToPlane(const FVector& Vector, const FPlane& Mo
 	return ConstrainedResult;
 }
 
-FVector UMovementUtils::ComputeSlideDelta(const FVector& Delta, const float PctOfDeltaToMove, const FVector& Normal, const FHitResult& Hit)
+FVector UMovementUtils::ComputeSlideDelta(const FMovingComponentSet& MovingComps, const FVector& Delta, const float PctOfDeltaToMove, const FVector& Normal, const FHitResult& Hit)
 {
-	return FVector::VectorPlaneProject(Delta, Normal) * PctOfDeltaToMove;
+	FVector ConstrainedNormal = Normal;
+
+	if (MovingComps.MoverComponent.IsValid())
+	{
+		ConstrainedNormal = UPlanarConstraintUtils::ConstrainNormalToPlane(MovingComps.MoverComponent->GetPlanarConstraint(), Normal);
+	}
+
+	return FVector::VectorPlaneProject(Delta, ConstrainedNormal) * PctOfDeltaToMove;
 }
 
-FVector UMovementUtils::ComputeTwoWallAdjustedDelta(const FVector& MoveDelta, const FHitResult& Hit, const FVector& OldHitNormal)
+FVector UMovementUtils::ComputeTwoWallAdjustedDelta(const FMovingComponentSet& MovingComps, const FVector& MoveDelta, const FHitResult& Hit, const FVector& OldHitNormal)
 {
 	FVector Delta = MoveDelta;
 	const FVector HitNormal = Hit.Normal;
@@ -178,7 +185,7 @@ FVector UMovementUtils::ComputeTwoWallAdjustedDelta(const FVector& MoveDelta, co
 	else //adjust to new wall
 	{
 		const FVector DesiredDir = Delta;
-		Delta = UMovementUtils::ComputeSlideDelta(Delta, 1.f - Hit.Time, HitNormal, Hit);
+		Delta = UMovementUtils::ComputeSlideDelta(MovingComps, Delta, 1.f - Hit.Time, HitNormal, Hit);
 		if ((Delta | DesiredDir) <= 0.f)
 		{
 			Delta = FVector::ZeroVector;
@@ -194,7 +201,7 @@ FVector UMovementUtils::ComputeTwoWallAdjustedDelta(const FVector& MoveDelta, co
 	return Delta;
 }
 
-float UMovementUtils::TryMoveToSlideAlongSurface(USceneComponent* UpdatedComponent, UPrimitiveComponent* UpdatedPrimitive, UMoverComponent* MoverComponent, const FVector& Delta, float PctOfDeltaToMove, const FQuat Rotation, const FVector& Normal, FHitResult& Hit, bool bHandleImpact, FMovementRecord& MoveRecord)
+float UMovementUtils::TryMoveToSlideAlongSurface(const FMovingComponentSet& MovingComps, const FVector& Delta, float PctOfDeltaToMove, const FQuat Rotation, const FVector& Normal, FHitResult& Hit, bool bHandleImpact, FMovementRecord& MoveRecord)
 {
 	if (!Hit.bBlockingHit)
 	{
@@ -204,38 +211,38 @@ float UMovementUtils::TryMoveToSlideAlongSurface(USceneComponent* UpdatedCompone
 	float PctOfTimeUsed = 0.f;
 	const FVector OldHitNormal = Normal;
 
-	FVector SlideDelta = UMovementUtils::ComputeSlideDelta(Delta, PctOfDeltaToMove, Normal, Hit);
+	FVector SlideDelta = UMovementUtils::ComputeSlideDelta(MovingComps, Delta, PctOfDeltaToMove, Normal, Hit);
 
 	if ((SlideDelta | Delta) > 0.f)
 	{
-		TrySafeMoveUpdatedComponent(UpdatedComponent, UpdatedPrimitive, SlideDelta, Rotation, true, Hit, ETeleportType::None, MoveRecord);
+		TrySafeMoveUpdatedComponent(MovingComps, SlideDelta, Rotation, true, Hit, ETeleportType::None, MoveRecord);
 
 		PctOfTimeUsed = Hit.Time;
 
 		if (Hit.IsValidBlockingHit())
 		{
 			// Notify first impact
-			if (MoverComponent && bHandleImpact)
+			if (MovingComps.MoverComponent.IsValid() && bHandleImpact)
 			{
 				FMoverOnImpactParams ImpactParams(NAME_None, Hit, SlideDelta);
-				MoverComponent->HandleImpact(ImpactParams);
+				MovingComps.MoverComponent->HandleImpact(ImpactParams);
 			}
 
 			// Compute new slide normal when hitting multiple surfaces.
-			SlideDelta = UMovementUtils::ComputeTwoWallAdjustedDelta(SlideDelta, Hit, OldHitNormal);
+			SlideDelta = UMovementUtils::ComputeTwoWallAdjustedDelta(MovingComps, SlideDelta, Hit, OldHitNormal);
 
 			// Only proceed if the new direction is of significant length and not in reverse of original attempted move.
 			if (!SlideDelta.IsNearlyZero(UE::MoverUtils::SMALL_MOVE_DISTANCE) && (SlideDelta | Delta) > 0.f)
 			{
 				// Perform second move
-				TrySafeMoveUpdatedComponent(UpdatedComponent, UpdatedPrimitive, SlideDelta, Rotation, true, Hit, ETeleportType::None, MoveRecord);
+				TrySafeMoveUpdatedComponent(MovingComps, SlideDelta, Rotation, true, Hit, ETeleportType::None, MoveRecord);
 				PctOfTimeUsed += (Hit.Time * (1.f - PctOfTimeUsed));
 
 				// Notify second impact
-				if (MoverComponent && bHandleImpact && Hit.bBlockingHit)
+				if (MovingComps.MoverComponent.IsValid() && bHandleImpact && Hit.bBlockingHit)
 				{
 					FMoverOnImpactParams ImpactParams(NAME_None, Hit, SlideDelta);
-					MoverComponent->HandleImpact(ImpactParams);
+					MovingComps.MoverComponent->HandleImpact(ImpactParams);
 				}
 			}
 		}
@@ -249,9 +256,11 @@ float UMovementUtils::TryMoveToSlideAlongSurface(USceneComponent* UpdatedCompone
 
 static const FName SafeMoveSubstepName = "SafeMove";
 
-bool UMovementUtils::TrySafeMoveUpdatedComponent(USceneComponent* UpdatedComponent, UPrimitiveComponent* UpdatedPrimitive, const FVector& Delta, const FQuat& NewRotation, bool bSweep, FHitResult& OutHit, ETeleportType Teleport, FMovementRecord& MoveRecord)
+bool UMovementUtils::TrySafeMoveUpdatedComponent(const FMovingComponentSet& MovingComps, const FVector& Delta, const FQuat& NewRotation, bool bSweep, FHitResult& OutHit, ETeleportType Teleport, FMovementRecord& MoveRecord)
 {
-	if (UpdatedComponent == nullptr)
+	USceneComponent* UpdatedComponent = MovingComps.UpdatedComponent.Get();
+
+	if (!UpdatedComponent)
 	{
 		OutHit.Reset(1.f);
 		return false;
@@ -268,7 +277,7 @@ bool UMovementUtils::TrySafeMoveUpdatedComponent(USceneComponent* UpdatedCompone
 		const EMoveComponentFlags IncludeBlockingOverlapsWithoutEvents = (MOVECOMP_NeverIgnoreBlockingOverlaps | MOVECOMP_DisableBlockingOverlapDispatch);
 		//TGuardValue<EMoveComponentFlags> ScopedFlagRestore(MoveComponentFlags, MovementComponentCVars::MoveIgnoreFirstBlockingOverlap ? MoveComponentFlags : (MoveComponentFlags | IncludeBlockingOverlapsWithoutEvents));
 		MoveComponentFlags = (MoveComponentFlags | IncludeBlockingOverlapsWithoutEvents);
-		bMoveResult = TryMoveUpdatedComponent_Internal(UpdatedComponent, Delta, NewRotation, bSweep, MoveComponentFlags, &OutHit, Teleport);
+		bMoveResult = TryMoveUpdatedComponent_Internal(MovingComps, Delta, NewRotation, bSweep, MoveComponentFlags, &OutHit, Teleport);
 
 		if (UpdatedComponent)
 		{
@@ -281,12 +290,12 @@ bool UMovementUtils::TrySafeMoveUpdatedComponent(USceneComponent* UpdatedCompone
 	if (OutHit.bStartPenetrating && UpdatedComponent)
 	{
 		const FVector RequestedAdjustment = ComputePenetrationAdjustment(OutHit);
-		if (TryMoveToResolvePenetration(UpdatedComponent, UpdatedPrimitive, MoveComponentFlags, RequestedAdjustment, OutHit, NewRotation, MoveRecord))
+		if (TryMoveToResolvePenetration(MovingComps, MoveComponentFlags, RequestedAdjustment, OutHit, NewRotation, MoveRecord))
 		{
 			PreviousCompPos = UpdatedComponent->GetComponentLocation();
 
 			// Retry original move
-			bMoveResult = TryMoveUpdatedComponent_Internal(UpdatedComponent, Delta, NewRotation, bSweep, MoveComponentFlags, &OutHit, Teleport);
+			bMoveResult = TryMoveUpdatedComponent_Internal(MovingComps, Delta, NewRotation, bSweep, MoveComponentFlags, &OutHit, Teleport);
 
 			UE_LOG(LogMover, VeryVerbose, TEXT("TrySafeMove retry: %s (role %i) Delta=%s DidMove=%i"),
 				*GetNameSafe(UpdatedComponent->GetOwner()), UpdatedComponent->GetOwnerRole(), *Delta.ToCompactString(), bMoveResult);
@@ -321,15 +330,19 @@ FVector UMovementUtils::ComputePenetrationAdjustment(const FHitResult& Hit)
 
 static const FName PenetrationResolutionSubstepName = "ResolvePenetration";
 
-bool UMovementUtils::TryMoveToResolvePenetration(USceneComponent* UpdatedComponent, UPrimitiveComponent* UpdatedPrimitive, EMoveComponentFlags MoveComponentFlags, const FVector& ProposedAdjustment, const FHitResult& Hit, const FQuat& NewRotationQuat, FMovementRecord& MoveRecord)
+bool UMovementUtils::TryMoveToResolvePenetration(const FMovingComponentSet& MovingComps, EMoveComponentFlags MoveComponentFlags, const FVector& ProposedAdjustment, const FHitResult& Hit, const FQuat& NewRotationQuat, FMovementRecord& MoveRecord)
 {
+	USceneComponent* UpdatedComponent = MovingComps.UpdatedComponent.Get();
+	UPrimitiveComponent* UpdatedPrimitive = Cast<UPrimitiveComponent>(UpdatedComponent);
+	UMoverComponent* MoverComp = MovingComps.MoverComponent.Get();
+
 	// SceneComponent can't be in penetration, so this function really only applies to PrimitiveComponent.
-	const FVector Adjustment = ProposedAdjustment; //ConstrainDirectionToPlane(ProposedAdjustment);
+	const FVector Adjustment = UPlanarConstraintUtils::ConstrainDirectionToPlane(MovingComps.MoverComponent->GetPlanarConstraint(), ProposedAdjustment);
 	if (!Adjustment.IsZero() && UpdatedPrimitive)
 	{
 		QUICK_SCOPE_CYCLE_COUNTER(STAT_BaseMovementComponent_ResolvePenetration);
 		// See if we can fit at the adjusted location without overlapping anything.
-		AActor* ActorOwner = UpdatedComponent->GetOwner();
+		AActor* ActorOwner = MoverComp->GetOwner();
 		if (!ActorOwner)
 		{
 			return false;
@@ -359,7 +372,7 @@ bool UMovementUtils::TryMoveToResolvePenetration(USceneComponent* UpdatedCompone
 
 			// Try sweeping as far as possible...
 			FHitResult SweepOutHit(1.f);
-			bool bMoved = TryMoveUpdatedComponent_Internal(UpdatedComponent, Adjustment, NewRotationQuat, true, MoveComponentFlags, &SweepOutHit, ETeleportType::TeleportPhysics);
+			bool bMoved = TryMoveUpdatedComponent_Internal(MovingComps, Adjustment, NewRotationQuat, true, MoveComponentFlags, &SweepOutHit, ETeleportType::TeleportPhysics);
 
 			UE_LOG(LogMover, VeryVerbose, TEXT("TryMoveToResolvePenetration: %s (role %i) Adjustment=%s DidMove=%i"),
 				*GetNameSafe(UpdatedComponent->GetOwner()), UpdatedComponent->GetOwnerRole(), *Adjustment.ToCompactString(), bMoved);
@@ -372,7 +385,7 @@ bool UMovementUtils::TryMoveToResolvePenetration(USceneComponent* UpdatedCompone
 				const FVector CombinedMTD = Adjustment + SecondMTD;
 				if (SecondMTD != Adjustment && !CombinedMTD.IsZero())
 				{
-					bMoved = TryMoveUpdatedComponent_Internal(UpdatedComponent, CombinedMTD, NewRotationQuat, true, MoveComponentFlags, nullptr, ETeleportType::TeleportPhysics);
+					bMoved = TryMoveUpdatedComponent_Internal(MovingComps, CombinedMTD, NewRotationQuat, true, MoveComponentFlags, nullptr, ETeleportType::TeleportPhysics);
 
 					UE_LOG(LogMover, VeryVerbose, TEXT("TryMoveToResolvePenetration combined: %s (role %i) CombinedAdjustment=%s DidMove=%i"),
 						*GetNameSafe(UpdatedComponent->GetOwner()), UpdatedComponent->GetOwnerRole(), *CombinedMTD.ToCompactString(), bMoved);
@@ -383,11 +396,11 @@ bool UMovementUtils::TryMoveToResolvePenetration(USceneComponent* UpdatedCompone
 			if (!bMoved)
 			{
 				// Try moving the proposed adjustment plus the attempted move direction. This can sometimes get out of penetrations with multiple objects
-				const FVector MoveDelta = (Hit.TraceEnd - Hit.TraceStart); //ConstrainDirectionToPlane(Hit.TraceEnd - Hit.TraceStart);
+				const FVector MoveDelta = UPlanarConstraintUtils::ConstrainDirectionToPlane(MoverComp->GetPlanarConstraint(), (Hit.TraceEnd - Hit.TraceStart));
 				if (!MoveDelta.IsZero())
 				{
 					const FVector AdjustAndMoveDelta = Adjustment + MoveDelta;
-					bMoved = TryMoveUpdatedComponent_Internal(UpdatedComponent, AdjustAndMoveDelta, NewRotationQuat, true, MoveComponentFlags, nullptr, ETeleportType::TeleportPhysics);
+					bMoved = TryMoveUpdatedComponent_Internal(MovingComps, AdjustAndMoveDelta, NewRotationQuat, true, MoveComponentFlags, nullptr, ETeleportType::TeleportPhysics);
 
 					UE_LOG(LogMover, VeryVerbose, TEXT("TryMoveToResolvePenetration multiple: %s (role %i) AdjustAndMoveDelta=%s DidMove=%i"),
 						*GetNameSafe(UpdatedComponent->GetOwner()), UpdatedComponent->GetOwnerRole(), *AdjustAndMoveDelta.ToCompactString(), bMoved);
@@ -479,14 +492,18 @@ FVector UMovementUtils::ComputeDirectionIntent(const FVector& MoveInput, EMoveIn
 	return ResultDirIntent;
 }
 
-bool UMovementUtils::TryMoveUpdatedComponent_Internal(USceneComponent* UpdatedComponent, const FVector& Delta, const FQuat& NewRotation, bool bSweep, EMoveComponentFlags MoveComponentFlags, FHitResult* OutHit, ETeleportType Teleport)
+bool UMovementUtils::TryMoveUpdatedComponent_Internal(const FMovingComponentSet& MovingComps, const FVector& Delta, const FQuat& NewRotation, bool bSweep, EMoveComponentFlags MoveComponentFlags, FHitResult* OutHit, ETeleportType Teleport)
 {
-
-	if (UpdatedComponent)
+	if (MovingComps.UpdatedComponent.IsValid())
 	{
-		// TODO: Consider whether we need NewDelta... seems pointless in this context
-		const FVector NewDelta = Delta;
-		return UpdatedComponent->MoveComponent(NewDelta, NewRotation, bSweep, OutHit, MoveComponentFlags, Teleport);
+		FVector ConstrainedDelta = Delta;
+
+		if (MovingComps.MoverComponent.IsValid())
+		{
+			ConstrainedDelta = UPlanarConstraintUtils::ConstrainDirectionToPlane(MovingComps.MoverComponent->GetPlanarConstraint(), Delta);	
+		}
+
+		return MovingComps.UpdatedComponent->MoveComponent(ConstrainedDelta, NewRotation, bSweep, OutHit, MoveComponentFlags, Teleport);
 	}
 
 	return false;
diff --git a/Engine/Plugins/Experimental/Mover/Source/Mover/Private/MoveLibrary/MovementUtilsTypes.cpp b/Engine/Plugins/Experimental/Mover/Source/Mover/Private/MoveLibrary/MovementUtilsTypes.cpp
new file mode 100644
--- /dev/null
+++ b/Engine/Plugins/Experimental/Mover/Source/Mover/Private/MoveLibrary/MovementUtilsTypes.cpp
@@ -0,0 +1,31 @@
+// Copyright Epic Games, Inc. All Rights Reserved.
+
+#include "MoveLibrary/MovementUtilsTypes.h"
+#include "Components/PrimitiveComponent.h"
+#include "MoverComponent.h"
+
+
+
+FMovingComponentSet::FMovingComponentSet(USceneComponent* InUpdatedComponent)
+{
+	UpdatedComponent = InUpdatedComponent;
+
+	if (UpdatedComponent.IsValid())
+	{
+		UpdatedPrimitive = Cast<UPrimitiveComponent>(UpdatedComponent);
+		MoverComponent = UpdatedComponent->GetOwner()->FindComponentByClass<UMoverComponent>();
+
+		checkf(!MoverComponent.IsValid() || UpdatedComponent == MoverComponent->GetUpdatedComponent(), TEXT("Expected MoverComponent to have the same UpdatedComponent"));
+	}
+}
+
+FMovingComponentSet::FMovingComponentSet(UMoverComponent* InMoverComponent)
+{
+	MoverComponent = InMoverComponent;
+
+	if (MoverComponent.IsValid())
+	{
+		UpdatedComponent = MoverComponent->GetUpdatedComponent();
+		UpdatedPrimitive = Cast<UPrimitiveComponent>(UpdatedComponent);
+	}
+}
diff --git a/Engine/Plugins/Experimental/Mover/Source/Mover/Private/MovementModeStateMachine.cpp b/Engine/Plugins/Experimental/Mover/Source/Mover/Private/MovementModeStateMachine.cpp
--- a/Engine/Plugins/Experimental/Mover/Source/Mover/Private/MovementModeStateMachine.cpp
+++ b/Engine/Plugins/Experimental/Mover/Source/Mover/Private/MovementModeStateMachine.cpp
@@ -245,9 +245,7 @@ void UMovementModeStateMachine::OnSimulationTick(USceneComponent* UpdatedCompone
 			{
 				FSimulationTickParams SimTickParams;
 				SimTickParams.StartState = SubstepStartData;
-				SimTickParams.UpdatedComponent = UpdatedComponent;
-				SimTickParams.UpdatedPrimitive = UpdatedPrimitive;
-				SimTickParams.MoverComponent = MoverComp;
+				SimTickParams.MovingComps = FMovingComponentSet(MoverComp);
 				SimTickParams.TimeStep = SubTimeStep;
 				SimTickParams.ProposedMove = CombinedMove;
 
diff --git a/Engine/Plugins/Experimental/Mover/Source/Mover/Private/MoverComponent.cpp b/Engine/Plugins/Experimental/Mover/Source/Mover/Private/MoverComponent.cpp
--- a/Engine/Plugins/Experimental/Mover/Source/Mover/Private/MoverComponent.cpp
+++ b/Engine/Plugins/Experimental/Mover/Source/Mover/Private/MoverComponent.cpp
@@ -1245,6 +1245,15 @@ FVector UMoverComponent::GetUpDirection() const
 	return DeducedUpDir;
 }
 
+const FPlanarConstraint& UMoverComponent::GetPlanarConstraint() const
+{
+	return PlanarConstraint;
+}
+
+void UMoverComponent::SetPlanarConstraint(const FPlanarConstraint& InConstraint)
+{
+	PlanarConstraint = InConstraint;
+}
 
 TArray<FTrajectorySampleInfo> UMoverComponent::GetFutureTrajectory(float FutureSeconds, float SamplesPerSecond) const
 {
diff --git a/Engine/Plugins/Experimental/Mover/Source/Mover/Private/PhysicsMover/Modes/PhysicsDrivenFallingMode.cpp b/Engine/Plugins/Experimental/Mover/Source/Mover/Private/PhysicsMover/Modes/PhysicsDrivenFallingMode.cpp
--- a/Engine/Plugins/Experimental/Mover/Source/Mover/Private/PhysicsMover/Modes/PhysicsDrivenFallingMode.cpp
+++ b/Engine/Plugins/Experimental/Mover/Source/Mover/Private/PhysicsMover/Modes/PhysicsDrivenFallingMode.cpp
@@ -53,8 +53,8 @@ EDataValidationResult UPhysicsDrivenFallingMode::IsDataValid(FDataValidationCont
 void UPhysicsDrivenFallingMode::OnSimulationTick(const FSimulationTickParams& Params, FMoverTickEndData& OutputState)
 {
 	const FMoverTickStartData& StartState = Params.StartState;
-	USceneComponent* UpdatedComponent = Params.UpdatedComponent;
-	UPrimitiveComponent* UpdatedPrimitive = Params.UpdatedPrimitive;
+	USceneComponent* UpdatedComponent = Params.MovingComps.UpdatedComponent.Get();
+	UPrimitiveComponent* UpdatedPrimitive = Params.MovingComps.UpdatedPrimitive.Get();
 	FProposedMove ProposedMove = Params.ProposedMove;
 
 	const FCharacterDefaultInputs* CharacterInputs = StartState.InputCmd.InputCollection.FindDataByType<FCharacterDefaultInputs>();
diff --git a/Engine/Plugins/Experimental/Mover/Source/Mover/Private/PhysicsMover/Modes/PhysicsDrivenFlyingMode.cpp b/Engine/Plugins/Experimental/Mover/Source/Mover/Private/PhysicsMover/Modes/PhysicsDrivenFlyingMode.cpp
--- a/Engine/Plugins/Experimental/Mover/Source/Mover/Private/PhysicsMover/Modes/PhysicsDrivenFlyingMode.cpp
+++ b/Engine/Plugins/Experimental/Mover/Source/Mover/Private/PhysicsMover/Modes/PhysicsDrivenFlyingMode.cpp
@@ -50,8 +50,8 @@ EDataValidationResult UPhysicsDrivenFlyingMode::IsDataValid(FDataValidationConte
 void UPhysicsDrivenFlyingMode::OnSimulationTick(const FSimulationTickParams& Params, FMoverTickEndData& OutputState)
 {
 	const FMoverTickStartData& StartState = Params.StartState;
-	USceneComponent* UpdatedComponent = Params.UpdatedComponent;
-	UPrimitiveComponent* UpdatedPrimitive = Params.UpdatedPrimitive;
+	USceneComponent* UpdatedComponent = Params.MovingComps.UpdatedComponent.Get();
+	UPrimitiveComponent* UpdatedPrimitive = Params.MovingComps.UpdatedPrimitive.Get();
 	FProposedMove ProposedMove = Params.ProposedMove;
 
 	const FMoverDefaultSyncState* StartingSyncState = StartState.SyncState.SyncStateCollection.FindDataByType<FMoverDefaultSyncState>();
diff --git a/Engine/Plugins/Experimental/Mover/Source/Mover/Private/PhysicsMover/Modes/PhysicsDrivenSwimmingMode.cpp b/Engine/Plugins/Experimental/Mover/Source/Mover/Private/PhysicsMover/Modes/PhysicsDrivenSwimmingMode.cpp
--- a/Engine/Plugins/Experimental/Mover/Source/Mover/Private/PhysicsMover/Modes/PhysicsDrivenSwimmingMode.cpp
+++ b/Engine/Plugins/Experimental/Mover/Source/Mover/Private/PhysicsMover/Modes/PhysicsDrivenSwimmingMode.cpp
@@ -54,8 +54,8 @@ void UPhysicsDrivenSwimmingMode::OnSimulationTick(const FSimulationTickParams& P
 	const UMoverComponent* MoverComp = GetMoverComponent();
 
 	const FMoverTickStartData& StartState = Params.StartState;
-	USceneComponent* UpdatedComponent = Params.UpdatedComponent;
-	UPrimitiveComponent* UpdatedPrimitive = Params.UpdatedPrimitive;
+	USceneComponent* UpdatedComponent = Params.MovingComps.UpdatedComponent.Get();
+	UPrimitiveComponent* UpdatedPrimitive = Params.MovingComps.UpdatedPrimitive.Get();
 	FProposedMove ProposedMove = Params.ProposedMove;
 	
 	const FVector UpDir = GetMoverComponent()->GetUpDirection();
diff --git a/Engine/Plugins/Experimental/Mover/Source/Mover/Private/PhysicsMover/Modes/PhysicsDrivenWalkingMode.cpp b/Engine/Plugins/Experimental/Mover/Source/Mover/Private/PhysicsMover/Modes/PhysicsDrivenWalkingMode.cpp
--- a/Engine/Plugins/Experimental/Mover/Source/Mover/Private/PhysicsMover/Modes/PhysicsDrivenWalkingMode.cpp
+++ b/Engine/Plugins/Experimental/Mover/Source/Mover/Private/PhysicsMover/Modes/PhysicsDrivenWalkingMode.cpp
@@ -263,8 +263,8 @@ void UPhysicsDrivenWalkingMode::OnSimulationTick(const FSimulationTickParams& Pa
 {
 	const UMoverComponent* MoverComp = GetMoverComponent();
 	const FMoverTickStartData& StartState = Params.StartState;
-	USceneComponent* UpdatedComponent = Params.UpdatedComponent;
-	UPrimitiveComponent* UpdatedPrimitive = Params.UpdatedPrimitive;
+	USceneComponent* UpdatedComponent = Params.MovingComps.UpdatedComponent.Get();
+	UPrimitiveComponent* UpdatedPrimitive = Params.MovingComps.UpdatedPrimitive.Get();
 	FProposedMove ProposedMove = Params.ProposedMove;
 
 	const FVector UpDir = GetMoverComponent()->GetUpDirection();
diff --git a/Engine/Plugins/Experimental/Mover/Source/Mover/Private/PhysicsMover/Transitions/PhysicsJumpCheck.cpp b/Engine/Plugins/Experimental/Mover/Source/Mover/Private/PhysicsMover/Transitions/PhysicsJumpCheck.cpp
--- a/Engine/Plugins/Experimental/Mover/Source/Mover/Private/PhysicsMover/Transitions/PhysicsJumpCheck.cpp
+++ b/Engine/Plugins/Experimental/Mover/Source/Mover/Private/PhysicsMover/Transitions/PhysicsJumpCheck.cpp
@@ -35,7 +35,7 @@ void UPhysicsJumpCheck::OnTrigger(const FSimulationTickParams& Params)
 	TSharedPtr<FJumpImpulseEffect> JumpMove = MakeShared<FJumpImpulseEffect>();
 	JumpMove->UpwardsSpeed = JumpUpwardsSpeed;
 
-	Params.MoverComponent->QueueInstantMovementEffect(JumpMove);
+	Params.MovingComps.MoverComponent->QueueInstantMovementEffect(JumpMove);
 }
 
 #if WITH_EDITOR
diff --git a/Engine/Plugins/Experimental/Mover/Source/Mover/Public/MoveLibrary/AirMovementUtils.h b/Engine/Plugins/Experimental/Mover/Source/Mover/Public/MoveLibrary/AirMovementUtils.h
--- a/Engine/Plugins/Experimental/Mover/Source/Mover/Public/MoveLibrary/AirMovementUtils.h
+++ b/Engine/Plugins/Experimental/Mover/Source/Mover/Public/MoveLibrary/AirMovementUtils.h
@@ -64,10 +64,10 @@ public:
 	
     // Checks if a hit result represents a walkable location that an actor can land on
     UFUNCTION(BlueprintCallable, Category=Mover)
-    static bool IsValidLandingSpot(USceneComponent* UpdatedComponent, UPrimitiveComponent* UpdatedPrimitive, const FVector& Location, const FHitResult& Hit, float FloorSweepDistance, float MaxWalkSlopeCosine, FFloorCheckResult& OutFloorResult);
+    static bool IsValidLandingSpot(const FMovingComponentSet& MovingComps, const FVector& Location, const FHitResult& Hit, float FloorSweepDistance, float MaxWalkSlopeCosine, FFloorCheckResult& OutFloorResult);
     
     /** Attempts to move a component along a surface, while checking for landing on a walkable surface. Intended for use while falling. Returns the percent of time applied, with 0.0 meaning no movement occurred. */
     UFUNCTION(BlueprintCallable, Category=Mover)
-    static float TryMoveToFallAlongSurface(USceneComponent* UpdatedComponent, UPrimitiveComponent* UpdatedPrimitive, UMoverComponent* MoverComponent, const FVector& Delta, float PctOfDeltaToMove, const FQuat Rotation, const FVector& Normal, FHitResult& Hit, bool bHandleImpact, float FloorSweepDistance, float MaxWalkSlopeCosine, FFloorCheckResult& OutFloorResult, FMovementRecord& MoveRecord);
+    static float TryMoveToFallAlongSurface(const FMovingComponentSet& MovingComps, const FVector& Delta, float PctOfDeltaToMove, const FQuat Rotation, const FVector& Normal, FHitResult& Hit, bool bHandleImpact, float FloorSweepDistance, float MaxWalkSlopeCosine, FFloorCheckResult& OutFloorResult, FMovementRecord& MoveRecord);
 
 };
diff --git a/Engine/Plugins/Experimental/Mover/Source/Mover/Public/MoveLibrary/ConstrainedMoveUtils.h b/Engine/Plugins/Experimental/Mover/Source/Mover/Public/MoveLibrary/ConstrainedMoveUtils.h
new file mode 100644
--- /dev/null
+++ b/Engine/Plugins/Experimental/Mover/Source/Mover/Public/MoveLibrary/ConstrainedMoveUtils.h
@@ -0,0 +1,75 @@
+// Copyright Epic Games, Inc. All Rights Reserved.
+
+#pragma once
+
+#include "Kismet/BlueprintFunctionLibrary.h"
+#include "Engine/EngineBaseTypes.h"
+#include "ConstrainedMoveUtils.generated.h"
+
+
+
+/** Encapsulates info about constraining movement to a plane, such as in a side-scroller */
+USTRUCT(BlueprintType)
+struct MOVER_API FPlanarConstraint
+{
+	GENERATED_USTRUCT_BODY()
+
+	/** If true, movement will be constrained to a plane. */
+	UPROPERTY(EditAnywhere, BlueprintReadOnly, Category = PlanarMovement)
+	bool bConstrainToPlane = false;
+
+	/**
+	 * The normal or axis of the plane that constrains movement, if bConstrainToPlane is enabled.
+	 * If for example you wanted to constrain movement to the X-Z plane (so that Y cannot change), the normal would be set to X=0 Y=1 Z=0.
+	 * It is normalized once the component is registered with the game world.
+	 */
+	UPROPERTY(EditAnywhere, BlueprintReadOnly, Category = PlanarMovement, meta = (editcondition = bConstrainToPlane))
+	FVector PlaneConstraintNormal = FVector::ForwardVector;
+
+	/** The origin of the plane that constrains movement, if plane constraint is enabled. */
+	UPROPERTY(EditAnywhere, BlueprintReadOnly, Category = PlanarMovement, meta = (editcondition = bConstrainToPlane))
+	FVector PlaneConstraintOrigin = FVector::ZeroVector;
+};
+
+
+
+/**
+ * PlanarConstraintUtils: a collection of stateless static BP-accessible functions for working with planar constraints
+ */
+UCLASS()
+class MOVER_API UPlanarConstraintUtils : public UBlueprintFunctionLibrary
+{
+	GENERATED_BODY()
+
+public:
+	/** Sets whether or not the constraint is enabled */
+	UFUNCTION(BlueprintCallable, Category = "Mover|PlanarConstraint")
+	static void SetPlanarConstraintEnabled(UPARAM(ref) FPlanarConstraint& Constraint, bool bEnabled);
+
+	/**
+	 * Sets the normal of the plane that constrains movement, enforced if the plane constraint is enabled.
+	 *
+	 * @param PlaneNormal	The normal of the plane. If non-zero in length, it will be normalized.
+	 */
+	UFUNCTION(BlueprintCallable, Category = "Mover|PlanarConstraint")
+	static void SetPlanarConstraintNormal(UPARAM(ref) FPlanarConstraint& Constraint, FVector PlaneNormal);
+
+	/** Sets the origin of the plane that constrains movement, enforced if the plane constraint is enabled. */
+	UFUNCTION(BlueprintCallable, Category = "Mover|PlanarConstraint")
+	static void SetPlaneConstraintOrigin(UPARAM(ref) FPlanarConstraint& Constraint, FVector PlaneOrigin);
+
+	// APPLICATION OF CONSTRAINT
+
+	/** Constrains a direction to be on the plane, if enabled */
+	UFUNCTION(BlueprintCallable, Category = "Mover|PlanarConstraint")
+	static FVector ConstrainDirectionToPlane(const FPlanarConstraint& Constraint, FVector Direction, bool bMaintainMagnitude=false);
+
+	/** Constrains a location to be on the plane, if enabled */
+	UFUNCTION(BlueprintCallable, Category = "Mover|PlanarConstraint")
+	static FVector ConstrainLocationToPlane(const FPlanarConstraint& Constraint, FVector Location);
+
+	/** Constrains a normal to be on the plane, if enabled. Result will be re-normalized. */
+	UFUNCTION(BlueprintCallable, Category = "Mover|PlanarConstraint")
+	static FVector ConstrainNormalToPlane(const FPlanarConstraint& Constraint, FVector Normal);
+
+}; // UPlanarConstraintUtils
\ No newline at end of file
diff --git a/Engine/Plugins/Experimental/Mover/Source/Mover/Public/MoveLibrary/GroundMovementUtils.h b/Engine/Plugins/Experimental/Mover/Source/Mover/Public/MoveLibrary/GroundMovementUtils.h
--- a/Engine/Plugins/Experimental/Mover/Source/Mover/Public/MoveLibrary/GroundMovementUtils.h
+++ b/Engine/Plugins/Experimental/Mover/Source/Mover/Public/MoveLibrary/GroundMovementUtils.h
@@ -75,16 +75,16 @@ public:
 	
 	// TODO: Refactor this API for fewer parameters
 	/** Move up steps or slope. Does nothing and returns false if hit surface is invalid for step-up use */
-	static bool TryMoveToStepUp(USceneComponent* UpdatedComponent, UPrimitiveComponent* UpdatedPrimitive, UMoverComponent* MoverComponent, const FVector& GravDir, float MaxStepHeight, float MaxWalkSlopeCosine, float FloorSweepDistance, const FVector& MoveDelta, const FHitResult& Hit, const FFloorCheckResult& CurrentFloor, bool bIsFalling, FOptionalFloorCheckResult* OutStepDownResult, FMovementRecord& MoveRecord);
+	static bool TryMoveToStepUp(const FMovingComponentSet& MovingComps, const FVector& GravDir, float MaxStepHeight, float MaxWalkSlopeCosine, float FloorSweepDistance, const FVector& MoveDelta, const FHitResult& Hit, const FFloorCheckResult& CurrentFloor, bool bIsFalling, FOptionalFloorCheckResult* OutStepDownResult, FMovementRecord& MoveRecord);
 
 	/** Moves vertically to stay within range of the walkable floor. Does nothing and returns false if floor is unwalkable or if already in range. */ 
-	static bool TryMoveToAdjustHeightAboveFloor(USceneComponent* UpdatedComponent, UPrimitiveComponent* UpdatedPrimitive, FFloorCheckResult& CurrentFloor, float MaxWalkSlopeCosine, FMovementRecord& MoveRecord);
+	static bool TryMoveToAdjustHeightAboveFloor(const FMovingComponentSet& MovingComps, FFloorCheckResult& CurrentFloor, float MaxWalkSlopeCosine, FMovementRecord& MoveRecord);
 
 	/** Attempts to move a component along a surface in the walking mode. Returns the percent of time applied, with 0.0 meaning no movement occurred.
      *  Note: This modifies the normal and calls UMovementUtils::TryMoveToSlideAlongSurface
      */
     UFUNCTION(BlueprintCallable, Category=Mover)
-    static float TryWalkToSlideAlongSurface(USceneComponent* UpdatedComponent, UPrimitiveComponent* UpdatedPrimitive, UMoverComponent* MoverComponent, const FVector& Delta, float PctOfDeltaToMove, const FQuat Rotation, const FVector& Normal, FHitResult& Hit, bool bHandleImpact, FMovementRecord& MoveRecord, float MaxWalkSlopeCosine, float MaxStepHeight);
+    static float TryWalkToSlideAlongSurface(const FMovingComponentSet& MovingComps, const FVector& Delta, float PctOfDeltaToMove, const FQuat Rotation, const FVector& Normal, FHitResult& Hit, bool bHandleImpact, FMovementRecord& MoveRecord, float MaxWalkSlopeCosine, float MaxStepHeight);
 
 	/** Used to change a movement to be along a ramp's surface, typically to prevent slow movement when running up/down a ramp */
 	UFUNCTION(BlueprintCallable, Category = Mover)
diff --git a/Engine/Plugins/Experimental/Mover/Source/Mover/Public/MoveLibrary/MovementUtils.h b/Engine/Plugins/Experimental/Mover/Source/Mover/Public/MoveLibrary/MovementUtils.h
--- a/Engine/Plugins/Experimental/Mover/Source/Mover/Public/MoveLibrary/MovementUtils.h
+++ b/Engine/Plugins/Experimental/Mover/Source/Mover/Public/MoveLibrary/MovementUtils.h
@@ -146,25 +146,26 @@ public:
 
 	/** Returns an alternative move delta to slide along a surface, based on parameters describing a blocked attempted move */
 	UFUNCTION(BlueprintCallable, Category=Mover)
-	static FVector ComputeSlideDelta(const FVector& Delta, const float PctOfDeltaToMove, const FVector& Normal, const FHitResult& Hit);
+	static FVector ComputeSlideDelta(const FMovingComponentSet& MovingComps, const FVector& Delta, const float PctOfDeltaToMove, const FVector& Normal, const FHitResult& Hit);
 
 	/** Returns an alternative move delta when we are in contact with 2 surfaces */
-	static FVector ComputeTwoWallAdjustedDelta(const FVector& MoveDelta, const FHitResult& Hit, const FVector& OldHitNormal);
+	static FVector ComputeTwoWallAdjustedDelta(const FMovingComponentSet& MovingComps, const FVector& MoveDelta, const FHitResult& Hit, const FVector& OldHitNormal);
 
 	/** Attempts to move a component along a surface. Returns the percent of time applied, with 0.0 meaning no movement occurred. */
 	UFUNCTION(BlueprintCallable, Category=Mover)
-	static float TryMoveToSlideAlongSurface(USceneComponent* UpdatedComponent, UPrimitiveComponent* UpdatedPrimitive, UMoverComponent* MoverComponent, const FVector& Delta, float PctOfDeltaToMove, const FQuat Rotation, const FVector& Normal, FHitResult& Hit, bool bHandleImpact, FMovementRecord& MoveRecord);
+	static float TryMoveToSlideAlongSurface(const FMovingComponentSet& MovingComps, const FVector& Delta, float PctOfDeltaToMove, const FQuat Rotation, const FVector& Normal, FHitResult& Hit, bool bHandleImpact, FMovementRecord& MoveRecord);
+
 	// Component movement
 
 	/** Attempts to move a component and resolve any penetration issues with the proposed move Delta */
 	UFUNCTION(BlueprintCallable, Category=Mover)
-	static bool TrySafeMoveUpdatedComponent(USceneComponent* UpdatedComponent, UPrimitiveComponent* UpdatedPrimitive, const FVector& Delta, const FQuat& NewRotation, bool bSweep, FHitResult& OutHit, ETeleportType Teleport, FMovementRecord& MoveRecord);
+	static bool TrySafeMoveUpdatedComponent(const FMovingComponentSet& MovingComps, const FVector& Delta, const FQuat& NewRotation, bool bSweep, FHitResult& OutHit, ETeleportType Teleport, FMovementRecord& MoveRecord);
 
 	/** Returns a movement step that should get the subject of the hit result out of an initial penetration condition */
 	static FVector ComputePenetrationAdjustment(const FHitResult& Hit);
 	
 	/** Attempts to move out of a situation where the component is stuck in geometry, using a suggested adjustment to start. */
-	static bool TryMoveToResolvePenetration(USceneComponent* UpdatedComponent, UPrimitiveComponent* UpdatedPrimitive, EMoveComponentFlags MoveComponentFlags, const FVector& ProposedAdjustment, const FHitResult& Hit, const FQuat& NewRotationQuat, FMovementRecord& MoveRecord);
+	static bool TryMoveToResolvePenetration(const FMovingComponentSet& MovingComps, EMoveComponentFlags MoveComponentFlags, const FVector& ProposedAdjustment, const FHitResult& Hit, const FQuat& NewRotationQuat, FMovementRecord& MoveRecord);
 	static void InitCollisionParams(const UPrimitiveComponent* UpdatedPrimitive, FCollisionQueryParams& OutParams, FCollisionResponseParams& OutResponseParam);
 	static bool OverlapTest(const USceneComponent* UpdatedComponent, const UPrimitiveComponent* UpdatedPrimitive, const FVector& Location, const FQuat& RotationQuat, const ECollisionChannel CollisionChannel, const FCollisionShape& CollisionShape, const AActor* IgnoreActor);
 
@@ -185,5 +186,6 @@ public:
 	/** Internal function that other move functions use to perform all actual component movement and retrieve results
 	 *  Note: This function moves the character directly and should only be used if needed. Consider using something like TrySafeMoveUpdatedComponent.
 	 */
-	static bool TryMoveUpdatedComponent_Internal(USceneComponent* UpdatedComponent, const FVector& Delta, const FQuat& NewRotation, bool bSweep, EMoveComponentFlags MoveComponentFlags, FHitResult* OutHit, ETeleportType Teleport);
+	static bool TryMoveUpdatedComponent_Internal(const FMovingComponentSet& MovingComps, const FVector& Delta, const FQuat& NewRotation, bool bSweep, EMoveComponentFlags MoveComponentFlags, FHitResult* OutHit, ETeleportType Teleport);
+
 };
diff --git a/Engine/Plugins/Experimental/Mover/Source/Mover/Public/MoveLibrary/MovementUtilsTypes.h b/Engine/Plugins/Experimental/Mover/Source/Mover/Public/MoveLibrary/MovementUtilsTypes.h
--- a/Engine/Plugins/Experimental/Mover/Source/Mover/Public/MoveLibrary/MovementUtilsTypes.h
+++ b/Engine/Plugins/Experimental/Mover/Source/Mover/Public/MoveLibrary/MovementUtilsTypes.h
@@ -4,6 +4,9 @@
 
 #include "MovementUtilsTypes.generated.h"
 
+class USceneComponent;
+class UPrimitiveComponent;
+class UMoverComponent;
 
 
 UENUM()
@@ -44,4 +47,29 @@ struct MOVER_API FProposedMove
 	FVector  LinearVelocity = FVector::ZeroVector;		// Units per second, world space, possibly mapped onto walking surface
 	UPROPERTY(EditAnywhere, BlueprintReadWrite, Category = Mover)
 	FRotator AngularVelocity = FRotator::ZeroRotator;	// Degrees per second, local space
+};
+
+
+/** 
+ * Encapsulates components involved in movement. Used by many library functions. 
+ * Only a scene component is required for movement, but this is typically a primitive
+ * component so we provide a pre-cast ptr for convenience.
+ */
+USTRUCT(BlueprintType)
+struct MOVER_API FMovingComponentSet
+{
+	GENERATED_USTRUCT_BODY()
+
+	FMovingComponentSet() {}
+	FMovingComponentSet(USceneComponent* InUpdatedComponent);
+	FMovingComponentSet(UMoverComponent* InMoverComponent);
+
+	UPROPERTY(EditAnywhere, BlueprintReadWrite, Category = Mover)
+	TWeakObjectPtr<USceneComponent> UpdatedComponent = nullptr;
+
+	UPROPERTY(EditAnywhere, BlueprintReadWrite, Category = Mover)
+	TWeakObjectPtr<UPrimitiveComponent> UpdatedPrimitive = nullptr;
+
+	UPROPERTY(EditAnywhere, BlueprintReadWrite, Category = Mover)
+	TWeakObjectPtr<UMoverComponent> MoverComponent = nullptr;
 };
\ No newline at end of file
diff --git a/Engine/Plugins/Experimental/Mover/Source/Mover/Public/MoverComponent.h b/Engine/Plugins/Experimental/Mover/Source/Mover/Public/MoverComponent.h
--- a/Engine/Plugins/Experimental/Mover/Source/Mover/Public/MoverComponent.h
+++ b/Engine/Plugins/Experimental/Mover/Source/Mover/Public/MoverComponent.h
@@ -11,6 +11,7 @@
 #include "MoverTypes.h"
 #include "LayeredMove.h"
 #include "MoveLibrary/BasedMovementUtils.h"
+#include "MoveLibrary/ConstrainedMoveUtils.h"
 #if UE_ENABLE_INCLUDE_ORDER_DEPRECATED_IN_5_4
 #include "Engine/HitResult.h"
 #endif // UE_ENABLE_INCLUDE_ORDER_DEPRECATED_IN_5_4
@@ -186,6 +187,14 @@ public:
 	UFUNCTION(BlueprintCallable, BlueprintPure = false, Category = Mover)
 	FVector GetUpDirection() const;
 
+	// Access the planar constraint that may be limiting movement direction
+	UFUNCTION(BlueprintCallable, BlueprintPure = false, Category = Mover)
+	const FPlanarConstraint& GetPlanarConstraint() const;
+
+	// Sets planar constraint that can limit movement direction
+	UFUNCTION(BlueprintCallable, Category = Mover)
+	void SetPlanarConstraint(const FPlanarConstraint& InConstraint);
+
 public:
 
 	/**
@@ -417,14 +426,18 @@ private:
 	UPROPERTY(EditDefaultsOnly, EditFixedSize, Instanced, Category = Mover, meta = (NoResetToDefault, MustImplement = "/Script/Mover.MovementSettingsInterface"))
 	TArray<TObjectPtr<UObject>> SharedSettings;
 
-	// Whether or not gravity is overridden on this actor. Otherwise, fall back on world settings. See @SetGravityOverride
+	/** Whether or not gravity is overridden on this actor. Otherwise, fall back on world settings. See @SetGravityOverride */
 	UPROPERTY(EditDefaultsOnly, Category="Mover|Gravity")
 	bool bHasGravityOverride = false;
 	
-	// cm/s^2, only meaningful if @bHasGravityOverride is enabled. Set @SetGravityOverride
+	/** cm/s^2, only meaningful if @bHasGravityOverride is enabled.Set @SetGravityOverride */
 	UPROPERTY(EditDefaultsOnly, Category="Mover|Gravity", meta=(ForceUnits = "cm/s^2"))
 	FVector GravityAccelOverride;
 
+	/** Settings that can lock movement to a particular plane */
+	UPROPERTY(EditDefaultsOnly, Category = "Mover|Constraints")
+	FPlanarConstraint PlanarConstraint;
+
 	/** If enabled, this actor will be moved to follow a base actor that it's standing on. Typically disabled for physics-based movement, which handles based movement internally. */
 	UPROPERTY(EditDefaultsOnly, Category = "Mover")
 	bool bSupportsKinematicBasedMovement = true;
diff --git a/Engine/Plugins/Experimental/Mover/Source/Mover/Public/MoverSimulationTypes.h b/Engine/Plugins/Experimental/Mover/Source/Mover/Public/MoverSimulationTypes.h
--- a/Engine/Plugins/Experimental/Mover/Source/Mover/Public/MoverSimulationTypes.h
+++ b/Engine/Plugins/Experimental/Mover/Source/Mover/Public/MoverSimulationTypes.h
@@ -243,17 +243,9 @@ struct MOVER_API FSimulationTickParams
 {
 	GENERATED_BODY()
 
-	// The scene component that is being moved by the simulation, usually the same as the primitive component
+	// Components involved in movement by the simulation
 	UPROPERTY(BlueprintReadOnly, EditAnywhere, Category = Mover)
-	TObjectPtr<USceneComponent> UpdatedComponent;
-
-	// The primitive component used for collision checking by the simulation, usually the same as UpdatedComponent
-	UPROPERTY(BlueprintReadOnly, EditAnywhere, Category = Mover)
-	TObjectPtr<UPrimitiveComponent> UpdatedPrimitive;
-
-	// The Mover Component associated with this sim tick
-	UPROPERTY(BlueprintReadOnly, EditAnywhere, Category = Mover)
-	TObjectPtr<UMoverComponent> MoverComponent;
+	FMovingComponentSet MovingComps;
 
 	// Simulation state data at the start of the tick, including Input Cmd
 	UPROPERTY(BlueprintReadOnly, EditAnywhere, Category = Mover)
diff --git a/Engine/Plugins/Experimental/MoverExamples/Source/MoverExamples/Private/CharacterVariants/Ziplining/ZipliningMode.cpp b/Engine/Plugins/Experimental/MoverExamples/Source/MoverExamples/Private/CharacterVariants/Ziplining/ZipliningMode.cpp
--- a/Engine/Plugins/Experimental/MoverExamples/Source/MoverExamples/Private/CharacterVariants/Ziplining/ZipliningMode.cpp
+++ b/Engine/Plugins/Experimental/MoverExamples/Source/MoverExamples/Private/CharacterVariants/Ziplining/ZipliningMode.cpp
@@ -88,10 +88,8 @@ void UZipliningMode::OnSimulationTick(const FSimulationTickParams& Params, FMove
 	FMoverDefaultSyncState& OutputSyncState = OutputState.SyncState.SyncStateCollection.FindOrAddMutableDataByType<FMoverDefaultSyncState>();
 	FZipliningState& OutZipState            = OutputState.SyncState.SyncStateCollection.FindOrAddMutableDataByType<FZipliningState>();
 
-
-	UMoverComponent* MoverComp = Params.MoverComponent;
-	USceneComponent* UpdatedComponent = Params.UpdatedComponent;
-	UPrimitiveComponent* UpdatedPrimitive = Params.UpdatedPrimitive;
+	USceneComponent* UpdatedComponent = Params.MovingComps.UpdatedComponent.Get();
+	UMoverComponent* MoverComp = Params.MovingComps.MoverComponent.Get();
 	AActor* MoverActor = MoverComp->GetOwner();
 
 	USceneComponent* StartPoint = nullptr;
@@ -214,7 +212,7 @@ void UZipliningMode::OnSimulationTick(const FSimulationTickParams& Params, FMove
 		
 		FHitResult Hit(1.f);
 
-		UMovementUtils::TrySafeMoveUpdatedComponent(UpdatedComponent, UpdatedPrimitive, MoveDelta, FlatFacingDir.ToOrientationQuat(), true, Hit, ETeleportType::None, MoveRecord);
+		UMovementUtils::TrySafeMoveUpdatedComponent(Params.MovingComps, MoveDelta, FlatFacingDir.ToOrientationQuat(), true, Hit, ETeleportType::None, MoveRecord);
 	}
 
 
diff --git a/Engine/Plugins/Experimental/MoverExamples/Source/MoverExamples/Private/CharacterVariants/Ziplining/ZipliningTransitions.cpp b/Engine/Plugins/Experimental/MoverExamples/Source/MoverExamples/Private/CharacterVariants/Ziplining/ZipliningTransitions.cpp
--- a/Engine/Plugins/Experimental/MoverExamples/Source/MoverExamples/Private/CharacterVariants/Ziplining/ZipliningTransitions.cpp
+++ b/Engine/Plugins/Experimental/MoverExamples/Source/MoverExamples/Private/CharacterVariants/Ziplining/ZipliningTransitions.cpp
@@ -19,7 +19,7 @@ FTransitionEvalResult UZiplineStartTransition::OnEvaluate(const FSimulationTickP
 {
 	FTransitionEvalResult EvalResult = FTransitionEvalResult::NoTransition;
 
-	UCharacterMoverComponent* MoverComp = Cast<UCharacterMoverComponent>(Params.MoverComponent);
+	UCharacterMoverComponent* MoverComp = Cast<UCharacterMoverComponent>(Params.MovingComps.MoverComponent.Get());
 
 	const FMoverSyncState& SyncState = Params.StartState.SyncState;
 
diff --git a/Engine/Plugins/Experimental/MoverExamples/Source/MoverExamples/Private/MovementBases/FollowPathMode.cpp b/Engine/Plugins/Experimental/MoverExamples/Source/MoverExamples/Private/MovementBases/FollowPathMode.cpp
--- a/Engine/Plugins/Experimental/MoverExamples/Source/MoverExamples/Private/MovementBases/FollowPathMode.cpp
+++ b/Engine/Plugins/Experimental/MoverExamples/Source/MoverExamples/Private/MovementBases/FollowPathMode.cpp
@@ -21,8 +21,7 @@ void UFollowPathMode::OnGenerateMove(const FMoverTickStartData& StartState, cons
 void UFollowPathMode::OnSimulationTick(const FSimulationTickParams& Params, FMoverTickEndData& OutputState)
 {
 	const FMoverTickStartData& StartState = Params.StartState;
-	USceneComponent* UpdatedComponent = Params.UpdatedComponent;
-	UPrimitiveComponent* UpdatedPrimitive = Params.UpdatedPrimitive;
+	USceneComponent* UpdatedComponent = Params.MovingComps.UpdatedComponent.Get();
 
 	const FMoverDefaultSyncState* StartingMoveState = StartState.SyncState.SyncStateCollection.FindDataByType<FMoverDefaultSyncState>();
 	const FFollowPathState* StartingPathState		= StartState.SyncState.SyncStateCollection.FindDataByType<FFollowPathState>();
diff --git a/Engine/Plugins/Experimental/MoverExamples/Source/MoverExamples/Private/MovementBases/FollowSplineMode.cpp b/Engine/Plugins/Experimental/MoverExamples/Source/MoverExamples/Private/MovementBases/FollowSplineMode.cpp
--- a/Engine/Plugins/Experimental/MoverExamples/Source/MoverExamples/Private/MovementBases/FollowSplineMode.cpp
+++ b/Engine/Plugins/Experimental/MoverExamples/Source/MoverExamples/Private/MovementBases/FollowSplineMode.cpp
@@ -54,8 +54,7 @@ void UFollowSplineMode::OnGenerateMove(const FMoverTickStartData& StartState, co
 void UFollowSplineMode::OnSimulationTick(const FSimulationTickParams& Params, FMoverTickEndData& OutputState)
 {
 	const FMoverTickStartData& StartState = Params.StartState;
-	USceneComponent* UpdatedComponent = Params.UpdatedComponent;
-	UPrimitiveComponent* UpdatedPrimitive = Params.UpdatedPrimitive;
+	const FMovingComponentSet& MovingComps = Params.MovingComps;
 
 	const FMoverDefaultSyncState* StartingMoveState = StartState.SyncState.SyncStateCollection.FindDataByType<FMoverDefaultSyncState>();
 	const FFollowSplineState* StartingPathState		= StartState.SyncState.SyncStateCollection.FindDataByType<FFollowSplineState>();
@@ -122,7 +121,7 @@ void UFollowSplineMode::OnSimulationTick(const FSimulationTickParams& Params, FM
 	if (CanMove(MappedSplineTime))
 	{
 		// Move the object
-		const FTransform SplineTransform = GetTransformAtTime(MappedSplineTime, UpdatedComponent->GetComponentRotation());
+		const FTransform SplineTransform = GetTransformAtTime(MappedSplineTime, MovingComps.UpdatedComponent->GetComponentRotation());
 		const FVector MoveDelta = SplineTransform.GetLocation() - StartingLocation;
 		const FVector Velocity = MoveDelta / DeltaSeconds;
 
@@ -131,14 +130,13 @@ void UFollowSplineMode::OnSimulationTick(const FSimulationTickParams& Params, FM
 		
 		FHitResult MoveHitResult;
 		
-		UMovementUtils::TrySafeMoveUpdatedComponent(UpdatedComponent, UpdatedPrimitive, MoveDelta, SplineTransform.GetRotation(), true, MoveHitResult, ETeleportType::None, MoveRecord);
+		UMovementUtils::TrySafeMoveUpdatedComponent(MovingComps, MoveDelta, SplineTransform.GetRotation(), true, MoveHitResult, ETeleportType::None, MoveRecord);
 
 		UpdatePathState(OutputPathState);
 
 		// Update Move State
-		OutputMoveState.SetTransforms_WorldSpace(UpdatedComponent->GetComponentLocation(),
-
-			UpdatedComponent->GetComponentRotation(),
+		OutputMoveState.SetTransforms_WorldSpace(MovingComps.UpdatedComponent->GetComponentLocation(),
+			MovingComps.UpdatedComponent->GetComponentRotation(),
 			Velocity,
 			nullptr); // no movement base
 	}
