Mover: allow exporting of common movement mode gameplay tags, so other modules can use them directly

diff --git a/Engine/Plugins/Experimental/Mover/Source/Mover/Private/MoverTypes.cpp b/Engine/Plugins/Experimental/Mover/Source/Mover/Private/MoverTypes.cpp
--- a/Engine/Plugins/Experimental/Mover/Source/Mover/Private/MoverTypes.cpp
+++ b/Engine/Plugins/Experimental/Mover/Source/Mover/Private/MoverTypes.cpp
@@ -7,10 +7,10 @@
 
 #define LOCTEXT_NAMESPACE "MoverData"
 
-UE_DEFINE_GAMEPLAY_TAG_COMMENT(Mover_IsOnGround, "Mover.IsOnGround", "Default Mover state flag indicating character is on the ground.")
-UE_DEFINE_GAMEPLAY_TAG_COMMENT(Mover_IsInAir, "Mover.IsInAir", "Default Mover state flag indicating character is in the air.")
-UE_DEFINE_GAMEPLAY_TAG_COMMENT(Mover_IsFalling, "Mover.IsFalling", "Default Mover state flag indicating character is falling.")
-UE_DEFINE_GAMEPLAY_TAG_COMMENT(Mover_IsSwimming, "Mover.IsSwimming", "Default Mover state flag indicating character is swimming.")
+UE_DEFINE_GAMEPLAY_TAG_COMMENT(Mover_IsOnGround, "Mover.IsOnGround", "Default Mover state flag indicating character is on the ground.");
+UE_DEFINE_GAMEPLAY_TAG_COMMENT(Mover_IsInAir, "Mover.IsInAir", "Default Mover state flag indicating character is in the air.");
+UE_DEFINE_GAMEPLAY_TAG_COMMENT(Mover_IsFalling, "Mover.IsFalling", "Default Mover state flag indicating character is falling.");
+UE_DEFINE_GAMEPLAY_TAG_COMMENT(Mover_IsSwimming, "Mover.IsSwimming", "Default Mover state flag indicating character is swimming.");
 
 FMoverOnImpactParams::FMoverOnImpactParams() 
 	: AttemptedMoveDelta(0) 
diff --git a/Engine/Plugins/Experimental/Mover/Source/Mover/Public/MoverTypes.h b/Engine/Plugins/Experimental/Mover/Source/Mover/Public/MoverTypes.h
--- a/Engine/Plugins/Experimental/Mover/Source/Mover/Public/MoverTypes.h
+++ b/Engine/Plugins/Experimental/Mover/Source/Mover/Public/MoverTypes.h
@@ -11,10 +11,10 @@
 #include "MoverTypes.generated.h"
 
 // Gameplay tags
-UE_DECLARE_GAMEPLAY_TAG_EXTERN(Mover_IsOnGround)
-UE_DECLARE_GAMEPLAY_TAG_EXTERN(Mover_IsInAir)
-UE_DECLARE_GAMEPLAY_TAG_EXTERN(Mover_IsFalling)
-UE_DECLARE_GAMEPLAY_TAG_EXTERN(Mover_IsSwimming)
+MOVER_API UE_DECLARE_GAMEPLAY_TAG_EXTERN(Mover_IsOnGround);
+MOVER_API UE_DECLARE_GAMEPLAY_TAG_EXTERN(Mover_IsInAir);
+MOVER_API UE_DECLARE_GAMEPLAY_TAG_EXTERN(Mover_IsFalling);
+MOVER_API UE_DECLARE_GAMEPLAY_TAG_EXTERN(Mover_IsSwimming);
 
 // Struct to hold params for when an impact happens. This contains all of the data for impacts including what gets passed to the FMover_OnImpact delegate
 USTRUCT(BlueprintType, meta = (DisplayName = "Impact Data"))
