Mover - Support multiple settings classes

diff --git a/Engine/Plugins/Experimental/Mover/Source/Mover/Private/DefaultMovementSet/Modes/FallingMode.cpp b/Engine/Plugins/Experimental/Mover/Source/Mover/Private/DefaultMovementSet/Modes/FallingMode.cpp
--- a/Engine/Plugins/Experimental/Mover/Source/Mover/Private/DefaultMovementSet/Modes/FallingMode.cpp
+++ b/Engine/Plugins/Experimental/Mover/Source/Mover/Private/DefaultMovementSet/Modes/FallingMode.cpp
@@ -23,7 +23,7 @@ UFallingMode::UFallingMode(const FObjectInitializer& ObjectInitializer)
 	, VerticalFallingDeceleration(4000.0f)
 	, TerminalVerticalSpeed(2000.0f)
 {
-	SharedSettingsClass = UCommonLegacyMovementSettings::StaticClass();
+	SharedSettingsClasses.Add(UCommonLegacyMovementSettings::StaticClass());
 
 	GameplayTags.AddTag(Mover_IsInAir);
 	GameplayTags.AddTag(Mover_IsFalling);
diff --git a/Engine/Plugins/Experimental/Mover/Source/Mover/Private/DefaultMovementSet/Modes/SwimmingMode.cpp b/Engine/Plugins/Experimental/Mover/Source/Mover/Private/DefaultMovementSet/Modes/SwimmingMode.cpp
--- a/Engine/Plugins/Experimental/Mover/Source/Mover/Private/DefaultMovementSet/Modes/SwimmingMode.cpp
+++ b/Engine/Plugins/Experimental/Mover/Source/Mover/Private/DefaultMovementSet/Modes/SwimmingMode.cpp
@@ -15,7 +15,7 @@
 USwimmingMode::USwimmingMode(const FObjectInitializer& ObjectInitializer)
 	: Super(ObjectInitializer)
 {
-	SharedSettingsClass = UCommonLegacyMovementSettings::StaticClass();
+	SharedSettingsClasses.Add(UCommonLegacyMovementSettings::StaticClass());
 
 	GameplayTags.AddTag(Mover_IsSwimming);
 }
diff --git a/Engine/Plugins/Experimental/Mover/Source/Mover/Private/DefaultMovementSet/Modes/WalkingMode.cpp b/Engine/Plugins/Experimental/Mover/Source/Mover/Private/DefaultMovementSet/Modes/WalkingMode.cpp
--- a/Engine/Plugins/Experimental/Mover/Source/Mover/Private/DefaultMovementSet/Modes/WalkingMode.cpp
+++ b/Engine/Plugins/Experimental/Mover/Source/Mover/Private/DefaultMovementSet/Modes/WalkingMode.cpp
@@ -19,7 +19,7 @@
 UWalkingMode::UWalkingMode(const FObjectInitializer& ObjectInitializer)
 	: Super(ObjectInitializer)
 {
-	SharedSettingsClass = UCommonLegacyMovementSettings::StaticClass();
+	SharedSettingsClasses.Add(UCommonLegacyMovementSettings::StaticClass());
 
 	GameplayTags.AddTag(Mover_IsOnGround);
 }
diff --git a/Engine/Plugins/Experimental/Mover/Source/Mover/Private/MoverComponent.cpp b/Engine/Plugins/Experimental/Mover/Source/Mover/Private/MoverComponent.cpp
--- a/Engine/Plugins/Experimental/Mover/Source/Mover/Private/MoverComponent.cpp
+++ b/Engine/Plugins/Experimental/Mover/Source/Mover/Private/MoverComponent.cpp
@@ -471,6 +471,13 @@ bool UMoverComponent::InitMoverSimulation()
 
 	for (const TPair<FName, TObjectPtr<UBaseMovementMode>>& Element : MovementModes)
 	{
+		if (Element.Value.Get() == nullptr)
+		{
+			UE_LOG(LogMover, Warning, TEXT("Invalid Movement Mode type '%s' detected on %s. Mover actor will not function correctly."),
+				*Element.Key.ToString(), *GetNameSafe(GetOwner()));
+			continue;
+		}
+
 		ModeFSM->RegisterMovementMode(Element.Key, Element.Value);
 
 		bHasMatchingStartingState |= (StartingMovementMode == Element.Key);
@@ -708,15 +715,28 @@ bool UMoverComponent::ValidateSetup(FDataValidationContext& Context) const
 		}
 
 		// Verify that the movement mode's shared settings object exists (if any)
-		if (Element.Value && Element.Value->SharedSettingsClass &&
-			FindSharedSettings(Element.Value->SharedSettingsClass) == nullptr)
+		if (Element.Value)
 		{
-			Context.AddError(FText::Format(LOCTEXT("MissingModeSettingsError", "Movement mode on {0}, mapped as {1}, is missing its desired SharedSettingsClass {2}. You may need to save the asset and/or recompile."),
-				FText::FromString(GetNameSafe(GetOwner())),
-				FText::FromName(Element.Key),
-				FText::FromString(Element.Value->SharedSettingsClass->GetName())));
+			for (TSubclassOf<UObject>& Type : Element.Value->SharedSettingsClasses)
+			{
+				if (Type.Get() == nullptr)
+				{
+					Context.AddError(FText::Format(LOCTEXT("InvalidModeSettingsError", "Movement mode on {0}, mapped as {1}, has an invalid SharedSettingsClass. You may need to remove the invalid settings class."),
+						FText::FromString(GetNameSafe(GetOwner())),
+						FText::FromName(Element.Key)));
 
-			bDidFindAnyProblems = true;
+					bDidFindAnyProblems = true;
+				}
+				else if (FindSharedSettings(Type) == nullptr)
+				{
+					Context.AddError(FText::Format(LOCTEXT("MissingModeSettingsError", "Movement mode on {0}, mapped as {1}, is missing its desired SharedSettingsClass {2}. You may need to save the asset and/or recompile."),
+						FText::FromString(GetNameSafe(GetOwner())),
+						FText::FromName(Element.Key),
+						FText::FromString(Type->GetName())));
+
+					bDidFindAnyProblems = true;
+				}
+			}
 		}
 	}
 
@@ -811,12 +831,18 @@ void UMoverComponent::RefreshSharedSettings()
 	{
 		if (UBaseMovementMode* Mode = Element.Value.Get())
 		{
-			if (Mode->SharedSettingsClass != nullptr)
+			for (TSubclassOf<UObject>& SharedSettingsType : Mode->SharedSettingsClasses)
 			{
+				if (SharedSettingsType.Get() == nullptr)
+				{
+					UE_LOG(LogMover, Warning, TEXT("Invalid shared setting class detected on Movement Mode %s."), *Mode->GetName());
+					continue;
+				}
+
 				bool bFoundMatchingClass = false;
 				for (const TObjectPtr<UObject>& SettingsObj : SharedSettings)
 				{
-					if (SettingsObj && SettingsObj->IsA(Mode->SharedSettingsClass))
+					if (SettingsObj && SettingsObj->IsA(SharedSettingsType))
 					{
 						bFoundMatchingClass = true;
 						UnreferencedSettingsObjs.Remove(SettingsObj);
@@ -826,10 +852,9 @@ void UMoverComponent::RefreshSharedSettings()
 
 				if (!bFoundMatchingClass)
 				{
-					UObject* NewSettings = NewObject<UObject>(this, Mode->SharedSettingsClass, NAME_None, GetMaskedFlags(RF_PropagateToSubObjects) | RF_Transactional);
+					UObject* NewSettings = NewObject<UObject>(this, SharedSettingsType, NAME_None, GetMaskedFlags(RF_PropagateToSubObjects) | RF_Transactional);
 					SharedSettings.Add(NewSettings);
 				}
-
 			}
 		}
 	}
@@ -1142,7 +1167,6 @@ USceneComponent* UMoverComponent::GetUpdatedComponent() const
 	return UpdatedComponent.Get();
 }
 
-
 USceneComponent* UMoverComponent::GetPrimaryVisualComponent() const
 {
 	return PrimaryVisualComponent.Get();
diff --git a/Engine/Plugins/Experimental/Mover/Source/Mover/Public/MovementMode.h b/Engine/Plugins/Experimental/Mover/Source/Mover/Public/MovementMode.h
--- a/Engine/Plugins/Experimental/Mover/Source/Mover/Public/MovementMode.h
+++ b/Engine/Plugins/Experimental/Mover/Source/Mover/Public/MovementMode.h
@@ -71,7 +71,7 @@ public:
 
 	/** Settings object type that this mode depends on. May be shared with other movement modes. When the mode is added to a Mover Component, it will create a shared instance of this settings class. */
 	UPROPERTY(EditDefaultsOnly, Category = Mover, meta = (MustImplement = "/Script/Mover.MovementSettingsInterface"))
-	TSubclassOf<UObject> SharedSettingsClass = nullptr;
+	TArray<TSubclassOf<UObject>> SharedSettingsClasses;
 
 	/** Transition checks for the current mode. Evaluated in order, stopping at the first successful transition check */
 	UPROPERTY(EditAnywhere, BlueprintReadOnly, Instanced, Category = Mover, meta = (FullyExpand = true))
