Physics Mover - Use settings to find default mode names for in built transitions

diff --git a/Engine/Plugins/Experimental/Mover/Source/Mover/Private/PhysicsMover/Modes/PhysicsDrivenFallingMode.cpp b/Engine/Plugins/Experimental/Mover/Source/Mover/Private/PhysicsMover/Modes/PhysicsDrivenFallingMode.cpp
--- a/Engine/Plugins/Experimental/Mover/Source/Mover/Private/PhysicsMover/Modes/PhysicsDrivenFallingMode.cpp
+++ b/Engine/Plugins/Experimental/Mover/Source/Mover/Private/PhysicsMover/Modes/PhysicsDrivenFallingMode.cpp
@@ -103,7 +103,7 @@ void UPhysicsDrivenFallingMode::OnSimulationTick(const FSimulationTickParams& Pa
 
 	if (WaterResult.IsSwimmableVolume() && bStartSwimming && !bIsMovingUp)
 	{
-		OutputState.MovementEndState.NextModeName = DefaultModeNames::Swimming;
+		OutputState.MovementEndState.NextModeName = CommonLegacySettings->SwimmingMovementModeName;
 		OutputState.MovementEndState.RemainingMs = Params.TimeStep.StepMs;
 		return;
 	}
@@ -134,16 +134,12 @@ void UPhysicsDrivenFallingMode::OnSimulationTick(const FSimulationTickParams& Pa
 
 	if (FloorResult.IsWalkableFloor() && bIsFloorWithinReach && !bIsMovingUpRelativeToFloor)
 	{
-		OutputState.MovementEndState.NextModeName = DefaultModeNames::Walking;
+		OutputState.MovementEndState.NextModeName = CommonLegacySettings->GroundMovementModeName;
 		//TargetVel = FVector::VectorPlaneProject(TargetVel, FloorResult.HitResult.Normal);
 		const FPlane MovementPlane(FVector::ZeroVector, UpDir);
 		TargetVel = UMovementUtils::ConstrainToPlane(TargetVel, MovementPlane, false);
 		TargetPos -= (UpDir.Dot(TargetVel * DeltaSeconds) + (FloorResult.FloorDist - TargetHeight)) * UpDir;
 	}
-	else
-	{
-		OutputState.MovementEndState.NextModeName = DefaultModeNames::Falling;
-	}
 
 	OutputState.MovementEndState.RemainingMs = 0.0f;
 	OutputSyncState.MoveDirectionIntent = ProposedMove.bHasDirIntent ? ProposedMove.DirectionIntent : FVector::ZeroVector;
diff --git a/Engine/Plugins/Experimental/Mover/Source/Mover/Private/PhysicsMover/Modes/PhysicsDrivenSwimmingMode.cpp b/Engine/Plugins/Experimental/Mover/Source/Mover/Private/PhysicsMover/Modes/PhysicsDrivenSwimmingMode.cpp
--- a/Engine/Plugins/Experimental/Mover/Source/Mover/Private/PhysicsMover/Modes/PhysicsDrivenSwimmingMode.cpp
+++ b/Engine/Plugins/Experimental/Mover/Source/Mover/Private/PhysicsMover/Modes/PhysicsDrivenSwimmingMode.cpp
@@ -131,20 +131,16 @@ void UPhysicsDrivenSwimmingMode::OnSimulationTick(const FSimulationTickParams& P
 	
 		if (bWalkTrigger && bIsWithinReach)
 		{
-			OutputState.MovementEndState.NextModeName = DefaultModeNames::Walking;
+			OutputState.MovementEndState.NextModeName = CommonLegacySettings->GroundMovementModeName;
 		}
 		else if (bFallTrigger)
 		{
-			OutputState.MovementEndState.NextModeName = DefaultModeNames::Falling;
-		}
-		else
-		{
-			OutputState.MovementEndState.NextModeName = DefaultModeNames::Swimming;
+			OutputState.MovementEndState.NextModeName = CommonLegacySettings->AirMovementModeName;
 		}
 	}
 	else
 	{
-		OutputState.MovementEndState.NextModeName = DefaultModeNames::Falling;
+		OutputState.MovementEndState.NextModeName = CommonLegacySettings->AirMovementModeName;
 	}
 	
 	OutputState.MovementEndState.RemainingMs = 0.0f;
diff --git a/Engine/Plugins/Experimental/Mover/Source/Mover/Private/PhysicsMover/Modes/PhysicsDrivenWalkingMode.cpp b/Engine/Plugins/Experimental/Mover/Source/Mover/Private/PhysicsMover/Modes/PhysicsDrivenWalkingMode.cpp
--- a/Engine/Plugins/Experimental/Mover/Source/Mover/Private/PhysicsMover/Modes/PhysicsDrivenWalkingMode.cpp
+++ b/Engine/Plugins/Experimental/Mover/Source/Mover/Private/PhysicsMover/Modes/PhysicsDrivenWalkingMode.cpp
@@ -166,7 +166,7 @@ void UPhysicsDrivenWalkingMode::FloorCheck(const FMoverDefaultSyncState& SyncSta
 		UPhysicsMovementUtils::FloorSweep_Internal(SyncState.GetLocation_WorldSpace(), FVector::ZeroVector, UpdatedPrimitive, UpDir,
 			ShrinkMultiplier * QueryRadius, FloorSweepDistance, CommonLegacySettings->MaxWalkSlopeCosine, TargetHeight, OutFloorResult, OutWaterResult);
 
-		OutFloorResult.bWalkableFloor = OutFloorResult.bWalkableFloor && CanStepUpOnHitSurface(OutFloorResult);
+		OutFloorResult.bWalkableFloor = OutFloorResult.bWalkableFloor & CanStepUpOnHitSurface(OutFloorResult);
 		return;
 	}
 
@@ -219,7 +219,7 @@ void UPhysicsDrivenWalkingMode::FloorCheck(const FMoverDefaultSyncState& SyncSta
 		UPhysicsMovementUtils::FloorSweep_Internal(SyncState.GetLocation_WorldSpace(), NewDeltaPos, UpdatedPrimitive, UpDir,
 			NewQueryRadius, FloorSweepDistance, CommonLegacySettings->MaxWalkSlopeCosine, TargetHeight, OutFloorResult, OutWaterResult);
 
-		OutFloorResult.bWalkableFloor = OutFloorResult.bWalkableFloor && CanStepUpOnHitSurface(OutFloorResult);
+		OutFloorResult.bWalkableFloor = OutFloorResult.bWalkableFloor & CanStepUpOnHitSurface(OutFloorResult);
 		if (OutFloorResult.bWalkableFloor)
 		{
 			OutDeltaPos = NewDeltaPos;
@@ -239,7 +239,7 @@ void UPhysicsDrivenWalkingMode::FloorCheck(const FMoverDefaultSyncState& SyncSta
 		UPhysicsMovementUtils::FloorSweep_Internal(SyncState.GetLocation_WorldSpace(), NewDeltaPos, UpdatedPrimitive, UpDir,
 			NewQueryRadius, FloorSweepDistance, CommonLegacySettings->MaxWalkSlopeCosine, TargetHeight, OutFloorResult, OutWaterResult);
 
-		OutFloorResult.bWalkableFloor = OutFloorResult.bWalkableFloor && CanStepUpOnHitSurface(OutFloorResult);
+		OutFloorResult.bWalkableFloor = OutFloorResult.bWalkableFloor & CanStepUpOnHitSurface(OutFloorResult);
 		OutDeltaPos = NewDeltaPos;
 	}
 }
@@ -307,7 +307,7 @@ void UPhysicsDrivenWalkingMode::OnSimulationTick(const FSimulationTickParams& Pa
 	
 	if (WaterResult.IsSwimmableVolume() && bStartSwimming)
 	{
-		SwitchToState(DefaultModeNames::Swimming, Params, OutputState);
+		SwitchToState(CommonLegacySettings->SwimmingMovementModeName, Params, OutputState);
 	}
 	else if (FloorResult.IsWalkableFloor())
 	{
@@ -396,9 +396,7 @@ void UPhysicsDrivenWalkingMode::OnSimulationTick(const FSimulationTickParams& Pa
 
 		if (bIsSupported)
 		{
-			OutputState.MovementEndState.NextModeName = DefaultModeNames::Walking;
 			OutputState.MovementEndState.RemainingMs = 0.0f;
-
 			OutputSyncState.MoveDirectionIntent = ProposedMove.bHasDirIntent ? ProposedMove.DirectionIntent : FVector::ZeroVector;
 			OutputSyncState.SetTransforms_WorldSpace(
 				TargetPosition,
@@ -409,13 +407,13 @@ void UPhysicsDrivenWalkingMode::OnSimulationTick(const FSimulationTickParams& Pa
 		else
 		{
 			// Blocking hit but not supported
-			SwitchToState(DefaultModeNames::Falling, Params, OutputState);
+			SwitchToState(CommonLegacySettings->AirMovementModeName, Params, OutputState);
 		}
 	}
 	else
 	{
 		// No water or floor not found
-		SwitchToState(DefaultModeNames::Falling, Params, OutputState);
+		SwitchToState(CommonLegacySettings->AirMovementModeName, Params, OutputState);
 	}
 }
 
diff --git a/Engine/Plugins/Experimental/Mover/Source/Mover/Public/DefaultMovementSet/Settings/CommonLegacyMovementSettings.h b/Engine/Plugins/Experimental/Mover/Source/Mover/Public/DefaultMovementSet/Settings/CommonLegacyMovementSettings.h
--- a/Engine/Plugins/Experimental/Mover/Source/Mover/Public/DefaultMovementSet/Settings/CommonLegacyMovementSettings.h
+++ b/Engine/Plugins/Experimental/Mover/Source/Mover/Public/DefaultMovementSet/Settings/CommonLegacyMovementSettings.h
@@ -25,6 +25,10 @@ public:
 	UPROPERTY(Category="General", EditAnywhere, BlueprintReadWrite)
 	FName AirMovementModeName = DefaultModeNames::Falling;
 
+	// What movement mode to use when airborne.
+	UPROPERTY(Category = "General", EditAnywhere, BlueprintReadWrite)
+	FName SwimmingMovementModeName = DefaultModeNames::Swimming;
+
 	/** Walkable slope angle, represented as cosine(max slope angle) for performance reasons. Ex: for max slope angle of 30 degrees, value is cosine(30 deg) = 0.866 */
 	UPROPERTY(BlueprintReadWrite, EditAnywhere, Category="Ground Movement")
 	float MaxWalkSlopeCosine = 0.71f;
