Physcs Mover - Fix for static analysis error

diff --git a/Engine/Plugins/Experimental/Mover/Source/Mover/Private/PhysicsMover/Modes/PhysicsDrivenWalkingMode.cpp b/Engine/Plugins/Experimental/Mover/Source/Mover/Private/PhysicsMover/Modes/PhysicsDrivenWalkingMode.cpp
--- a/Engine/Plugins/Experimental/Mover/Source/Mover/Private/PhysicsMover/Modes/PhysicsDrivenWalkingMode.cpp
+++ b/Engine/Plugins/Experimental/Mover/Source/Mover/Private/PhysicsMover/Modes/PhysicsDrivenWalkingMode.cpp
@@ -166,7 +166,7 @@ void UPhysicsDrivenWalkingMode::FloorCheck(const FMoverDefaultSyncState& SyncSta
 		UPhysicsMovementUtils::FloorSweep_Internal(SyncState.GetLocation_WorldSpace(), FVector::ZeroVector, UpdatedPrimitive, UpDir,
 			ShrinkMultiplier * QueryRadius, FloorSweepDistance, CommonLegacySettings->MaxWalkSlopeCosine, TargetHeight, OutFloorResult, OutWaterResult);
 
-		OutFloorResult.bWalkableFloor = OutFloorResult.bWalkableFloor & CanStepUpOnHitSurface(OutFloorResult);
+		OutFloorResult.bWalkableFloor = OutFloorResult.bWalkableFloor && CanStepUpOnHitSurface(OutFloorResult);
 		return;
 	}
 
@@ -219,7 +219,7 @@ void UPhysicsDrivenWalkingMode::FloorCheck(const FMoverDefaultSyncState& SyncSta
 		UPhysicsMovementUtils::FloorSweep_Internal(SyncState.GetLocation_WorldSpace(), NewDeltaPos, UpdatedPrimitive, UpDir,
 			NewQueryRadius, FloorSweepDistance, CommonLegacySettings->MaxWalkSlopeCosine, TargetHeight, OutFloorResult, OutWaterResult);
 
-		OutFloorResult.bWalkableFloor = OutFloorResult.bWalkableFloor & CanStepUpOnHitSurface(OutFloorResult);
+		OutFloorResult.bWalkableFloor = OutFloorResult.bWalkableFloor && CanStepUpOnHitSurface(OutFloorResult);
 		if (OutFloorResult.bWalkableFloor)
 		{
 			OutDeltaPos = NewDeltaPos;
@@ -239,7 +239,7 @@ void UPhysicsDrivenWalkingMode::FloorCheck(const FMoverDefaultSyncState& SyncSta
 		UPhysicsMovementUtils::FloorSweep_Internal(SyncState.GetLocation_WorldSpace(), NewDeltaPos, UpdatedPrimitive, UpDir,
 			NewQueryRadius, FloorSweepDistance, CommonLegacySettings->MaxWalkSlopeCosine, TargetHeight, OutFloorResult, OutWaterResult);
 
-		OutFloorResult.bWalkableFloor = OutFloorResult.bWalkableFloor & CanStepUpOnHitSurface(OutFloorResult);
+		OutFloorResult.bWalkableFloor = OutFloorResult.bWalkableFloor && CanStepUpOnHitSurface(OutFloorResult);
 		OutDeltaPos = NewDeltaPos;
 	}
 }
