Mover 2.0: Removing movement mode GetBlackboard functions. Blackboard access should be through the movement component if needed.

diff --git a/Engine/Plugins/Experimental/Mover/Source/Mover/Private/DefaultMovementSet/Modes/FallingMode.cpp b/Engine/Plugins/Experimental/Mover/Source/Mover/Private/DefaultMovementSet/Modes/FallingMode.cpp
--- a/Engine/Plugins/Experimental/Mover/Source/Mover/Private/DefaultMovementSet/Modes/FallingMode.cpp
+++ b/Engine/Plugins/Experimental/Mover/Source/Mover/Private/DefaultMovementSet/Modes/FallingMode.cpp
@@ -37,6 +37,7 @@ void UFallingMode::OnGenerateMove(const FMoverTickStartData& StartState, const F
 	const FCharacterDefaultInputs* CharacterInputs = StartState.InputCmd.InputCollection.FindDataByType<FCharacterDefaultInputs>();
 	const FMoverDefaultSyncState* StartingSyncState = StartState.SyncState.SyncStateCollection.FindDataByType<FMoverDefaultSyncState>();
 	check(StartingSyncState);
+	UMoverComponent* MoverComponent = GetMoverComponent();
 
 	const float DeltaSeconds = TimeStep.StepMs * 0.001f;
 
@@ -92,7 +93,7 @@ void UFallingMode::OnGenerateMove(const FMoverTickStartData& StartState, const F
 		Params.Deceleration = OverTerminalSpeedFallingDeceleration;
 	}
 	
-	UMoverBlackboard* SimBlackboard = GetBlackboard_Mutable();
+	UMoverBlackboard* SimBlackboard = MoverComponent->GetSimBlackboard_Mutable();
 	FFloorCheckResult LastFloorResult;
 	// limit our moveinput based on the floor we're on
 	if (SimBlackboard && SimBlackboard->TryGet(CommonBlackboard::LastFloorResult, LastFloorResult))
@@ -110,7 +111,7 @@ void UFallingMode::OnGenerateMove(const FMoverTickStartData& StartState, const F
 	}
 	
 	OutProposedMove = UAirMovementUtils::ComputeControlledFreeMove(Params);
-	const FVector VelocityWithGravity = StartVelocity + UMovementUtils::ComputeVelocityFromGravity(GetMoverComponent()->GetGravityAcceleration(), DeltaSeconds);
+	const FVector VelocityWithGravity = StartVelocity + UMovementUtils::ComputeVelocityFromGravity(MoverComponent->GetGravityAcceleration(), DeltaSeconds);
 
 	//  If we are going faster than TerminalVerticalVelocity apply VerticalFallingDeceleration otherwise reset Z velocity to before we applied deceleration 
 	if (VelocityWithGravity.GetAbs().Z > TerminalVerticalSpeed)
@@ -135,6 +136,7 @@ void UFallingMode::OnGenerateMove(const FMoverTickStartData& StartState, const F
 
 void UFallingMode::OnSimulationTick(const FSimulationTickParams& Params, FMoverTickEndData& OutputState)
 {
+	UMoverComponent* MoverComponent = GetMoverComponent();
 	const FMoverTickStartData& StartState = Params.StartState;
 	USceneComponent* UpdatedComponent = Params.MovingComps.UpdatedComponent.Get();
 	FProposedMove ProposedMove = Params.ProposedMove;
@@ -152,7 +154,7 @@ void UFallingMode::OnSimulationTick(const FSimulationTickParams& Params, FMoverT
 	FMovementRecord MoveRecord;
 	MoveRecord.SetDeltaSeconds(DeltaSeconds);
 	
-	UMoverBlackboard* SimBlackboard = GetBlackboard_Mutable();
+	UMoverBlackboard* SimBlackboard = MoverComponent->GetSimBlackboard_Mutable();
 
 	SimBlackboard->Invalidate(CommonBlackboard::LastFloorResult);	// falling = no valid floor
 	SimBlackboard->Invalidate(CommonBlackboard::LastFoundDynamicMovementBase);
@@ -204,8 +206,7 @@ void UFallingMode::OnSimulationTick(const FSimulationTickParams& Params, FMoverT
 		
 		LandingFloor.HitResult = Hit;
 		SimBlackboard->Set(CommonBlackboard::LastFloorResult, LandingFloor);
-
-		UMoverComponent* MoverComponent = GetMoverComponent();
+		
 		FMoverOnImpactParams ImpactParams(DefaultModeNames::Falling, Hit, MoveDelta);
 		MoverComponent->HandleImpact(ImpactParams);
 
@@ -250,7 +251,8 @@ void UFallingMode::OnUnregistered()
 
 void UFallingMode::ProcessLanded(const FFloorCheckResult& FloorResult, FVector& Velocity, FRelativeBaseInfo& BaseInfo, FMoverTickEndData& TickEndData) const
 {
-	UMoverBlackboard* SimBlackboard = GetBlackboard_Mutable();
+	const UMoverComponent* MoverComp = GetMoverComponent();
+	UMoverBlackboard* SimBlackboard = MoverComp->GetSimBlackboard_Mutable();
 
 	FName NextMovementMode = NAME_None; 
 	// if we can walk on the floor we landed on
@@ -258,7 +260,7 @@ void UFallingMode::ProcessLanded(const FFloorCheckResult& FloorResult, FVector&
 	{
 		if (bCancelVerticalSpeedOnLanding)
 		{
-			const FPlane MovementPlane(FVector::ZeroVector, GetMoverComponent()->GetUpDirection());
+			const FPlane MovementPlane(FVector::ZeroVector, MoverComp->GetUpDirection());
 			Velocity = UMovementUtils::ConstrainToPlane(Velocity, MovementPlane, false);
 		}
 		else
@@ -290,7 +292,7 @@ void UFallingMode::ProcessLanded(const FFloorCheckResult& FloorResult, FVector&
 
 void UFallingMode::CaptureFinalState(USceneComponent* UpdatedComponent, const FMoverDefaultSyncState& StartSyncState, const FFloorCheckResult& FloorResult, float DeltaSeconds, float DeltaSecondsUsed, FMoverDefaultSyncState& OutputSyncState, FMoverTickEndData& TickEndData, FMovementRecord& Record) const
 {
-	UMoverBlackboard* SimBlackboard = GetBlackboard_Mutable();
+	UMoverBlackboard* SimBlackboard = GetMoverComponent()->GetSimBlackboard_Mutable();
 
 	const FVector FinalLocation = UpdatedComponent->GetComponentLocation();
 
diff --git a/Engine/Plugins/Experimental/Mover/Source/Mover/Private/DefaultMovementSet/Modes/FlyingMode.cpp b/Engine/Plugins/Experimental/Mover/Source/Mover/Private/DefaultMovementSet/Modes/FlyingMode.cpp
--- a/Engine/Plugins/Experimental/Mover/Source/Mover/Private/DefaultMovementSet/Modes/FlyingMode.cpp
+++ b/Engine/Plugins/Experimental/Mover/Source/Mover/Private/DefaultMovementSet/Modes/FlyingMode.cpp
@@ -63,6 +63,7 @@ void UFlyingMode::OnGenerateMove(const FMoverTickStartData& StartState, const FM
 
 void UFlyingMode::OnSimulationTick(const FSimulationTickParams& Params, FMoverTickEndData& OutputState)
 {
+	const UMoverComponent* MoverComp = GetMoverComponent();
 	const FMoverTickStartData& StartState = Params.StartState;
 	USceneComponent* UpdatedComponent = Params.MovingComps.UpdatedComponent.Get();
 	FProposedMove ProposedMove = Params.ProposedMove;
@@ -78,7 +79,7 @@ void UFlyingMode::OnSimulationTick(const FSimulationTickParams& Params, FMoverTi
 	FMovementRecord MoveRecord;
 	MoveRecord.SetDeltaSeconds(DeltaSeconds);
 
-	UMoverBlackboard* SimBlackboard = GetBlackboard_Mutable();
+	UMoverBlackboard* SimBlackboard = MoverComp->GetSimBlackboard_Mutable();
 
 	SimBlackboard->Invalidate(CommonBlackboard::LastFloorResult);	// flying = no valid floor
 	SimBlackboard->Invalidate(CommonBlackboard::LastFoundDynamicMovementBase);
diff --git a/Engine/Plugins/Experimental/Mover/Source/Mover/Private/DefaultMovementSet/Modes/SwimmingMode.cpp b/Engine/Plugins/Experimental/Mover/Source/Mover/Private/DefaultMovementSet/Modes/SwimmingMode.cpp
--- a/Engine/Plugins/Experimental/Mover/Source/Mover/Private/DefaultMovementSet/Modes/SwimmingMode.cpp
+++ b/Engine/Plugins/Experimental/Mover/Source/Mover/Private/DefaultMovementSet/Modes/SwimmingMode.cpp
@@ -29,7 +29,7 @@ void USwimmingMode::OnGenerateMove(const FMoverTickStartData& StartState, const
 	
 	const float DeltaSeconds = TimeStep.StepMs * 0.001f;
 	
-	UMoverBlackboard* SimBlackboard = GetBlackboard_Mutable();
+	UMoverBlackboard* SimBlackboard = MoverComp->GetSimBlackboard_Mutable();
 	FWaterCheckResult LastWaterResult;
 	
 	const float CapsuleHalfHeight = MoverComp->GetOwner()->GetSimpleCollisionHalfHeight();
diff --git a/Engine/Plugins/Experimental/Mover/Source/Mover/Private/DefaultMovementSet/Modes/WalkingMode.cpp b/Engine/Plugins/Experimental/Mover/Source/Mover/Private/DefaultMovementSet/Modes/WalkingMode.cpp
--- a/Engine/Plugins/Experimental/Mover/Source/Mover/Private/DefaultMovementSet/Modes/WalkingMode.cpp
+++ b/Engine/Plugins/Experimental/Mover/Source/Mover/Private/DefaultMovementSet/Modes/WalkingMode.cpp
@@ -36,7 +36,7 @@ void UWalkingMode::OnGenerateMove(const FMoverTickStartData& StartState, const F
 	FFloorCheckResult LastFloorResult;
 	FVector MovementNormal;
 
-	UMoverBlackboard* SimBlackboard = GetBlackboard_Mutable();
+	UMoverBlackboard* SimBlackboard = MoverComp->GetSimBlackboard_Mutable();
 
 	// Try to use the floor as the basis for the intended move direction (i.e. try to walk along slopes, rather than into them)
 	if (SimBlackboard && SimBlackboard->TryGet(CommonBlackboard::LastFloorResult, LastFloorResult) && LastFloorResult.IsWalkableFloor())
@@ -137,7 +137,7 @@ void UWalkingMode::OnSimulationTick(const FSimulationTickParams& Params, FMoverT
 	MoveRecord.SetDeltaSeconds(DeltaSeconds);
 
 	FFloorCheckResult CurrentFloor;
-	UMoverBlackboard* SimBlackboard = GetBlackboard_Mutable();
+	UMoverBlackboard* SimBlackboard = MoverComp->GetSimBlackboard_Mutable();
 
 	// If we don't have cached floor information, we need to search for it again
 	if (!SimBlackboard->TryGet(CommonBlackboard::LastFloorResult, CurrentFloor))
@@ -331,7 +331,8 @@ void UWalkingMode::CaptureFinalState(USceneComponent* UpdatedComponent, bool bDi
 {
 	FRelativeBaseInfo PriorBaseInfo;
 
-	UMoverBlackboard* SimBlackboard = GetBlackboard_Mutable();
+	const UMoverComponent* MoverComp = GetMoverComponent();
+	UMoverBlackboard* SimBlackboard = MoverComp->GetSimBlackboard_Mutable();
 
 	const bool bHasPriorBaseInfo = SimBlackboard->TryGet(CommonBlackboard::LastFoundDynamicMovementBase, PriorBaseInfo);
 
@@ -375,7 +376,8 @@ FRelativeBaseInfo UWalkingMode::UpdateFloorAndBaseInfo(const FFloorCheckResult&
 {
 	FRelativeBaseInfo ReturnBaseInfo;
 
-	UMoverBlackboard* SimBlackboard = GetBlackboard_Mutable();
+	const UMoverComponent* MoverComp = GetMoverComponent();
+	UMoverBlackboard* SimBlackboard = MoverComp->GetSimBlackboard_Mutable();
 
 	SimBlackboard->Set(CommonBlackboard::LastFloorResult, FloorResult);
 
diff --git a/Engine/Plugins/Experimental/Mover/Source/Mover/Private/MovementMode.cpp b/Engine/Plugins/Experimental/Mover/Source/Mover/Private/MovementMode.cpp
--- a/Engine/Plugins/Experimental/Mover/Source/Mover/Private/MovementMode.cpp
+++ b/Engine/Plugins/Experimental/Mover/Source/Mover/Private/MovementMode.cpp
@@ -95,24 +95,6 @@ UMoverComponent* UBaseMovementMode::GetMoverComponent() const
 	return CastChecked<UMoverComponent>(GetOuter());
 }
 
-
-const UMoverBlackboard* UBaseMovementMode::GetBlackboard() const
-{
-	return GetBlackboard_Mutable();
-}
-
-
-UMoverBlackboard* UBaseMovementMode::GetBlackboard_Mutable() const
-{
-	if (UMoverComponent* OwnerComponent = GetMoverComponent())
-	{
-		return OwnerComponent->SimBlackboard.Get();
-	}
-
-	return nullptr;
-}
-
-
 #if WITH_EDITOR
 EDataValidationResult UBaseMovementMode::IsDataValid(FDataValidationContext& Context) const
 {
diff --git a/Engine/Plugins/Experimental/Mover/Source/Mover/Private/PhysicsMover/Modes/PhysicsDrivenFallingMode.cpp b/Engine/Plugins/Experimental/Mover/Source/Mover/Private/PhysicsMover/Modes/PhysicsDrivenFallingMode.cpp
--- a/Engine/Plugins/Experimental/Mover/Source/Mover/Private/PhysicsMover/Modes/PhysicsDrivenFallingMode.cpp
+++ b/Engine/Plugins/Experimental/Mover/Source/Mover/Private/PhysicsMover/Modes/PhysicsDrivenFallingMode.cpp
@@ -56,6 +56,7 @@ void UPhysicsDrivenFallingMode::OnSimulationTick(const FSimulationTickParams& Pa
 	USceneComponent* UpdatedComponent = Params.MovingComps.UpdatedComponent.Get();
 	UPrimitiveComponent* UpdatedPrimitive = Params.MovingComps.UpdatedPrimitive.Get();
 	FProposedMove ProposedMove = Params.ProposedMove;
+	const UMoverComponent* MoverComp = GetMoverComponent();
 
 	const FCharacterDefaultInputs* CharacterInputs = StartState.InputCmd.InputCollection.FindDataByType<FCharacterDefaultInputs>();
 	const FMoverDefaultSyncState* StartingSyncState = StartState.SyncState.SyncStateCollection.FindDataByType<FMoverDefaultSyncState>();
@@ -64,11 +65,11 @@ void UPhysicsDrivenFallingMode::OnSimulationTick(const FSimulationTickParams& Pa
 	FMoverDefaultSyncState& OutputSyncState = OutputState.SyncState.SyncStateCollection.FindOrAddMutableDataByType<FMoverDefaultSyncState>();
 
 	const float DeltaSeconds = Params.TimeStep.StepMs * 0.001f;
-	const FVector UpDir = GetMoverComponent()->GetUpDirection();
+	const FVector UpDir = MoverComp->GetUpDirection();
 
 	// Floor query
 
-	UMoverBlackboard* SimBlackboard = GetBlackboard_Mutable();
+	UMoverBlackboard* SimBlackboard = MoverComp->GetSimBlackboard_Mutable();
 	if (!SimBlackboard)
 	{
 		OutputSyncState = *StartingSyncState;
diff --git a/Engine/Plugins/Experimental/Mover/Source/Mover/Private/PhysicsMover/Modes/PhysicsDrivenFlyingMode.cpp b/Engine/Plugins/Experimental/Mover/Source/Mover/Private/PhysicsMover/Modes/PhysicsDrivenFlyingMode.cpp
--- a/Engine/Plugins/Experimental/Mover/Source/Mover/Private/PhysicsMover/Modes/PhysicsDrivenFlyingMode.cpp
+++ b/Engine/Plugins/Experimental/Mover/Source/Mover/Private/PhysicsMover/Modes/PhysicsDrivenFlyingMode.cpp
@@ -49,6 +49,7 @@ EDataValidationResult UPhysicsDrivenFlyingMode::IsDataValid(FDataValidationConte
 
 void UPhysicsDrivenFlyingMode::OnSimulationTick(const FSimulationTickParams& Params, FMoverTickEndData& OutputState)
 {
+	const UMoverComponent* MoverComp = GetMoverComponent();
 	const FMoverTickStartData& StartState = Params.StartState;
 	USceneComponent* UpdatedComponent = Params.MovingComps.UpdatedComponent.Get();
 	UPrimitiveComponent* UpdatedPrimitive = Params.MovingComps.UpdatedPrimitive.Get();
@@ -60,11 +61,11 @@ void UPhysicsDrivenFlyingMode::OnSimulationTick(const FSimulationTickParams& Par
 	FMoverDefaultSyncState& OutputSyncState = OutputState.SyncState.SyncStateCollection.FindOrAddMutableDataByType<FMoverDefaultSyncState>();
 
 	const float DeltaSeconds = Params.TimeStep.StepMs * 0.001f;
-	const FVector UpDir = GetMoverComponent()->GetUpDirection();
+	const FVector UpDir = MoverComp->GetUpDirection();
 
 	// Don't need a floor query - just invalidate the blackboard to ensure we don't use an old result elsewhere
 
-	if (UMoverBlackboard* SimBlackboard = GetBlackboard_Mutable())
+	if (UMoverBlackboard* SimBlackboard = MoverComp->GetSimBlackboard_Mutable())
 	{
 		SimBlackboard->Invalidate(CommonBlackboard::LastFloorResult);
 		SimBlackboard->Invalidate(CommonBlackboard::LastWaterResult);
diff --git a/Engine/Plugins/Experimental/Mover/Source/Mover/Private/PhysicsMover/Modes/PhysicsDrivenSwimmingMode.cpp b/Engine/Plugins/Experimental/Mover/Source/Mover/Private/PhysicsMover/Modes/PhysicsDrivenSwimmingMode.cpp
--- a/Engine/Plugins/Experimental/Mover/Source/Mover/Private/PhysicsMover/Modes/PhysicsDrivenSwimmingMode.cpp
+++ b/Engine/Plugins/Experimental/Mover/Source/Mover/Private/PhysicsMover/Modes/PhysicsDrivenSwimmingMode.cpp
@@ -58,7 +58,7 @@ void UPhysicsDrivenSwimmingMode::OnSimulationTick(const FSimulationTickParams& P
 	UPrimitiveComponent* UpdatedPrimitive = Params.MovingComps.UpdatedPrimitive.Get();
 	FProposedMove ProposedMove = Params.ProposedMove;
 	
-	const FVector UpDir = GetMoverComponent()->GetUpDirection();
+	const FVector UpDir = MoverComp->GetUpDirection();
 	
 	const FCharacterDefaultInputs* CharacterInputs = StartState.InputCmd.InputCollection.FindDataByType<FCharacterDefaultInputs>();
 
@@ -75,7 +75,7 @@ void UPhysicsDrivenSwimmingMode::OnSimulationTick(const FSimulationTickParams& P
 		return;
 	}
 	
-	UMoverBlackboard* SimBlackboard = GetBlackboard_Mutable();
+	UMoverBlackboard* SimBlackboard = MoverComp->GetSimBlackboard_Mutable();
 	if (!SimBlackboard)
 	{
 		OutputSyncState = *StartingSyncState;
diff --git a/Engine/Plugins/Experimental/Mover/Source/Mover/Private/PhysicsMover/Modes/PhysicsDrivenWalkingMode.cpp b/Engine/Plugins/Experimental/Mover/Source/Mover/Private/PhysicsMover/Modes/PhysicsDrivenWalkingMode.cpp
--- a/Engine/Plugins/Experimental/Mover/Source/Mover/Private/PhysicsMover/Modes/PhysicsDrivenWalkingMode.cpp
+++ b/Engine/Plugins/Experimental/Mover/Source/Mover/Private/PhysicsMover/Modes/PhysicsDrivenWalkingMode.cpp
@@ -267,7 +267,7 @@ void UPhysicsDrivenWalkingMode::OnSimulationTick(const FSimulationTickParams& Pa
 	UPrimitiveComponent* UpdatedPrimitive = Params.MovingComps.UpdatedPrimitive.Get();
 	FProposedMove ProposedMove = Params.ProposedMove;
 
-	const FVector UpDir = GetMoverComponent()->GetUpDirection();
+	const FVector UpDir = MoverComp->GetUpDirection();
 
 	const FMoverDefaultSyncState* StartingSyncState = StartState.SyncState.SyncStateCollection.FindDataByType<FMoverDefaultSyncState>();
 	check(StartingSyncState);
@@ -276,7 +276,7 @@ void UPhysicsDrivenWalkingMode::OnSimulationTick(const FSimulationTickParams& Pa
 
 	const float DeltaSeconds = Params.TimeStep.StepMs * 0.001f;
 
-	UMoverBlackboard* SimBlackboard = GetBlackboard_Mutable();
+	UMoverBlackboard* SimBlackboard = MoverComp->GetSimBlackboard_Mutable();
 	if (!SimBlackboard)
 	{
 		OutputSyncState = *StartingSyncState;
diff --git a/Engine/Plugins/Experimental/Mover/Source/Mover/Public/MovementMode.h b/Engine/Plugins/Experimental/Mover/Source/Mover/Public/MovementMode.h
--- a/Engine/Plugins/Experimental/Mover/Source/Mover/Public/MovementMode.h
+++ b/Engine/Plugins/Experimental/Mover/Source/Mover/Public/MovementMode.h
@@ -58,13 +58,6 @@ public:
 		return CastChecked<T>(GetOuter());
 	}
 
-	UFUNCTION(BlueprintCallable, Category=Mover, meta=(DisplayName="Get Blackboard"))
-	const UMoverBlackboard* GetBlackboard() const;
-
-	UFUNCTION(BlueprintCallable, Category=Mover, meta=(DisplayName="Get Blackboard"))
-	UMoverBlackboard* GetBlackboard_Mutable() const;
-
-
 #if WITH_EDITOR
 	virtual EDataValidationResult IsDataValid(class FDataValidationContext& Context) const override;
 #endif // WITH_EDITOR
