Mover 2.0 adding Flying gameplay tag and adding some state checking functions to CharacterMoverComponent

diff --git a/Engine/Plugins/Experimental/Mover/Source/Mover/Private/DefaultMovementSet/CharacterMoverComponent.cpp b/Engine/Plugins/Experimental/Mover/Source/Mover/Private/DefaultMovementSet/CharacterMoverComponent.cpp
--- a/Engine/Plugins/Experimental/Mover/Source/Mover/Private/DefaultMovementSet/CharacterMoverComponent.cpp
+++ b/Engine/Plugins/Experimental/Mover/Source/Mover/Private/DefaultMovementSet/CharacterMoverComponent.cpp
@@ -27,6 +27,21 @@ void UCharacterMoverComponent::BeginPlay()
 	OnPreSimulationTick.AddDynamic(this, &UCharacterMoverComponent::OnMoverPreSimulationTick);
 }
 
+bool UCharacterMoverComponent::IsCrouching() const
+{
+	return false; // Crouching not yet supported by Mover
+}
+
+bool UCharacterMoverComponent::IsFlying() const
+{
+	if (const UBaseMovementMode* Mode = GetMovementMode())
+	{
+		return Mode->GameplayTags.HasTag(Mover_IsFlying);
+	}
+
+	return false;
+}
+
 bool UCharacterMoverComponent::IsFalling() const
 {
 	if (const UBaseMovementMode* Mode = GetMovementMode())
diff --git a/Engine/Plugins/Experimental/Mover/Source/Mover/Private/DefaultMovementSet/Modes/FlyingMode.cpp b/Engine/Plugins/Experimental/Mover/Source/Mover/Private/DefaultMovementSet/Modes/FlyingMode.cpp
--- a/Engine/Plugins/Experimental/Mover/Source/Mover/Private/DefaultMovementSet/Modes/FlyingMode.cpp
+++ b/Engine/Plugins/Experimental/Mover/Source/Mover/Private/DefaultMovementSet/Modes/FlyingMode.cpp
@@ -13,6 +13,7 @@ UFlyingMode::UFlyingMode(const FObjectInitializer& ObjectInitializer)
 	: Super(ObjectInitializer)
 {
 	GameplayTags.AddTag(Mover_IsInAir);
+	GameplayTags.AddTag(Mover_IsFlying);
 }
 
 void UFlyingMode::OnGenerateMove(const FMoverTickStartData& StartState, const FMoverTimeStep& TimeStep, FProposedMove& OutProposedMove) const
diff --git a/Engine/Plugins/Experimental/Mover/Source/Mover/Private/MoverTypes.cpp b/Engine/Plugins/Experimental/Mover/Source/Mover/Private/MoverTypes.cpp
--- a/Engine/Plugins/Experimental/Mover/Source/Mover/Private/MoverTypes.cpp
+++ b/Engine/Plugins/Experimental/Mover/Source/Mover/Private/MoverTypes.cpp
@@ -10,6 +10,7 @@
 UE_DEFINE_GAMEPLAY_TAG_COMMENT(Mover_IsOnGround, "Mover.IsOnGround", "Default Mover state flag indicating character is on the ground.");
 UE_DEFINE_GAMEPLAY_TAG_COMMENT(Mover_IsInAir, "Mover.IsInAir", "Default Mover state flag indicating character is in the air.");
 UE_DEFINE_GAMEPLAY_TAG_COMMENT(Mover_IsFalling, "Mover.IsFalling", "Default Mover state flag indicating character is falling.");
+UE_DEFINE_GAMEPLAY_TAG_COMMENT(Mover_IsFlying, "Mover.IsFlying", "Default Mover state flag indicating character is flying.");
 UE_DEFINE_GAMEPLAY_TAG_COMMENT(Mover_IsSwimming, "Mover.IsSwimming", "Default Mover state flag indicating character is swimming.");
 
 FMoverOnImpactParams::FMoverOnImpactParams() 
diff --git a/Engine/Plugins/Experimental/Mover/Source/Mover/Public/DefaultMovementSet/CharacterMoverComponent.h b/Engine/Plugins/Experimental/Mover/Source/Mover/Public/DefaultMovementSet/CharacterMoverComponent.h
--- a/Engine/Plugins/Experimental/Mover/Source/Mover/Public/DefaultMovementSet/CharacterMoverComponent.h
+++ b/Engine/Plugins/Experimental/Mover/Source/Mover/Public/DefaultMovementSet/CharacterMoverComponent.h
@@ -15,6 +15,14 @@ public:
 	
 	virtual void BeginPlay() override;
 	
+	/** Returns true if currently crouching */ 
+	UFUNCTION(BlueprintCallable, Category = Mover)
+	virtual bool IsCrouching() const;
+
+	/** Returns true if currently flying (moving through a non-fluid volume without resting on the ground) */
+	UFUNCTION(BlueprintPure, Category = Mover)
+	virtual bool IsFlying() const;
+	
 	// Is this actor in a falling state? Note that this includes upwards motion induced by jumping.
 	UFUNCTION(BlueprintPure, Category = Mover)
 	virtual bool IsFalling() const;
diff --git a/Engine/Plugins/Experimental/Mover/Source/Mover/Public/MoverTypes.h b/Engine/Plugins/Experimental/Mover/Source/Mover/Public/MoverTypes.h
--- a/Engine/Plugins/Experimental/Mover/Source/Mover/Public/MoverTypes.h
+++ b/Engine/Plugins/Experimental/Mover/Source/Mover/Public/MoverTypes.h
@@ -14,6 +14,7 @@
 MOVER_API UE_DECLARE_GAMEPLAY_TAG_EXTERN(Mover_IsOnGround);
 MOVER_API UE_DECLARE_GAMEPLAY_TAG_EXTERN(Mover_IsInAir);
 MOVER_API UE_DECLARE_GAMEPLAY_TAG_EXTERN(Mover_IsFalling);
+MOVER_API UE_DECLARE_GAMEPLAY_TAG_EXTERN(Mover_IsFlying);
 MOVER_API UE_DECLARE_GAMEPLAY_TAG_EXTERN(Mover_IsSwimming);
 
 // Struct to hold params for when an impact happens. This contains all of the data for impacts including what gets passed to the FMover_OnImpact delegate
