Mover 2.0: Adding Pathfollowing and NavMovement support. Also adding some example content to show/test the feature.

#virtualized

diff --git a/Engine/Plugins/Experimental/Mover/Source/Mover/Private/DefaultMovementSet/NavMoverComponent.cpp b/Engine/Plugins/Experimental/Mover/Source/Mover/Private/DefaultMovementSet/NavMoverComponent.cpp
new file mode 100644
--- /dev/null
+++ b/Engine/Plugins/Experimental/Mover/Source/Mover/Private/DefaultMovementSet/NavMoverComponent.cpp
@@ -0,0 +1,248 @@
+﻿// Copyright Epic Games, Inc. All Rights Reserved.
+
+#include "DefaultMovementSet/NavMoverComponent.h"
+
+#include "AI/NavigationSystemBase.h"
+#include "AI/Navigation/PathFollowingAgentInterface.h"
+#include "Components/CapsuleComponent.h"
+#include "DefaultMovementSet/InstantMovementEffects/BasicInstantMovementEffects.h"
+#include "DefaultMovementSet/Settings/CommonLegacyMovementSettings.h"
+#include "MoveLibrary/MovementUtils.h"
+
+UNavMoverComponent::UNavMoverComponent()
+{
+	bWantsInitializeComponent = true;
+	bAutoActivate = true;
+}
+
+void UNavMoverComponent::InitializeComponent()
+{
+	Super::InitializeComponent();
+
+	if (const AActor* MovementCompOwner = GetOwner())
+	{
+		MoverComponent = MovementCompOwner->FindComponentByClass<UMoverComponent>();
+	}
+	
+	if (!MoverComponent)
+	{
+		UE_LOG(LogMover, Warning, TEXT("NavMoverComponent on %s could not find a valid MoverComponent and will not function properly."), *GetNameSafe(GetOwner()));
+	}
+}
+
+float UNavMoverComponent::GetMaxSpeedForNavMovement() const
+{
+	const UCommonLegacyMovementSettings* MovementSettings = MoverComponent->FindSharedSettings_Mutable<UCommonLegacyMovementSettings>();
+	return MovementSettings ? MovementSettings->MaxSpeed : 0;
+}
+
+void UNavMoverComponent::StopMovementImmediately()
+{
+	TSharedPtr<FApplyVelocityEffect> VelocityEffect = MakeShared<FApplyVelocityEffect>();
+	MoverComponent->QueueInstantMovementEffect(VelocityEffect);
+}
+
+FVector UNavMoverComponent::GetLocation() const
+{
+	const USceneComponent* UpdatedComponent = Cast<USceneComponent>(GetUpdatedObject());
+	return UpdatedComponent ? UpdatedComponent->GetComponentLocation() : FVector(FLT_MAX);
+}
+
+FVector UNavMoverComponent::GetFeetLocation() const
+{
+	const USceneComponent* UpdatedComponent = Cast<USceneComponent>(GetUpdatedObject());
+	return UpdatedComponent ? (UpdatedComponent->GetComponentLocation() - FVector(0,0,UpdatedComponent->Bounds.BoxExtent.Z)) : FNavigationSystem::InvalidLocation;
+}
+
+FBasedPosition UNavMoverComponent::GetFeetLocationBased() const
+{
+	FBasedPosition BasedPosition(NULL, GetFeetLocation());
+	
+	FRelativeBaseInfo MovementBaseInfo;
+	if (const UMoverBlackboard* Blackboard = MoverComponent->GetSimBlackboard())
+	{
+		if (Blackboard->TryGet(CommonBlackboard::LastFoundDynamicMovementBase, MovementBaseInfo)) 
+		{
+			BasedPosition.Base = MovementBaseInfo.MovementBase->GetOwner();
+			BasedPosition.Position = MovementBaseInfo.Location;
+			BasedPosition.CachedBaseLocation = MovementBaseInfo.ContactLocalPosition;
+			BasedPosition.CachedBaseRotation = MovementBaseInfo.Rotation.Rotator();
+		}
+	}
+
+	return BasedPosition;
+}
+
+void UNavMoverComponent::UpdateNavAgent(const UObject& ObjectToUpdateFrom)
+{
+	if (!NavMovementProperties.bUpdateNavAgentWithOwnersCollision)
+	{
+		return;
+	}
+
+	// initialize properties from navigation system
+	NavAgentProps.NavWalkingSearchHeightScale = FNavigationSystem::GetDefaultSupportedAgent().NavWalkingSearchHeightScale;
+	
+	if (const UCapsuleComponent* CapsuleComponent = Cast<UCapsuleComponent>(&ObjectToUpdateFrom))
+	{
+		NavAgentProps.AgentRadius = CapsuleComponent->GetScaledCapsuleRadius();
+		NavAgentProps.AgentHeight = CapsuleComponent->GetScaledCapsuleHalfHeight() * 2.f;;
+	}
+	else if (const AActor* ObjectAsActor = Cast<AActor>(&ObjectToUpdateFrom))
+	{
+		ensureMsgf(&ObjectToUpdateFrom == GetOwner(), TEXT("Object passed to UpdateNavAgent should be the owner actor of the Nav Movement Component"));
+		// Can't call GetSimpleCollisionCylinder(), because no components will be registered.
+		float BoundRadius, BoundHalfHeight;	
+		ObjectAsActor->GetSimpleCollisionCylinder(BoundRadius, BoundHalfHeight);
+		NavAgentProps.AgentRadius = BoundRadius;
+		NavAgentProps.AgentHeight = BoundHalfHeight * 2.f;
+	}
+}
+
+void UNavMoverComponent::RequestDirectMove(const FVector& MoveVelocity, bool bForceMaxSpeed)
+{
+	if (MoveVelocity.SizeSquared() < UE_KINDA_SMALL_NUMBER)
+	{
+		return;
+	}
+
+	if (IsFalling())
+	{
+		const FVector FallVelocity = MoveVelocity.GetClampedToMaxSize(GetMaxSpeedForNavMovement());
+		// TODO: NS - we may eventually need something to help with air control and pathfinding
+		//PerformAirControlForPathFollowing(FallVelocity, FallVelocity.Z);
+		CachedNavMoveInputVelocity = FallVelocity;
+		return;
+	}
+
+	CachedNavMoveInputVelocity = MoveVelocity;
+	bRequestedNavMovement = true;
+	
+	if (IsMovingOnGround())
+	{
+		const FPlane MovementPlane(FVector::ZeroVector, FVector::UpVector);
+		CachedNavMoveInputVelocity = UMovementUtils::ConstrainToPlane(CachedNavMoveInputVelocity, MovementPlane, true);
+	}
+}
+
+void UNavMoverComponent::RequestPathMove(const FVector& MoveInput)
+{
+	FVector AdjustedMoveInput(MoveInput);
+
+	// preserve magnitude when moving on ground/falling and requested input has Z component
+	// see ConstrainInputAcceleration for details
+	if (MoveInput.Z != 0.f && (IsMovingOnGround() || IsFalling()))
+	{
+		const float Mag = MoveInput.Size();
+		AdjustedMoveInput = MoveInput.GetSafeNormal2D() * Mag;
+	}
+
+	bRequestedNavMovement = true;
+	CachedNavMoveInputIntent += AdjustedMoveInput;
+}
+
+bool UNavMoverComponent::CanStopPathFollowing() const
+{
+	return true;
+}
+
+void UNavMoverComponent::SetPathFollowingAgent(IPathFollowingAgentInterface* InPathFollowingAgent)
+{
+	PathFollowingComp = Cast<UObject>(InPathFollowingAgent);
+}
+
+IPathFollowingAgentInterface* UNavMoverComponent::GetPathFollowingAgent()
+{
+	return Cast<IPathFollowingAgentInterface>(PathFollowingComp);
+}
+
+const IPathFollowingAgentInterface* UNavMoverComponent::GetPathFollowingAgent() const
+{
+	return Cast<const IPathFollowingAgentInterface>(PathFollowingComp);
+}
+
+const FNavAgentProperties& UNavMoverComponent::GetNavAgentPropertiesRef() const
+{
+	return NavAgentProps;
+}
+
+FNavAgentProperties& UNavMoverComponent::GetNavAgentPropertiesRef()
+{
+	return NavAgentProps;
+}
+
+void UNavMoverComponent::ResetMoveState()
+{
+	MovementState = NavAgentProps;
+}
+
+bool UNavMoverComponent::CanStartPathFollowing() const
+{
+	return true;
+}
+
+bool UNavMoverComponent::IsCrouching() const
+{
+	// TODO: Mover doesn't support crouching yet so return false for now
+	return false;
+}
+
+bool UNavMoverComponent::IsFalling() const
+{
+	if (const UBaseMovementMode* Mode = MoverComponent->GetMovementMode())
+	{
+		return Mode->GameplayTags.HasTag(Mover_IsFalling);
+	}
+
+	return false;
+}
+
+bool UNavMoverComponent::IsMovingOnGround() const
+{
+	if (const UBaseMovementMode* Mode = MoverComponent->GetMovementMode())
+	{
+		return Mode->GameplayTags.HasTag(Mover_IsOnGround);
+	}
+
+	return false;
+}
+
+bool UNavMoverComponent::IsSwimming() const
+{
+	if (const UBaseMovementMode* Mode = MoverComponent->GetMovementMode())
+	{
+		return Mode->GameplayTags.HasTag(Mover_IsSwimming);
+	}
+
+	return false;
+}
+
+bool UNavMoverComponent::IsFlying() const
+{
+	if (const UBaseMovementMode* Mode = MoverComponent->GetMovementMode())
+	{
+		return Mode->GameplayTags.HasTag(Mover_IsFlying);
+	}
+
+	return false;
+}
+
+void UNavMoverComponent::GetSimpleCollisionCylinder(float& CollisionRadius, float& CollisionHalfHeight) const
+{
+	GetOwner()->GetSimpleCollisionCylinder(CollisionRadius, CollisionHalfHeight);
+}
+
+FVector UNavMoverComponent::GetSimpleCollisionCylinderExtent() const
+{
+	return GetOwner()->GetSimpleCollisionCylinderExtent();
+}
+
+FVector UNavMoverComponent::GetForwardVector() const
+{
+	return GetOwner()->GetActorForwardVector();
+}
+
+FVector UNavMoverComponent::GetVelocityForNavMovement() const
+{
+	return MoverComponent->GetVelocity();
+}
diff --git a/Engine/Plugins/Experimental/Mover/Source/Mover/Public/DefaultMovementSet/NavMoverComponent.h b/Engine/Plugins/Experimental/Mover/Source/Mover/Public/DefaultMovementSet/NavMoverComponent.h
new file mode 100644
--- /dev/null
+++ b/Engine/Plugins/Experimental/Mover/Source/Mover/Public/DefaultMovementSet/NavMoverComponent.h
@@ -0,0 +1,145 @@
+﻿// Copyright Epic Games, Inc. All Rights Reserved.
+
+#pragma once
+
+#include "CoreMinimal.h"
+#include "GameFramework/Actor.h"
+#include "Components/ActorComponent.h"
+#include "MoverComponent.h"
+#include "GameFramework/NavMovementInterface.h"
+
+#include "NavMoverComponent.generated.h"
+
+/**
+ * NavMoverComponent: Responsible for implementing INavMoveInterface with UMoverComponent so pathfinding and other forms of navigation movement work.
+ * This component also caches the input given to it that is then consumed by the mover system.
+ * Note: This component relies on the parent actor having a MoverComponent as well. By default this component only has a reference to MoverComponent meaning
+ * we use other ways (such as gameplay tags for the active movement mode) to check for state rather than calling specific functions on the active MoverComponent.
+ */
+UCLASS(BlueprintType, meta = (BlueprintSpawnableComponent))
+class MOVER_API UNavMoverComponent : public UActorComponent, public INavMovementInterface
+{
+	GENERATED_BODY()
+
+public:
+	UNavMoverComponent();
+
+	FVector CachedNavMoveInputIntent = FVector::ZeroVector;
+	FVector CachedNavMoveInputVelocity = FVector::ZeroVector;
+
+	// TODO: Not sure if we need these
+	FRotator CachedTurnInput = FRotator::ZeroRotator;
+	FRotator CachedLookInput = FRotator::ZeroRotator;
+
+	/** Properties that define how the component can move. */
+	UPROPERTY(EditAnywhere, Category="Nav Movement", meta = (DisplayName = "Movement Capabilities", Keywords = "Nav Agent"))
+	FNavAgentProperties NavAgentProps;
+
+	/** Expresses runtime state of character's movement. Put all temporal changes to movement properties here */
+	UPROPERTY()
+	FMovementProperties MovementState;
+	
+	/** bool for keeping track of requested nav movement - is set to true when movement is requested, is false once nav movement was consumed */
+	UPROPERTY(BlueprintReadOnly, Category="Nav Movement")
+	bool bRequestedNavMovement = false;
+	
+protected:
+	/** associated properties for nav movement */
+	UPROPERTY(EditAnywhere, BlueprintReadWrite, Category="Nav Movement")
+	FNavMovementProperties NavMovementProperties;
+	
+private:
+	/** object implementing IPathFollowingAgentInterface. Is private to control access to it.
+	 *	@see SetPathFollowingAgent, GetPathFollowingAgent */
+	UPROPERTY()
+	TObjectPtr<UObject> PathFollowingComp;
+
+	/** Associated Movement component that will actually move the actor */ 
+	UPROPERTY()
+	TObjectPtr<UMoverComponent> MoverComponent;
+	
+public:
+	virtual void InitializeComponent() override;
+
+	/** Get the owner of the object consuming nav movement */
+	virtual UObject* GetOwnerAsObject() const override { return GetOwner(); }
+	
+	/** Get the component this movement component is updating */
+	virtual TObjectPtr<UObject> GetUpdatedObject() const override { return MoverComponent->GetUpdatedComponent(); }
+
+	/** Get axis-aligned cylinder around this actor, used for simple collision checks in nav movement */
+	virtual void GetSimpleCollisionCylinder(float& CollisionRadius, float& CollisionHalfHeight) const override;
+
+	/** Returns collision extents vector for this object, based on GetSimpleCollisionCylinder. */
+	virtual FVector GetSimpleCollisionCylinderExtent() const override;
+	
+	/** Get forward vector of the object being driven by nav movement */
+	virtual FVector GetForwardVector() const override;
+	
+	/** Get the current velocity of the movement component */
+	UFUNCTION(BlueprintCallable, Category="AI|Components|NavMovement")
+	virtual FVector GetVelocityForNavMovement() const override;
+
+	/** Get the max speed of the movement component */
+	UFUNCTION(BlueprintCallable, Category="AI|Components|NavMovement")
+	virtual float GetMaxSpeedForNavMovement() const override;
+
+	// Overridden to also call StopActiveMovement().
+	virtual void StopMovementImmediately() override;
+
+	/** Returns location of controlled actor - meaning center of collision bounding box */
+	virtual FVector GetLocation() const override;
+	/** Returns location of controlled actor's "feet" meaning center of bottom of collision bounding box */
+	virtual FVector GetFeetLocation() const override;
+	/** Returns based location of controlled actor */
+	virtual FBasedPosition GetFeetLocationBased() const override;
+
+	/** Set nav agent properties from an object */
+	virtual void UpdateNavAgent(const UObject& ObjectToUpdateFrom) override;
+	
+	/** path following: request new velocity */
+	virtual void RequestDirectMove(const FVector& MoveVelocity, bool bForceMaxSpeed) override;
+
+	/** path following: request new move input (normal vector = full strength) */
+	virtual void RequestPathMove(const FVector& MoveInput) override;
+
+	/** check if current move target can be reached right now if positions are matching
+	 *  (e.g. performing scripted move and can't stop) */
+	virtual bool CanStopPathFollowing() const override;
+
+	/** Get Nav movement props struct this component uses */
+	virtual FNavMovementProperties* GetNavMovementProperties() override { return &NavMovementProperties; }
+	/** Returns the NavMovementProps(const) */
+	virtual const FNavMovementProperties& GetNavMovementProperties() const override{ return NavMovementProperties; } 
+	
+	virtual void SetPathFollowingAgent(IPathFollowingAgentInterface* InPathFollowingAgent) override;
+	virtual IPathFollowingAgentInterface* GetPathFollowingAgent() override;
+	virtual const IPathFollowingAgentInterface* GetPathFollowingAgent() const override;
+
+	/** Returns the NavAgentProps(const) */
+	virtual const FNavAgentProperties& GetNavAgentPropertiesRef() const override;
+	/** Returns the NavAgentProps */
+	virtual FNavAgentProperties& GetNavAgentPropertiesRef() override;
+
+	/** Resets runtime movement state to character's movement capabilities */
+	virtual void ResetMoveState() override;
+
+	/** Returns true if path following can start */
+	virtual bool CanStartPathFollowing() const override;
+
+	/** Returns true if currently crouching */
+	virtual bool IsCrouching() const override;
+	
+	/** Returns true if currently falling (not flying, in a non-fluid volume, and not on the ground) */ 
+	virtual bool IsFalling() const override;
+
+	/** Returns true if currently moving on the ground (e.g. walking or driving) */
+	virtual bool IsMovingOnGround() const override;
+	
+	/** Returns true if currently swimming (moving through a fluid volume) */
+	UFUNCTION(BlueprintCallable, Category="AI|Components|NavMovement")
+	virtual bool IsSwimming() const override;
+
+	/** Returns true if currently flying (moving through a non-fluid volume without resting on the ground) */
+	virtual bool IsFlying() const override;
+};
diff --git a/Engine/Plugins/Experimental/MoverExamples/Source/MoverExamples/Private/MoverExamplesCharacter.cpp b/Engine/Plugins/Experimental/MoverExamples/Source/MoverExamples/Private/MoverExamplesCharacter.cpp
--- a/Engine/Plugins/Experimental/MoverExamples/Source/MoverExamples/Private/MoverExamplesCharacter.cpp
+++ b/Engine/Plugins/Experimental/MoverExamples/Source/MoverExamples/Private/MoverExamplesCharacter.cpp
@@ -15,11 +15,13 @@
 #include "GameFramework/PhysicsVolume.h"
 #include "EnhancedInputComponent.h"
 #include "InputAction.h"
+#include "DefaultMovementSet/NavMoverComponent.h"
 
 static const FName Name_CharacterMotionComponent(TEXT("MoverComponent"));
 
 AMoverExamplesCharacter::AMoverExamplesCharacter(const FObjectInitializer& ObjectInitializer)
 	: Super(ObjectInitializer)
+	, NavMoverComponent(nullptr)
 {
 	CharacterMotionComponent = CreateDefaultSubobject<UCharacterMoverComponent>(Name_CharacterMotionComponent);
 	ensure(CharacterMotionComponent);
@@ -37,6 +39,11 @@ AMoverExamplesCharacter::AMoverExamplesCharacter(const FObjectInitializer& Objec
 	static FName ProduceInputBPFuncName = FName(TEXT("OnProduceInputInBlueprint"));
 	UFunction* ProduceInputFunction = GetClass()->FindFunctionByName(ProduceInputBPFuncName);
 	bHasProduceInputinBpFunc = IsImplementedInBlueprint(ProduceInputFunction);
+
+	if (USceneComponent* UpdatedComponent = CharacterMotionComponent->GetUpdatedComponent())
+	{
+		UpdatedComponent->SetCanEverAffectNavigation(bCanAffectNavigationGeneration);
+	}
 }
 
 
@@ -72,6 +79,8 @@ void AMoverExamplesCharacter::BeginPlay()
 		PC->PlayerCameraManager->ViewPitchMax = 89.0f;
 		PC->PlayerCameraManager->ViewPitchMin = -89.0f;
 	}
+	
+	NavMoverComponent = FindComponentByClass<UNavMoverComponent>();
 }
 
 // Called to bind functionality to input
@@ -93,6 +102,35 @@ void AMoverExamplesCharacter::SetupPlayerInputComponent(UInputComponent* PlayerI
 }
 
 
+FVector AMoverExamplesCharacter::GetNavAgentLocation() const
+{
+	FVector AgentLocation = FNavigationSystem::InvalidLocation;
+	const USceneComponent* UpdatedComponent = CharacterMotionComponent ? CharacterMotionComponent->GetUpdatedComponent() : nullptr;
+	
+	if (NavMoverComponent)
+	{
+		AgentLocation = NavMoverComponent->GetFeetLocation();
+	}
+	
+	if (FNavigationSystem::IsValidLocation(AgentLocation) == false && UpdatedComponent != nullptr)
+	{
+		AgentLocation = UpdatedComponent->GetComponentLocation() - FVector(0,0,UpdatedComponent->Bounds.BoxExtent.Z);
+	}
+
+	return AgentLocation;
+}
+
+void AMoverExamplesCharacter::UpdateNavigationRelevance()
+{
+	if (CharacterMotionComponent)
+	{
+		if (USceneComponent* UpdatedComponent = CharacterMotionComponent->GetUpdatedComponent())
+		{
+			UpdatedComponent->SetCanEverAffectNavigation(bCanAffectNavigationGeneration);
+		}
+	}
+}
+
 void AMoverExamplesCharacter::ProduceInput_Implementation(int32 SimTimeMs, FMoverInputCmdContext& InputCmdResult)
 {
 	OnProduceInput((float)SimTimeMs, InputCmdResult);
@@ -148,6 +186,21 @@ void AMoverExamplesCharacter::OnProduceInput(float DeltaMs, FMoverInputCmdContex
 		CharacterInputs.ControlRotation = PC->GetControlRotation();
 	}
 
+	bool bRequestedNavMovement = false;
+	if (NavMoverComponent)
+	{
+		bRequestedNavMovement = NavMoverComponent->bRequestedNavMovement;
+		if (bRequestedNavMovement)
+		{
+			CachedMoveInputIntent = NavMoverComponent->CachedNavMoveInputIntent;
+			CachedMoveInputVelocity = NavMoverComponent->CachedNavMoveInputVelocity;
+			
+			NavMoverComponent->bRequestedNavMovement = false;
+			NavMoverComponent->CachedNavMoveInputIntent = FVector::ZeroVector;
+			NavMoverComponent->CachedNavMoveInputVelocity = FVector::ZeroVector;
+		}
+	}
+	
 	// Favor velocity input 
 	bool bUsingInputIntentForMove = CachedMoveInputVelocity.IsZero();
 
@@ -173,6 +226,13 @@ void AMoverExamplesCharacter::OnProduceInput(float DeltaMs, FMoverInputCmdContex
 		CharacterInputs.SetMoveInput(EMoveInputType::Velocity, CachedMoveInputVelocity);
 	}
 
+	// Normally cached input is cleared by OnMoveCompleted input event but that won't be called if movement came from nav movement
+	if (bRequestedNavMovement)
+	{
+		CachedMoveInputIntent = FVector::ZeroVector;
+		CachedMoveInputVelocity = FVector::ZeroVector;
+	}
+	
 	static float RotationMagMin(1e-3);
 
 	const bool bHasAffirmativeMoveInput = (CharacterInputs.GetMoveInput().Size() >= RotationMagMin);
@@ -181,7 +241,7 @@ void AMoverExamplesCharacter::OnProduceInput(float DeltaMs, FMoverInputCmdContex
 	CharacterInputs.OrientationIntent = FVector::ZeroVector;
 
 
-	if (bUsingInputIntentForMove && bHasAffirmativeMoveInput)
+	if (bHasAffirmativeMoveInput)
 	{
 		if (bOrientRotationToMovement)
 		{
@@ -195,7 +255,6 @@ void AMoverExamplesCharacter::OnProduceInput(float DeltaMs, FMoverInputCmdContex
 		}
 
 		LastAffirmativeMoveInput = CharacterInputs.GetMoveInput();
-
 	}
 	else if (bMaintainLastInputOrientation)
 	{
diff --git a/Engine/Plugins/Experimental/MoverExamples/Source/MoverExamples/Public/MoverExamplesCharacter.h b/Engine/Plugins/Experimental/MoverExamples/Source/MoverExamples/Public/MoverExamplesCharacter.h
--- a/Engine/Plugins/Experimental/MoverExamples/Source/MoverExamples/Public/MoverExamplesCharacter.h
+++ b/Engine/Plugins/Experimental/MoverExamples/Source/MoverExamples/Public/MoverExamplesCharacter.h
@@ -7,6 +7,7 @@
 #include "GameFramework/Pawn.h"
 #include "MoverExamplesCharacter.generated.h"
 
+class UNavMoverComponent;
 class UInputAction;
 class UCharacterMoverComponent;
 struct FInputActionValue;
@@ -47,6 +48,12 @@ public:
 	UFUNCTION(BlueprintCallable, Category=MoverExamples)
 	virtual void RequestMoveByVelocity(const FVector& DesiredVelocity) { CachedMoveInputVelocity=DesiredVelocity; }
 
+	//~ Begin INavAgentInterface Interface
+	virtual FVector GetNavAgentLocation() const override;
+	//~ End INavAgentInterface Interface
+	
+	virtual void UpdateNavigationRelevance() override;
+	
 protected:
 	// Entry point for input production. Do not override. To extend in derived character types, override OnProduceInput for native types or implement "Produce Input" blueprint event
 	virtual void ProduceInput_Implementation(int32 SimTimeMs, FMoverInputCmdContext& InputCmdResult) override;
@@ -102,6 +109,10 @@ protected:
 	UPROPERTY(Category = Movement, VisibleAnywhere, BlueprintReadOnly, meta = (AllowPrivateAccess = "true"))
 	TObjectPtr<UCharacterMoverComponent> CharacterMotionComponent;
 
+	/** Holds functionality for nav movement data and functions */
+	UPROPERTY(VisibleAnywhere, BlueprintReadOnly, Transient, Category="Nav Movement")
+	TObjectPtr<UNavMoverComponent> NavMoverComponent;
+	
 private:
 	FVector LastAffirmativeMoveInput = FVector::ZeroVector;	// Movement input (intent or velocity) the last time we had one that wasn't zero
 
