Mover 2.0 adding ProcessGeneratedMovement delegate for mode-agnostic modification of proposed movement after it has been generated.

diff --git a/Engine/Plugins/Experimental/Mover/Source/Mover/Private/MovementModeStateMachine.cpp b/Engine/Plugins/Experimental/Mover/Source/Mover/Private/MovementModeStateMachine.cpp
--- a/Engine/Plugins/Experimental/Mover/Source/Mover/Private/MovementModeStateMachine.cpp
+++ b/Engine/Plugins/Experimental/Mover/Source/Mover/Private/MovementModeStateMachine.cpp
@@ -241,6 +241,8 @@ void UMovementModeStateMachine::OnSimulationTick(USceneComponent* UpdatedCompone
 			}
 			CurrentLayeredMoves.ResetResidualVelocity();
 
+			MoverComp->ProcessGeneratedMovement.ExecuteIfBound(SubstepStartData, SubTimeStep, OUT CombinedMove);
+			
 			// Execute the combined proposed move
 			{
 				FSimulationTickParams SimTickParams;
diff --git a/Engine/Plugins/Experimental/Mover/Source/Mover/Private/MoverComponent.cpp b/Engine/Plugins/Experimental/Mover/Source/Mover/Private/MoverComponent.cpp
--- a/Engine/Plugins/Experimental/Mover/Source/Mover/Private/MoverComponent.cpp
+++ b/Engine/Plugins/Experimental/Mover/Source/Mover/Private/MoverComponent.cpp
@@ -234,6 +234,16 @@ void UMoverComponent::TickComponent(float DeltaTime, enum ELevelTick TickType, F
 	Super::TickComponent(DeltaTime, TickType, ThisTickFunction);
 }
 
+void UMoverComponent::BindProcessGeneratedMovement(FMover_ProcessGeneratedMovement ProcessGeneratedMovementEvent)
+{
+	ProcessGeneratedMovement = ProcessGeneratedMovementEvent;
+}
+
+void UMoverComponent::UnbindProcessGeneratedMovement()
+{
+	ProcessGeneratedMovement.Clear();
+}
+
 void UMoverComponent::ProduceInput(const int32 DeltaTimeMS, FMoverInputCmdContext* Cmd)
 {
 	Cmd->InputCollection.Empty();
diff --git a/Engine/Plugins/Experimental/Mover/Source/Mover/Public/MoverComponent.h b/Engine/Plugins/Experimental/Mover/Source/Mover/Public/MoverComponent.h
--- a/Engine/Plugins/Experimental/Mover/Source/Mover/Public/MoverComponent.h
+++ b/Engine/Plugins/Experimental/Mover/Source/Mover/Public/MoverComponent.h
@@ -45,6 +45,9 @@ DECLARE_DYNAMIC_MULTICAST_DELEGATE_TwoParams(FMover_OnPostSimRollback, const FMo
 // Fired after changing movement modes. First param is the name of the previous movement mode. Second is the name of the new movement mode. 
 DECLARE_DYNAMIC_MULTICAST_DELEGATE_TwoParams(FMover_OnMovementModeChanged, const FName&, PreviousMovementModeName, const FName&, NewMovementModeName);
 
+// Fired after proposed movement has been generated (i.e. after movement modes and layered moves have generated movement and mixed together).
+DECLARE_DYNAMIC_DELEGATE_ThreeParams(FMover_ProcessGeneratedMovement, const FMoverTickStartData&, StartState, const FMoverTimeStep&, TimeStep, FProposedMove&, OutProposedMove);
+
 /**
  * 
  */
@@ -86,6 +89,19 @@ public:
 	UPROPERTY(BlueprintAssignable, Category = Mover)
 	FMover_OnMovementModeChanged OnMovementModeChanged;
 
+	/**
+	 * Broadcast after proposed movement has been generated. After movement modes and layered moves have generated movement and mixed together.
+	 * This allows for final modifications to proposed movement before it's executed.
+	 */
+	FMover_ProcessGeneratedMovement ProcessGeneratedMovement;
+	
+	// Binds event for processing movement after it has been generated. Allows for final modifications to proposed movement before it's executed.
+	UFUNCTION(BlueprintCallable, Category = Mover)
+	void BindProcessGeneratedMovement(FMover_ProcessGeneratedMovement ProcessGeneratedMovementEvent);
+	// Clears current bound event for processing movement after it has been generated.
+	UFUNCTION(BlueprintCallable, Category = Mover)
+	void UnbindProcessGeneratedMovement();
+	
 	// Callbacks
 	UFUNCTION()
 	virtual void OnBeginOverlap(UPrimitiveComponent* OverlappedComp, AActor* Other, UPrimitiveComponent* OtherComp, int32 OtherBodyIndex, bool bFromSweep, const FHitResult& SweepResult) { }
