Mover 2.0 fixing issue with layered move mixing and preferred movement modes where Override velocity mix mode would override Preferred mode from OverrideAll

diff --git a/Engine/Plugins/Experimental/Mover/Source/Mover/Private/MoveLibrary/MovementMixer.cpp b/Engine/Plugins/Experimental/Mover/Source/Mover/Private/MoveLibrary/MovementMixer.cpp
--- a/Engine/Plugins/Experimental/Mover/Source/Mover/Private/MoveLibrary/MovementMixer.cpp
+++ b/Engine/Plugins/Experimental/Mover/Source/Mover/Private/MoveLibrary/MovementMixer.cpp
@@ -38,7 +38,7 @@ void UMovementMixer::MixLayeredMove(const FLayeredMoveBase& ActiveMove, const FP
 				UE_LOG(LogMover, Log, TEXT("Multiple LayeredMoves with Override mix mode are active simultaneously. Layered move with the highest priority will take effect."));
 			}
 
-			if (!MoveStep.PreferredMode.IsNone())
+			if (!MoveStep.PreferredMode.IsNone() && OutCumulativeMove.MixMode != EMoveMixMode::OverrideAll)
 			{
 				OutCumulativeMove.PreferredMode = MoveStep.PreferredMode;
 			}
