Mover 2.0 fixing another mixing issue where cumulative move didin't get correct mixing mode from mixing layered moves

diff --git a/Engine/Plugins/Experimental/Mover/Source/Mover/Private/MoveLibrary/MovementMixer.cpp b/Engine/Plugins/Experimental/Mover/Source/Mover/Private/MoveLibrary/MovementMixer.cpp
--- a/Engine/Plugins/Experimental/Mover/Source/Mover/Private/MoveLibrary/MovementMixer.cpp
+++ b/Engine/Plugins/Experimental/Mover/Source/Mover/Private/MoveLibrary/MovementMixer.cpp
@@ -72,6 +72,7 @@ void UMovementMixer::MixLayeredMove(const FLayeredMoveBase& ActiveMove, const FP
 			}
 				
 			OutCumulativeMove = MoveStep;
+			OutCumulativeMove.MixMode = EMoveMixMode::OverrideAll;
 		}
 	}
 	else
