Mover: initial version of standalone backend liaison. Intended to be a very simple backend alternative for non-networked games, avoiding the overhead of NPP.

diff --git a/Engine/Plugins/Experimental/Mover/Source/Mover/Private/Backends/MoverStandaloneLiaison.cpp b/Engine/Plugins/Experimental/Mover/Source/Mover/Private/Backends/MoverStandaloneLiaison.cpp
new file mode 100644
--- /dev/null
+++ b/Engine/Plugins/Experimental/Mover/Source/Mover/Private/Backends/MoverStandaloneLiaison.cpp
@@ -0,0 +1,116 @@
+// Copyright Epic Games, Inc. All Rights Reserved.
+
+#include "Backends/MoverStandaloneLiaison.h"
+#include "Engine/World.h"
+#include "GameFramework/Actor.h"
+#include "MoverComponent.h"
+
+
+UMoverStandaloneLiaisonComponent::UMoverStandaloneLiaisonComponent(const FObjectInitializer& ObjectInitializer)
+	: Super(ObjectInitializer)
+{
+	PrimaryComponentTick.TickGroup = TG_PrePhysics;
+	PrimaryComponentTick.bCanEverTick = true;
+
+	bWantsInitializeComponent = true;
+	bAutoActivate = true;
+	SetIsReplicatedByDefault(false);
+
+	CurrentSimTimeMs = 0.0;
+	CurrentSimFrame = 0;
+
+}
+
+float UMoverStandaloneLiaisonComponent::GetCurrentSimTimeMs()
+{
+	return CurrentSimTimeMs;
+}
+
+int32 UMoverStandaloneLiaisonComponent::GetCurrentSimFrame()
+{
+	return CurrentSimFrame;
+}
+
+
+bool UMoverStandaloneLiaisonComponent::ReadPendingSyncState(OUT FMoverSyncState& OutSyncState)
+{
+	OutSyncState = CachedLastSyncState;
+	return true;
+}
+
+bool UMoverStandaloneLiaisonComponent::WritePendingSyncState(const FMoverSyncState& SyncStateToWrite)
+{
+	CachedLastSyncState = SyncStateToWrite;
+	return true;
+}
+
+void UMoverStandaloneLiaisonComponent::BeginPlay()
+{
+	Super::BeginPlay();
+
+	CurrentSimTimeMs = GetWorld()->GetTimeSeconds() * 1000.0;
+	CurrentSimFrame = 0;
+
+	if (const AActor* OwnerActor = GetOwner())
+	{
+		ensureMsgf(OwnerActor->GetNetMode() == NM_Standalone, TEXT("UMoverStandaloneLiaisonComponent is only valid for use in Standalone projects. Movement will not work properly in networked play."));
+
+		ensureMsgf(!OwnerActor->IsReplicatingMovement(),
+			TEXT("MoverComponent owning actor %s has the ReplicateMovement property enabled. This will conflict with Network Prediction and cause poor quality movement. Please disable it."),
+			*GetNameSafe(GetOwner()));
+
+		if (UMoverComponent* FoundMoverComp = OwnerActor->FindComponentByClass<UMoverComponent>())
+		{
+			MoverComp = FoundMoverComp;
+
+			MoverComp->InitMoverSimulation();
+
+			MoverComp->InitializeSimulationState(OUT &CachedLastSyncState, OUT &CachedLastAuxState);
+		}
+		else
+		{
+			ensureMsgf(MoverComp, TEXT("Owning actor %s does not have a MoverComponent."), *GetNameSafe(GetOwner()));
+		}
+	}	
+}
+
+
+void UMoverStandaloneLiaisonComponent::TickComponent(float DeltaTime, enum ELevelTick TickType, FActorComponentTickFunction* ThisTickFunction)
+{
+	Super::TickComponent(DeltaTime, TickType, ThisTickFunction);
+
+	if (MoverComp)
+	{
+		// TODO Here is where we might accumulate time and perform fixed tick updates
+
+		const int32 DeltaTimeMs = DeltaTime * 1000.f;
+
+		CurrentSimTimeMs = GetWorld()->GetTimeSeconds() * 1000.0;
+		CurrentSimFrame = GFrameCounter;
+
+		FMoverInputCmdContext InputCmd;
+		MoverComp->ProduceInput(DeltaTimeMs, OUT &InputCmd);
+
+		FMoverTimeStep TimeStep;
+		TimeStep.ServerFrame = CurrentSimFrame;
+		TimeStep.BaseSimTimeMs = CurrentSimTimeMs;
+		TimeStep.StepMs = FMath::Floor(DeltaTime * 1000.f);
+
+		FMoverTickStartData StartData;
+		FMoverTickEndData EndData;
+
+		StartData.InputCmd = InputCmd;
+		StartData.SyncState = CachedLastSyncState;
+		StartData.AuxState = CachedLastAuxState;
+
+		MoverComp->SimulationTick(TimeStep, StartData, OUT EndData);
+
+		CachedLastSyncState = EndData.SyncState;
+		CachedLastAuxState = EndData.AuxState;
+
+		MoverComp->FinalizeFrame(&EndData.SyncState, &EndData.AuxState);
+	}
+
+}
+
+
diff --git a/Engine/Plugins/Experimental/Mover/Source/Mover/Public/Backends/MoverStandaloneLiaison.h b/Engine/Plugins/Experimental/Mover/Source/Mover/Public/Backends/MoverStandaloneLiaison.h
new file mode 100644
--- /dev/null
+++ b/Engine/Plugins/Experimental/Mover/Source/Mover/Public/Backends/MoverStandaloneLiaison.h
@@ -0,0 +1,48 @@
+// Copyright Epic Games, Inc. All Rights Reserved.
+
+#pragma once
+
+#include "CoreMinimal.h"
+#include "Components/ActorComponent.h"
+#include "Backends/MoverBackendLiaison.h"
+#include "MoverTypes.h"
+
+#include "MoverStandaloneLiaison.generated.h"
+
+class UMoverComponent;
+
+/**
+ * MoverStandaloneLiaison: this component acts as a backend driver for an actor's Mover component, for use in Standalone (non-networked) games.
+ * This class is set on a Mover component as the "back end".
+ * TODO: Support options for fixed ticking rates and state smoothing
+ */
+UCLASS()
+class MOVER_API UMoverStandaloneLiaisonComponent : public UActorComponent, public IMoverBackendLiaisonInterface
+{
+	GENERATED_BODY()
+
+public:
+	UMoverStandaloneLiaisonComponent(const FObjectInitializer& ObjectInitializer = FObjectInitializer::Get());
+
+	// IMoverBackendLiaisonInterface
+	virtual float GetCurrentSimTimeMs() override;
+	virtual int32 GetCurrentSimFrame() override;
+	virtual bool ReadPendingSyncState(OUT FMoverSyncState& OutSyncState) override;
+	virtual bool WritePendingSyncState(const FMoverSyncState& SyncStateToWrite) override;
+	// End IMoverBackendLiaisonInterface
+
+	// Begin UActorComponent interface
+	virtual void BeginPlay() override;
+	virtual void TickComponent(float DeltaTime, enum ELevelTick TickType, FActorComponentTickFunction* ThisTickFunction) override;
+	// End UActorComponent interface
+
+
+protected:
+	TObjectPtr<UMoverComponent> MoverComp;	// the component that we're in charge of driving
+
+	double CurrentSimTimeMs;
+	int32 CurrentSimFrame;
+
+	FMoverSyncState CachedLastSyncState;
+	FMoverAuxStateContext CachedLastAuxState;
+};
\ No newline at end of file
