Mover: changes montage-based anim root motion move to end itself if the mesh is no longer playing the montage, while still maintaining residual velocity. This matches CMC's behavior, and allows smooth transition back to controlled movement without loss of momentum.

diff --git a/Engine/Plugins/Experimental/Mover/Source/Mover/Private/DefaultMovementSet/LayeredMoves/AnimRootMotionLayeredMove.cpp b/Engine/Plugins/Experimental/Mover/Source/Mover/Private/DefaultMovementSet/LayeredMoves/AnimRootMotionLayeredMove.cpp
--- a/Engine/Plugins/Experimental/Mover/Source/Mover/Private/DefaultMovementSet/LayeredMoves/AnimRootMotionLayeredMove.cpp
+++ b/Engine/Plugins/Experimental/Mover/Source/Mover/Private/DefaultMovementSet/LayeredMoves/AnimRootMotionLayeredMove.cpp
@@ -30,6 +30,26 @@ FLayeredMove_AnimRootMotion::FLayeredMove_AnimRootMotion()
 
 bool FLayeredMove_AnimRootMotion::GenerateMove(const FMoverTickStartData& SimState, const FMoverTimeStep& TimeStep, const UMoverComponent* MoverComp, UMoverBlackboard* SimBlackboard, FProposedMove& OutProposedMove)
 {
+	// Stop this move if the montage is no longer playing on the mesh
+	if (!TimeStep.bIsResimulating)
+	{
+		bool bIsMontageStillPlaying = false;
+
+		if (const USkeletalMeshComponent* MeshComp = Cast<USkeletalMeshComponent>(MoverComp->GetPrimaryVisualComponent()))
+		{
+			if (const UAnimInstance* MeshAnimInstance = MeshComp->GetAnimInstance())
+			{
+				bIsMontageStillPlaying = MeshAnimInstance->Montage_IsPlaying(Montage);
+			}
+		}
+
+		if (!bIsMontageStillPlaying)
+		{
+			DurationMs = 0.f;
+			return false;
+		}
+	}
+
 	const float DeltaSeconds = TimeStep.StepMs / 1000.f;
 
 	const AActor* MoverActor = MoverComp->GetOwner();
diff --git a/Engine/Plugins/Experimental/Mover/Source/Mover/Public/DefaultMovementSet/LayeredMoves/AnimRootMotionLayeredMove.h b/Engine/Plugins/Experimental/Mover/Source/Mover/Public/DefaultMovementSet/LayeredMoves/AnimRootMotionLayeredMove.h
--- a/Engine/Plugins/Experimental/Mover/Source/Mover/Public/DefaultMovementSet/LayeredMoves/AnimRootMotionLayeredMove.h
+++ b/Engine/Plugins/Experimental/Mover/Source/Mover/Public/DefaultMovementSet/LayeredMoves/AnimRootMotionLayeredMove.h
@@ -9,7 +9,10 @@
 class UAnimMontage;
 
 
-/** Anim Root Motion Move: handles root motion from a montage played on the associated mesh */
+/** Anim Root Motion Move: handles root motion from a montage played on the primary visual component (skeletal mesh). 
+ * In this method, root motion is extracted independently from anim playback. The move will end itself if the animation
+ * is interrupted on the mesh.
+ */
 USTRUCT(BlueprintType)
 struct MOVER_API FLayeredMove_AnimRootMotion : public FLayeredMoveBase
 {
