Mover: tweaking BP display name for QueueNextMode to make clear that it's for movement modes

diff --git a/Engine/Plugins/Experimental/Mover/Source/Mover/Public/MoverComponent.h b/Engine/Plugins/Experimental/Mover/Source/Mover/Public/MoverComponent.h
--- a/Engine/Plugins/Experimental/Mover/Source/Mover/Public/MoverComponent.h
+++ b/Engine/Plugins/Experimental/Mover/Source/Mover/Public/MoverComponent.h
@@ -175,7 +175,7 @@ public:
 	void QueueInstantMovementEffect(TSharedPtr<FInstantMovementEffect> Move);
 	
 	// Queue a movement mode change to occur during the next simulation frame. If bShouldReenter is true, then a mode change will occur even if already in that mode.
-	UFUNCTION(BlueprintCallable, Category = Mover)
+	UFUNCTION(BlueprintCallable, Category = Mover, DisplayName="Queue Next Movement Mode")
 	void QueueNextMode(FName DesiredModeName, bool bShouldReenter=false);
 
 	// Add a movement mode to available movement modes. Returns true if the movement mode was added successfully. Returns the mode that was made.
