Mover 2.0: Adding Movement Modifier system - a system that is used to apply changes that indirectly influence the movement simulation, without proposing or executing any movement, but still in sync with the sim.

diff --git a/Engine/Plugins/Experimental/Mover/Source/Mover/Private/Backends/MoverNetworkPhysicsLiaison.cpp b/Engine/Plugins/Experimental/Mover/Source/Mover/Private/Backends/MoverNetworkPhysicsLiaison.cpp
--- a/Engine/Plugins/Experimental/Mover/Source/Mover/Private/Backends/MoverNetworkPhysicsLiaison.cpp
+++ b/Engine/Plugins/Experimental/Mover/Source/Mover/Private/Backends/MoverNetworkPhysicsLiaison.cpp
@@ -833,7 +833,9 @@ void UMoverNetworkPhysicsLiaisonComponent::ProcessInputs_Internal(int32 PhysicsS
 				// Rollback mover state if on the first resimulation frame
 				if (bLocalPlayer && bIsSolverResim && bIsFirstResimFrame)
 				{
+					FMoverAuxStateContext UnusedInvalidAuxState;
 					FMoverAuxStateContext UnusedAuxState;
+					MoverComp->OnSimulationPreRollback(&NetSyncState, &Input.SyncState, &UnusedInvalidAuxState, &UnusedAuxState);
 					MoverComp->OnSimulationRollback(&Input.SyncState, &UnusedAuxState);
 				}
 			}
diff --git a/Engine/Plugins/Experimental/Mover/Source/Mover/Private/Backends/MoverNetworkPredictionLiaison.cpp b/Engine/Plugins/Experimental/Mover/Source/Mover/Private/Backends/MoverNetworkPredictionLiaison.cpp
--- a/Engine/Plugins/Experimental/Mover/Source/Mover/Private/Backends/MoverNetworkPredictionLiaison.cpp
+++ b/Engine/Plugins/Experimental/Mover/Source/Mover/Private/Backends/MoverNetworkPredictionLiaison.cpp
@@ -54,6 +54,14 @@ void UMoverNetworkPredictionLiaisonComponent::RestoreFrame(const FMoverSyncState
 void UMoverNetworkPredictionLiaisonComponent::FinalizeFrame(const FMoverSyncState* SyncState, const FMoverAuxStateContext* AuxState)
 {
 	check(MoverComp);
+
+	const FNetworkPredictionSettings NetworkPredictionSettings = UNetworkPredictionWorldManager::ActiveInstance->GetSettings();
+	if (MoverComp->GetOwnerRole() == ROLE_SimulatedProxy && NetworkPredictionSettings.SimulatedProxyNetworkLOD == ENetworkLOD::Interpolated)
+	{
+		FMoverInputCmdContext InputCmd;
+		MoverComp->TickInterpolatedSimProxy(MoverComp->GetLastTimeStep(), InputCmd, MoverComp, MoverComp->GetSyncState(), *SyncState, *AuxState);
+	}
+	
 	MoverComp->FinalizeFrame(SyncState, AuxState);
 }
 
diff --git a/Engine/Plugins/Experimental/Mover/Source/Mover/Private/Debug/GameplayDebuggerCategory_Mover.cpp b/Engine/Plugins/Experimental/Mover/Source/Mover/Private/Debug/GameplayDebuggerCategory_Mover.cpp
--- a/Engine/Plugins/Experimental/Mover/Source/Mover/Private/Debug/GameplayDebuggerCategory_Mover.cpp
+++ b/Engine/Plugins/Experimental/Mover/Source/Mover/Private/Debug/GameplayDebuggerCategory_Mover.cpp
@@ -97,8 +97,8 @@ void FGameplayDebuggerCategory_Mover::CollectData(APlayerController* OwnerPC, AA
 	DataPack.Velocity = FVector::ZeroVector;
 	DataPack.MoveIntent = FVector::ZeroVector;
 	DataPack.ActiveLayeredMoves.Empty();
+	DataPack.ActiveModifiers.Empty();
 	DataPack.ModeMap.Empty();
-	DataPack.ActiveLayeredMoves.Empty();
 
 	if (MyMoverComponent)
 	{
@@ -136,6 +136,11 @@ void FGameplayDebuggerCategory_Mover::CollectData(APlayerController* OwnerPC, AA
 			{
 				DataPack.ActiveLayeredMoves.Add(*it->Get()->ToSimpleString());
 			}
+
+			for (auto it = SyncState.MovementModifiers.GetActiveModifiersIterator(); it; ++it)
+			{
+				DataPack.ActiveModifiers.Add(*it->Get()->ToSimpleString());
+			}
 		}
 	}
 }
@@ -153,13 +158,14 @@ void FGameplayDebuggerCategory_Mover::DrawData(APlayerController* OwnerPC, FGame
 		DrawInWorldInfo(*FocusedActor, CanvasContext);
 	}
 	
-	CanvasContext.Printf(TEXT("{yellow}%s\n{grey}Local Role: {white}%s\n{grey}Mode: {white}%s\n{grey}Velocity: {white}%s\n{grey}Speed: {white}%.2f\n{yellow}Active Moves: {white}\n%s\n{yellow}Mode Map: \n{white}%s\n{yellow}Active Transitions: {white}\n%s"),
+	CanvasContext.Printf(TEXT("{yellow}%s\n{grey}Local Role: {white}%s\n{grey}Mode: {white}%s\n{grey}Velocity: {white}%s\n{grey}Speed: {white}%.2f\n{yellow}Active Moves: {white}\n%s\n{yellow}Active Modifiers: {white}\n%s\n{yellow}Mode Map: \n{white}%s\n{yellow}Active Transitions: {white}\n%s"),
 		*DataPack.PawnName,
 		*DataPack.LocalRole,
 		*DataPack.MovementModeName,
 		*DataPack.Velocity.ToString(),
 		 DataPack.Velocity.Length(),
 		*FString::JoinBy(DataPack.ActiveLayeredMoves, TEXT("\n"), [](FString MoveAsString) { return MoveAsString; }),
+		*FString::JoinBy(DataPack.ActiveModifiers, TEXT("\n"), [](FString ModifierAsString) { return ModifierAsString; }),
 		*FString::JoinBy(DataPack.ModeMap, TEXT("\n"), [](FString ModeMappingAsString) { return ModeMappingAsString; }),
 		*FString::JoinBy(DataPack.ActiveTransitions, TEXT("\n"), [](FString TransitionAsString) { return TransitionAsString; })
 		);
@@ -269,6 +275,7 @@ void FGameplayDebuggerCategory_Mover::FRepData::Serialize(FArchive& Ar)
 	Ar << Velocity;
 	Ar << MoveIntent;
 	Ar << ActiveLayeredMoves;
+	Ar << ActiveModifiers;
 	Ar << ModeMap;
 	Ar << ActiveTransitions;
 }
diff --git a/Engine/Plugins/Experimental/Mover/Source/Mover/Private/Debug/GameplayDebuggerCategory_Mover.h b/Engine/Plugins/Experimental/Mover/Source/Mover/Private/Debug/GameplayDebuggerCategory_Mover.h
--- a/Engine/Plugins/Experimental/Mover/Source/Mover/Private/Debug/GameplayDebuggerCategory_Mover.h
+++ b/Engine/Plugins/Experimental/Mover/Source/Mover/Private/Debug/GameplayDebuggerCategory_Mover.h
@@ -34,6 +34,7 @@ protected:
 		TArray<FString> ActiveLayeredMoves;
 		TArray<FString> ModeMap;
 		TArray<FString> ActiveTransitions;
+		TArray<FString> ActiveModifiers;
 
 		FRepData() {}
 
diff --git a/Engine/Plugins/Experimental/Mover/Source/Mover/Private/MovementModeStateMachine.cpp b/Engine/Plugins/Experimental/Mover/Source/Mover/Private/MovementModeStateMachine.cpp
--- a/Engine/Plugins/Experimental/Mover/Source/Mover/Private/MovementModeStateMachine.cpp
+++ b/Engine/Plugins/Experimental/Mover/Source/Mover/Private/MovementModeStateMachine.cpp
@@ -162,6 +162,9 @@ void UMovementModeStateMachine::OnSimulationTick(USceneComponent* UpdatedCompone
 		FlushQueuedMovesToGroup(SubstepStartData.SyncState.LayeredMoves);
 		OutputState.SyncState.LayeredMoves = SubstepStartData.SyncState.LayeredMoves;
 
+		FlushQueuedModifiersToGroup(SubstepStartData.SyncState.MovementModifiers);
+		OutputState.SyncState.MovementModifiers = SubstepStartData.SyncState.MovementModifiers;
+		
 		FApplyMovementEffectParams EffectParams;
 		EffectParams.MoverComp = MoverComp;
 		EffectParams.StartState = &SubstepStartData;
@@ -188,6 +191,15 @@ void UMovementModeStateMachine::OnSimulationTick(USceneComponent* UpdatedCompone
 			}
 		}
 
+		FMovementModifierGroup& CurrentModifiers = OutputState.SyncState.MovementModifiers;
+		FlushModifierCancellationsToGroup(CurrentModifiers);
+		TArray<TSharedPtr<FMovementModifierBase>> ActiveModifiers = CurrentModifiers.GenerateActiveModifiers(MoverComp, TimeStep, SubstepStartData.SyncState, SubstepStartData.AuxState);
+
+		for (TSharedPtr<FMovementModifierBase> Modifier : ActiveModifiers)
+		{
+			Modifier->OnPreMovement(MoverComp, TimeStep);
+		}
+		
 		FLayeredMoveGroup& CurrentLayeredMoves = OutputState.SyncState.LayeredMoves;
 
 		// Gather any layered move contributions
@@ -323,6 +335,11 @@ void UMovementModeStateMachine::OnSimulationTick(USceneComponent* UpdatedCompone
 		AdvanceToNextMode();
 		OutputState.SyncState.MovementMode = CurrentModeName;
 
+		for (TSharedPtr<FMovementModifierBase> Modifier : ActiveModifiers)
+		{
+			Modifier->OnPostMovement(MoverComp, TimeStep, OutputState.SyncState, OutputState.AuxState);
+		}
+		
 		const float RemainingMs = FMath::Clamp(OutputState.MovementEndState.RemainingMs, 0.0f, SubTimeStep.StepMs);
 		SubTimeStep.BaseSimTimeMs += (SubTimeStep.StepMs - RemainingMs);
 		SubTimeStep.StepMs = EndingSimTimeMs - SubTimeStep.BaseSimTimeMs;
@@ -342,6 +359,11 @@ void UMovementModeStateMachine::OnSimulationTick(USceneComponent* UpdatedCompone
 	ApplyInstantEffects(EffectParams, OutputState.SyncState);
 }
 
+void UMovementModeStateMachine::OnSimulationPreRollback(const FMoverSyncState* InvalidSyncState, const FMoverSyncState* SyncState, const FMoverAuxStateContext* InvalidAuxState, const FMoverAuxStateContext* AuxState)
+{
+	RollbackModifiers(InvalidSyncState, SyncState, InvalidAuxState, AuxState);
+}
+
 void UMovementModeStateMachine::OnSimulationRollback(const FMoverSyncState* SyncState, const FMoverAuxStateContext* AuxState)
 {
 	ClearQueuedMode();
@@ -380,6 +402,32 @@ void UMovementModeStateMachine::QueueInstantMovementEffect(TSharedPtr<FInstantMo
 	QueuedInstantEffects.Add(Effect);
 }
 
+FMovementModifierHandle UMovementModeStateMachine::QueueMovementModifier(TSharedPtr<FMovementModifierBase> Modifier)
+{
+	if (ensure(Modifier.IsValid()))
+	{
+		QueuedMovementModifiers.Add(Modifier);
+		Modifier->GenerateHandle();
+
+		return Modifier->GetHandle();
+	}
+
+	return 0;
+}
+
+void UMovementModeStateMachine::CancelModifierFromHandle(FMovementModifierHandle ModifierHandle)
+{
+	for (TSharedPtr<FMovementModifierBase> Modifier : QueuedMovementModifiers)
+	{
+		if (Modifier->GetHandle() == ModifierHandle)
+		{
+			QueuedMovementModifiers.Remove(Modifier);
+		}
+	}
+
+	ModifiersToCancel.Add(ModifierHandle);
+}
+
 
 void UMovementModeStateMachine::ConstructDefaultModes()
 {
@@ -434,6 +482,75 @@ void UMovementModeStateMachine::FlushQueuedMovesToGroup(FLayeredMoveGroup& Group
 	}
 }
 
+void UMovementModeStateMachine::FlushQueuedModifiersToGroup(FMovementModifierGroup& ModifierGroup)
+{
+	if (!QueuedMovementModifiers.IsEmpty())
+	{
+		for (TSharedPtr<FMovementModifierBase>& QueuedModifier : QueuedMovementModifiers)
+		{
+			ModifierGroup.QueueMovementModifier(QueuedModifier);
+		}
+		
+		QueuedMovementModifiers.Empty();
+	}
+}
+
+void UMovementModeStateMachine::FlushModifierCancellationsToGroup(FMovementModifierGroup& ActiveModifierGroup)
+{
+	for (FMovementModifierHandle HandleToCancel : ModifiersToCancel)
+	{
+		ActiveModifierGroup.CancelModifierFromHandle(HandleToCancel);
+	}
+	
+	ModifiersToCancel.Empty();
+}
+
+void UMovementModeStateMachine::RollbackModifiers(const FMoverSyncState* InvalidSyncState, const FMoverSyncState* SyncState, const FMoverAuxStateContext* InvalidAuxState, const FMoverAuxStateContext* AuxState)
+{
+	QueuedMovementModifiers.Empty();
+	
+	if (UMoverComponent* MoverComp = Cast<UMoverComponent>(GetOuter()))
+	{
+		for (auto ModifierFromRollbackIt = SyncState->MovementModifiers.GetActiveModifiersIterator(); ModifierFromRollbackIt; ++ModifierFromRollbackIt)
+		{
+			bool bContainsModifier = false;
+			for (auto ModifierFromCacheIt = InvalidSyncState->MovementModifiers.GetActiveModifiersIterator(); ModifierFromCacheIt; ++ModifierFromCacheIt)
+			{
+				if (ModifierFromRollbackIt->Get()->Matches(ModifierFromCacheIt->Get()))
+				{
+					bContainsModifier = true;
+					break;
+				}
+			}
+		
+			if (!bContainsModifier)
+			{
+				UE_LOG(LogMover, Log, TEXT("Modifier(%s) was started on %s after a rollback."), *ModifierFromRollbackIt->Get()->ToSimpleString(), *GetNameSafe(MoverComp->GetOwner()));
+				ModifierFromRollbackIt->Get()->OnStart(MoverComp, MoverComp->GetLastTimeStep(), *SyncState, *AuxState);
+			}
+		}
+
+		for (auto ModifierFromCacheIt = InvalidSyncState->MovementModifiers.GetActiveModifiersIterator(); ModifierFromCacheIt; ++ModifierFromCacheIt)
+		{
+			bool bContainsModifier = false;
+			for (auto ModifierFromRollbackIt = SyncState->MovementModifiers.GetActiveModifiersIterator(); ModifierFromRollbackIt; ++ModifierFromRollbackIt)
+			{
+				if (ModifierFromRollbackIt->Get()->Matches(ModifierFromCacheIt->Get()))
+				{
+					bContainsModifier = true;
+					break;
+				}
+			}
+	
+			if (!bContainsModifier)
+			{
+				UE_LOG(LogMover, Log, TEXT("Modifier(%s) was ended on %s after a rollback."), *ModifierFromCacheIt->Get()->ToSimpleString(), *GetNameSafe(MoverComp->GetOwner()));
+				ModifierFromCacheIt->Get()->OnEnd(MoverComp, MoverComp->GetLastTimeStep(), *SyncState, *AuxState);
+			}
+		}
+	}
+}
+
 bool UMovementModeStateMachine::ApplyInstantEffects(FApplyMovementEffectParams& ApplyEffectParams, FMoverSyncState& OutputState)
 {
 	bool bInstantMovementEffectApplied = false;
diff --git a/Engine/Plugins/Experimental/Mover/Source/Mover/Private/MovementModeStateMachine.h b/Engine/Plugins/Experimental/Mover/Source/Mover/Private/MovementModeStateMachine.h
--- a/Engine/Plugins/Experimental/Mover/Source/Mover/Private/MovementModeStateMachine.h
+++ b/Engine/Plugins/Experimental/Mover/Source/Mover/Private/MovementModeStateMachine.h
@@ -51,6 +51,7 @@ public:
 	void ClearQueuedMode();
 
 	void OnSimulationTick(USceneComponent* UpdatedComponent, UPrimitiveComponent* UpdatedPrimitive, UMoverBlackboard* SimBlackboard, const FMoverTickStartData& StartState, const FMoverTimeStep& TimeStep, FMoverTickEndData& OutputState);
+ 	void OnSimulationPreRollback(const FMoverSyncState* InvalidSyncState, const FMoverSyncState* SyncState, const FMoverAuxStateContext* InvalidAuxState, const FMoverAuxStateContext* AuxState);
 	void OnSimulationRollback(const FMoverSyncState* SyncState, const FMoverAuxStateContext* AuxState);
 
 	FName GetCurrentModeName() const { return CurrentModeName; }
@@ -63,6 +64,10 @@ public:
 	
  	void QueueInstantMovementEffect(TSharedPtr<FInstantMovementEffect> Effect);
 
+	FMovementModifierHandle QueueMovementModifier(TSharedPtr<FMovementModifierBase> Modifier);
+
+ 	void CancelModifierFromHandle(FMovementModifierHandle ModifierHandle);
+ 	
 protected:
 
 	virtual void PostInitProperties() override;
@@ -81,11 +86,20 @@ protected:
 
  	/** Effects that are queued to be applied to the simulation at the start of the next sim subtick or at the end of this tick */
  	TArray<TSharedPtr<FInstantMovementEffect>> QueuedInstantEffects;
+
+ 	/** Modifiers that are queued to be added to the simulation at the start of the next sim subtick */
+ 	TArray<TSharedPtr<FMovementModifierBase>> QueuedMovementModifiers;
+
+ 	/** Modifiers that are to be canceled at the start of the next sim subtick */
+ 	TArray<FMovementModifierHandle> ModifiersToCancel;
  	
 private:
 	void ConstructDefaultModes();
 	void AdvanceToNextMode();
 	void FlushQueuedMovesToGroup(FLayeredMoveGroup& Group);
+ 	void FlushQueuedModifiersToGroup(FMovementModifierGroup& ModifierGroup);
+ 	void FlushModifierCancellationsToGroup(FMovementModifierGroup& ActiveModifierGroup);
+ 	void RollbackModifiers(const FMoverSyncState* InvalidSyncState, const FMoverSyncState* SyncState, const FMoverAuxStateContext* InvalidAuxState, const FMoverAuxStateContext* AuxState);
  	bool ApplyInstantEffects(FApplyMovementEffectParams& ApplyEffectParams, FMoverSyncState& OutputState);
 	AActor* GetOwnerActor() const;
 };
diff --git a/Engine/Plugins/Experimental/Mover/Source/Mover/Private/MovementModifier.cpp b/Engine/Plugins/Experimental/Mover/Source/Mover/Private/MovementModifier.cpp
new file mode 100644
--- /dev/null
+++ b/Engine/Plugins/Experimental/Mover/Source/Mover/Private/MovementModifier.cpp
@@ -0,0 +1,403 @@
+﻿// Copyright Epic Games, Inc. All Rights Reserved.
+
+#include "MovementModifier.h"
+#include "MoverComponent.h"
+
+const float MovementModifier_InvalidTime = -UE_BIG_NUMBER;
+
+void FMovementModifierHandle::GenerateHandle()
+{
+	static MODIFIER_HANDLE_TYPE LocalModifierIDGenerator = 0;
+	MODIFIER_HANDLE_TYPE LocalID = ++LocalModifierIDGenerator;
+	
+	// TODO: might want to change this from a magic number 0
+	if (LocalID == 0)
+	{
+		LocalID = ++LocalModifierIDGenerator;
+	}
+
+	Handle = LocalID;
+}
+
+FMovementModifierBase::FMovementModifierBase()
+	: DurationMs(-1)
+	, StartSimTimeMs(MovementModifier_InvalidTime)
+{
+}
+
+void FMovementModifierBase::StartModifier(UMoverComponent* MoverComp, const FMoverTimeStep& TimeStep, const FMoverSyncState& SyncState, const FMoverAuxStateContext& AuxState)
+{
+	StartSimTimeMs = TimeStep.BaseSimTimeMs;
+	OnStart(MoverComp, TimeStep, SyncState, AuxState);
+}
+
+void FMovementModifierBase::EndModifier(UMoverComponent* MoverComp, const FMoverTimeStep& TimeStep, const FMoverSyncState& SyncState, const FMoverAuxStateContext& AuxState)
+{
+	OnEnd(MoverComp, TimeStep, SyncState, AuxState);
+}
+
+bool FMovementModifierBase::IsFinished(float CurrentSimTimeMs) const
+{
+	const bool bHasStarted = (StartSimTimeMs >= 0.f);
+	const bool bTimeExpired = bHasStarted && (DurationMs > 0.f) && (StartSimTimeMs + DurationMs <= CurrentSimTimeMs);
+	const bool bDidTickOnceAndExpire = bHasStarted && (DurationMs == 0.f);
+
+	return bTimeExpired || bDidTickOnceAndExpire;
+}
+
+FMovementModifierBase* FMovementModifierBase::Clone() const
+{
+	// If child classes don't override this, saved modifiers will not work
+	checkf(false, TEXT("FMovementModifierBase::Clone() being called erroneously. This should always be overridden in child classes!"));
+	return nullptr;
+}
+
+void FMovementModifierBase::NetSerialize(FArchive& Ar)
+{
+	Ar << DurationMs;
+	Ar << StartSimTimeMs;
+}
+
+UScriptStruct* FMovementModifierBase::GetScriptStruct() const
+{
+	return FMovementModifierBase::StaticStruct();
+}
+
+FString FMovementModifierBase::ToSimpleString() const
+{
+	return GetScriptStruct()->GetName();
+}
+
+bool FMovementModifierBase::Matches(const FMovementModifierBase* Other) const
+{
+	// TODO: Consider checking other factors other than just type
+	return Other != nullptr && GetScriptStruct() == Other->GetScriptStruct();
+}
+
+FMovementModifierHandle FMovementModifierBase::GetHandle() const
+{
+	return LocalModifierHandle;
+}
+
+void FMovementModifierBase::GenerateHandle()
+{
+	LocalModifierHandle.GenerateHandle();
+}
+
+void FMovementModifierBase::SetHandleFromExistingModifier(const FMovementModifierHandle& OtherModifierHandle)
+{
+	if (OtherModifierHandle.IsValid() && !LocalModifierHandle.IsValid())
+	{
+		LocalModifierHandle = OtherModifierHandle;
+	}
+}
+
+void FMovementModifierGroup::NetSerialize(FArchive& Ar, uint8 MaxNumModifiersToSerialize)
+{
+	// TODO: Warn if some sources will be dropped
+	const uint8 NumActiveMovesToSerialize = FMath::Min<int32>(ActiveModifiers.Num(), MaxNumModifiersToSerialize);
+	const uint8 NumQueuedMovesToSerialize = NumActiveMovesToSerialize < MaxNumModifiersToSerialize ? MaxNumModifiersToSerialize - NumActiveMovesToSerialize : 0;
+	NetSerializeMovementModifierArray(Ar, ActiveModifiers, NumActiveMovesToSerialize);
+	NetSerializeMovementModifierArray(Ar, QueuedModifiers, NumQueuedMovesToSerialize);
+}
+
+void FMovementModifierGroup::QueueMovementModifier(TSharedPtr<FMovementModifierBase> Modifier)
+{
+	if (ensure(Modifier.IsValid()))
+	{
+		QueuedModifiers.Add(Modifier);
+		UE_LOG(LogMover, VeryVerbose, TEXT("Queued Movement Modifier (%s)"), *Modifier->ToSimpleString());
+	}
+}
+
+void FMovementModifierGroup::CancelModifierFromHandle(const FMovementModifierHandle& HandleToCancel)
+{
+	for (TSharedPtr<FMovementModifierBase> ActiveModifier : ActiveModifiers)
+	{ 
+		if (HandleToCancel == ActiveModifier->GetHandle())
+		{
+			ActiveModifier->DurationMs = 0;
+		}
+	}
+
+	for (TSharedPtr<FMovementModifierBase> QueuedModifier : QueuedModifiers)
+	{
+		if (HandleToCancel == QueuedModifier->GetHandle())
+		{
+			QueuedModifier->DurationMs = 0;
+		}
+	}
+}
+
+TArray<TSharedPtr<FMovementModifierBase>> FMovementModifierGroup::GenerateActiveModifiers(UMoverComponent* MoverComp, const FMoverTimeStep& TimeStep, const FMoverSyncState& SyncState, const FMoverAuxStateContext& AuxState)
+{
+	FlushModifierArrays(MoverComp, TimeStep, SyncState, AuxState);
+	return ActiveModifiers;
+}
+
+FMovementModifierGroup& FMovementModifierGroup::operator=(const FMovementModifierGroup& Other)
+{
+	// Perform deep copy of this Group
+	if (this != &Other)
+	{
+		// Deep copy active moves
+		ActiveModifiers.Empty(Other.ActiveModifiers.Num());
+		for (int i = 0; i < Other.ActiveModifiers.Num(); ++i)
+		{
+			if (Other.ActiveModifiers[i].IsValid())
+			{
+				FMovementModifierBase* CopyOfSourcePtr = Other.ActiveModifiers[i]->Clone();
+				ActiveModifiers.Add(TSharedPtr<FMovementModifierBase>(CopyOfSourcePtr));
+			}
+			else
+			{
+				UE_LOG(LogMover, Warning, TEXT("FMovementModifierGroup::operator= trying to copy invalid Other Modifier"));
+			}
+		}
+
+		// Deep copy queued moves
+		QueuedModifiers.Empty(Other.QueuedModifiers.Num());
+		for (int i = 0; i < Other.QueuedModifiers.Num(); ++i)
+		{
+			if (Other.QueuedModifiers[i].IsValid())
+			{
+				FMovementModifierBase* CopyOfSourcePtr = Other.QueuedModifiers[i]->Clone();
+				QueuedModifiers.Add(TSharedPtr<FMovementModifierBase>(CopyOfSourcePtr));
+			}
+			else
+			{
+				UE_LOG(LogMover, Warning, TEXT("FMovementModifierGroup::operator= trying to copy invalid Other Modifier"));
+			}
+		}
+	}
+
+	return *this;
+}
+
+bool FMovementModifierGroup::operator==(const FMovementModifierGroup& Other) const
+{
+	if (ActiveModifiers.Num() != Other.ActiveModifiers.Num())
+	{
+		return false;
+	}
+	if (QueuedModifiers.Num() != Other.QueuedModifiers.Num())
+	{
+		return false;
+	}
+
+
+	for (int32 i = 0; i < ActiveModifiers.Num(); ++i)
+	{
+		if (ActiveModifiers[i].IsValid() == Other.ActiveModifiers[i].IsValid())
+		{
+			if (ActiveModifiers[i].IsValid())
+			{
+				// TODO: Implement deep equality checks
+			}
+		}
+		else
+		{
+			return false; // Mismatch in validity
+		}
+	}
+	for (int32 i = 0; i < QueuedModifiers.Num(); ++i)
+	{
+		if (QueuedModifiers[i].IsValid() == Other.QueuedModifiers[i].IsValid())
+		{
+			if (QueuedModifiers[i].IsValid())
+			{
+				// TODO: Implement deep equality checks
+			}
+		}
+		else
+		{
+			return false; // Mismatch in validity
+		}
+	}
+	return true;
+}
+
+bool FMovementModifierGroup::operator!=(const FMovementModifierGroup& Other) const
+{
+	return !(FMovementModifierGroup::operator==(Other));
+}
+
+void FMovementModifierGroup::AddStructReferencedObjects(FReferenceCollector& Collector) const
+{
+	for (const TSharedPtr<FMovementModifierBase>& Modifier : ActiveModifiers)
+	{
+		if (Modifier.IsValid())
+		{
+			Modifier->AddReferencedObjects(Collector);
+		}
+	}
+
+	for (const TSharedPtr<FMovementModifierBase>& Modifier : QueuedModifiers)
+	{
+		if (Modifier.IsValid())
+		{
+			Modifier->AddReferencedObjects(Collector);
+		}
+	}
+}
+
+FString FMovementModifierGroup::ToSimpleString() const
+{
+	return FString::Printf(TEXT("FMovementModifierGroup: Active: %i Queued: %i"), ActiveModifiers.Num(), QueuedModifiers.Num());
+}
+
+TArray<TSharedPtr<FMovementModifierBase>>::TConstIterator FMovementModifierGroup::GetActiveModifiersIterator() const
+{
+	return ActiveModifiers.CreateConstIterator();
+}
+
+TArray<TSharedPtr<FMovementModifierBase>>::TConstIterator FMovementModifierGroup::GetQueuedModifiersIterator() const
+{
+	return ActiveModifiers.CreateConstIterator();
+}
+
+void FMovementModifierGroup::FlushModifierArrays(UMoverComponent* MoverComp, const FMoverTimeStep& TimeStep, const FMoverSyncState& SyncState, const FMoverAuxStateContext& AuxState)
+{
+	// Remove any finished moves
+	ActiveModifiers.RemoveAll([MoverComp, TimeStep, SyncState, AuxState, this]
+		(const TSharedPtr<FMovementModifierBase>& Modifier)
+		{
+			if (Modifier.IsValid())
+			{
+				if (Modifier->IsFinished(TimeStep.BaseSimTimeMs))
+				{
+					Modifier->EndModifier(MoverComp, TimeStep, SyncState, AuxState);
+					return true;
+				}
+			}
+			else
+			{
+				return true;	
+			}
+
+			return false;
+		});
+
+	// Make any queued moves active
+	for (TSharedPtr<FMovementModifierBase>& QueuedModifier : QueuedModifiers)
+	{
+		bool bModidiferAlreadyActive = false;
+		for (TSharedPtr<FMovementModifierBase>& ActiveModifier : ActiveModifiers)
+		{
+			// We don't really need to assign the QueuedModifier a start time but it would help if modifiers are compared based off of start time as well
+			QueuedModifier->StartSimTimeMs = TimeStep.BaseSimTimeMs;
+
+			// We only want to queue this queued modifier if it wasn't already added from state received from authority. If we already have the modifier just assign it a handle since it's already been activated.
+			if (QueuedModifier->Matches(ActiveModifier.Get()) && !ActiveModifier->GetHandle().IsValid())
+			{
+				ActiveModifier->SetHandleFromExistingModifier(QueuedModifier->GetHandle());
+				bModidiferAlreadyActive = true;
+				break;
+			}
+		}
+
+		if (!bModidiferAlreadyActive)
+		{
+			ActiveModifiers.Add(QueuedModifier);
+			QueuedModifier->StartModifier(MoverComp, TimeStep, SyncState, AuxState);
+		}
+	}
+
+	QueuedModifiers.Empty();
+}
+
+struct FMovementModifierDeleter
+{
+	FORCEINLINE void operator()(FMovementModifierBase* Object) const
+	{
+		check(Object);
+		UScriptStruct* ScriptStruct = Object->GetScriptStruct();
+		check(ScriptStruct);
+		ScriptStruct->DestroyStruct(Object);
+		FMemory::Free(Object);
+	}
+};
+
+/* static */ void FMovementModifierGroup::NetSerializeMovementModifierArray(FArchive& Ar, TArray< TSharedPtr<FMovementModifierBase> >& ModifierArray, uint8 MaxNumModifiersToSerialize)
+{
+	uint8 NumModifiersToSerialize;
+	if (Ar.IsSaving())
+	{
+		UE_CLOG(ModifierArray.Num() > MaxNumModifiersToSerialize, LogMover, Warning, TEXT("Too many Modifiers (%d!) to net serialize. Clamping to %d"),
+			ModifierArray.Num(), MaxNumModifiersToSerialize);
+
+		NumModifiersToSerialize = FMath::Min<int32>(ModifierArray.Num(), MaxNumModifiersToSerialize);
+	}
+
+	Ar << NumModifiersToSerialize;
+
+	if (Ar.IsLoading())
+	{
+		ModifierArray.SetNumZeroed(NumModifiersToSerialize);
+	}
+
+	for (int32 i = 0; i < NumModifiersToSerialize && !Ar.IsError(); ++i)
+	{
+		TCheckedObjPtr<UScriptStruct> ScriptStruct = ModifierArray[i].IsValid() ? ModifierArray[i]->GetScriptStruct() : nullptr;
+		UScriptStruct* ScriptStructLocal = ScriptStruct.Get();
+		Ar << ScriptStruct;
+
+		if (ScriptStruct.IsValid())
+		{
+			// Restrict replication to derived classes of FMovementModifierBase for security reasons:
+			// If FMovementModifierGroup is replicated through a Server RPC, we need to prevent clients from sending us
+			// arbitrary ScriptStructs due to the allocation/reliance on GetCppStructOps below which could trigger a server crash
+			// for invalid structs. All provided sources are direct children of FMovementModifierBase and we never expect to have deep hierarchies
+			// so this should not be too costly
+			bool bIsDerivedFromBase = false;
+			UStruct* CurrentSuperStruct = ScriptStruct->GetSuperStruct();
+			while (CurrentSuperStruct)
+			{
+				if (CurrentSuperStruct == FMovementModifierBase::StaticStruct())
+				{
+					bIsDerivedFromBase = true;
+					break;
+				}
+				CurrentSuperStruct = CurrentSuperStruct->GetSuperStruct();
+			}
+
+			if (bIsDerivedFromBase)
+			{
+				if (Ar.IsLoading())
+				{
+					if (ModifierArray[i].IsValid() && ScriptStructLocal == ScriptStruct.Get())
+					{
+						// What we have locally is the same type as we're being serialized into, so we don't need to
+						// reallocate - just use existing structure
+					}
+					else
+					{
+						// For now, just reset/reallocate the data when loading.
+						// Longer term if we want to generalize this and use it for property replication, we should support
+						// only reallocating when necessary
+						FMovementModifierBase* NewModifier = (FMovementModifierBase*)FMemory::Malloc(ScriptStruct->GetCppStructOps()->GetSize());
+						ScriptStruct->InitializeStruct(NewModifier);
+
+						ModifierArray[i] = TSharedPtr<FMovementModifierBase>(NewModifier, FMovementModifierDeleter());
+					}
+				}
+
+				ModifierArray[i]->NetSerialize(Ar);
+			}
+			else
+			{
+				UE_LOG(LogMover, Error, TEXT("FMovementModifierGroup::NetSerialize: ScriptStruct not derived from FMovementModifierBase attempted to serialize."));
+				Ar.SetError();
+				break;
+			}
+		}
+		else if (ScriptStruct.IsError())
+		{
+			UE_LOG(LogMover, Error, TEXT("FMovementModifierGroup::NetSerialize: Invalid ScriptStruct serialized."));
+			Ar.SetError();
+			break;
+		}
+	}
+}
+
+
diff --git a/Engine/Plugins/Experimental/Mover/Source/Mover/Private/MoverComponent.cpp b/Engine/Plugins/Experimental/Mover/Source/Mover/Private/MoverComponent.cpp
--- a/Engine/Plugins/Experimental/Mover/Source/Mover/Private/MoverComponent.cpp
+++ b/Engine/Plugins/Experimental/Mover/Source/Mover/Private/MoverComponent.cpp
@@ -261,6 +261,9 @@ void UMoverComponent::ProduceInput(const int32 DeltaTimeMS, FMoverInputCmdContex
 
 void UMoverComponent::RestoreFrame(const FMoverSyncState* SyncState, const FMoverAuxStateContext* AuxState)
 {
+	const FMoverSyncState& InvalidSyncState = GetSyncState();
+	const FMoverAuxStateContext& InvalidAuxState = CachedLastAuxState;
+	OnSimulationPreRollback(&InvalidSyncState, SyncState, &InvalidAuxState, AuxState);
 	SetFrameStateFromContext(SyncState, AuxState, /* rebase? */ true);
 	OnSimulationRollback(SyncState, AuxState);
 }
@@ -288,6 +291,76 @@ void UMoverComponent::FinalizeFrame(const FMoverSyncState* SyncState, const FMov
 	}
 }
 
+void UMoverComponent::TickInterpolatedSimProxy(const FMoverTimeStep& TimeStep, const FMoverInputCmdContext& InputCmd, UMoverComponent* MoverComp, const FMoverSyncState& CachedSyncState, const FMoverSyncState& SyncState, const FMoverAuxStateContext& AuxState)
+{
+	TArray<TSharedPtr<FMovementModifierBase>> ModifiersToStart;
+	TArray<TSharedPtr<FMovementModifierBase>> ModifiersToEnd;
+
+	for (auto ModifierFromSyncStateIt = SyncState.MovementModifiers.GetActiveModifiersIterator(); ModifierFromSyncStateIt; ++ModifierFromSyncStateIt)
+	{
+		const TSharedPtr<FMovementModifierBase> ModifierFromSyncState = *ModifierFromSyncStateIt;
+		
+		bool bContainsModifier = false;
+		for (auto ModifierFromCacheIt = CachedSyncState.MovementModifiers.GetActiveModifiersIterator(); ModifierFromCacheIt; ++ModifierFromCacheIt)
+		{
+			const TSharedPtr<FMovementModifierBase> ModifierFromCache = *ModifierFromCacheIt;
+			
+			if (ModifierFromSyncState->Matches(ModifierFromCache.Get()))
+			{
+				bContainsModifier = true;
+				break;
+			}
+		}
+
+		if (!bContainsModifier)
+		{
+			ModifiersToStart.Add(ModifierFromSyncState);
+		}
+	}
+
+	for (auto ModifierFromCacheIt = CachedSyncState.MovementModifiers.GetActiveModifiersIterator(); ModifierFromCacheIt; ++ModifierFromCacheIt)
+	{
+		const TSharedPtr<FMovementModifierBase> ModifierFromCache = *ModifierFromCacheIt;
+		
+		bool bContainsModifier = false;
+		for (auto ModifierFromSyncStateIt = SyncState.MovementModifiers.GetActiveModifiersIterator(); ModifierFromSyncStateIt; ++ModifierFromSyncStateIt)
+		{
+			const TSharedPtr<FMovementModifierBase> ModifierFromSyncState = *ModifierFromSyncStateIt;
+			
+			if (ModifierFromSyncState->Matches(ModifierFromCache.Get()))
+			{
+				bContainsModifier = true;
+				break;
+			}
+		}
+
+		if (!bContainsModifier)
+		{
+			ModifiersToEnd.Add(ModifierFromCache);
+		}
+	}
+
+	for (TSharedPtr<FMovementModifierBase> Modifier : ModifiersToStart)
+	{
+		Modifier->GenerateHandle();
+		Modifier->OnStart(MoverComp, TimeStep, SyncState, AuxState);
+	}
+
+	for (auto ModifierIt = SyncState.MovementModifiers.GetActiveModifiersIterator(); ModifierIt; ++ModifierIt)
+	{
+		if (ModifierIt->IsValid())
+		{
+			ModifierIt->Get()->OnPreMovement(this, TimeStep);
+			ModifierIt->Get()->OnPostMovement(this, TimeStep, SyncState, AuxState);
+		}
+	}
+
+	for (TSharedPtr<FMovementModifierBase> Modifier : ModifiersToEnd)
+	{
+		Modifier->OnEnd(MoverComp, TimeStep, SyncState, AuxState);
+	}
+}
+
 
 void UMoverComponent::InitializeSimulationState(FMoverSyncState* OutSync, FMoverAuxStateContext* OutAux)
 {
@@ -434,6 +507,104 @@ UBaseMovementMode* UMoverComponent::FindMovementMode(TSubclassOf<UBaseMovementMo
 	return FindMode_Mutable(MovementMode);
 }
 
+void UMoverComponent::K2_FindMovementModifier(FMovementModifierHandle ModifierHandle, bool& bFoundModifier, int32& TargetAsRawBytes) const
+{
+	// This will never be called, the exec version below will be hit instead
+	checkNoEntry();
+}
+
+DEFINE_FUNCTION(UMoverComponent::execK2_FindMovementModifier)
+{
+	P_GET_STRUCT(FMovementModifierHandle, ModifierHandle);
+	P_GET_UBOOL_REF(bFoundModifier);
+
+	Stack.MostRecentPropertyAddress = nullptr;
+	Stack.MostRecentPropertyContainer = nullptr;
+	Stack.StepCompiledIn<FStructProperty>(nullptr);
+	
+	void* ModifierPtr = Stack.MostRecentPropertyAddress;
+	FStructProperty* StructProp = CastField<FStructProperty>(Stack.MostRecentProperty);
+
+	P_FINISH;
+
+	bFoundModifier = false;
+	
+	if (!ModifierPtr)
+	{
+		FBlueprintExceptionInfo ExceptionInfo(
+			EBlueprintExceptionType::AbortExecution,
+			LOCTEXT("MoverComponent_GetActiveLayeredMove_UnresolvedTarget", "Failed to resolve the OutLayeredMove for GetActiveLayeredMove")
+		);
+	
+		FBlueprintCoreDelegates::ThrowScriptException(P_THIS, Stack, ExceptionInfo);
+	}
+	else if (!StructProp)
+	{
+		FBlueprintExceptionInfo ExceptionInfo(
+			EBlueprintExceptionType::AbortExecution,
+			LOCTEXT("MoverComponent_GetActiveLayeredMove_TargetNotStruct", "GetActiveLayeredMove: Target for OutLayeredMove is not a valid type. It must be a Struct and a child of FLayeredMoveBase.")
+		);
+	
+		FBlueprintCoreDelegates::ThrowScriptException(P_THIS, Stack, ExceptionInfo);
+	}
+	else if (!StructProp->Struct || !StructProp->Struct->IsChildOf(FMovementModifierBase::StaticStruct()))
+	{
+		FBlueprintExceptionInfo ExceptionInfo(
+			EBlueprintExceptionType::AbortExecution,
+			LOCTEXT("MoverComponent_GetActiveLayeredMove_BadType", "GetActiveLayeredMove: Target for OutLayeredMove is not a valid type. Must be a child of FLayeredMoveBase.")
+		);
+	
+		FBlueprintCoreDelegates::ThrowScriptException(P_THIS, Stack, ExceptionInfo);
+	}
+	else
+	{
+		P_NATIVE_BEGIN;
+		
+		if (const FMovementModifierBase* FoundActiveMove = P_THIS->FindMovementModifier(ModifierHandle))
+		{
+			StructProp->Struct->CopyScriptStruct(ModifierPtr, FoundActiveMove);
+			bFoundModifier = true;
+		}
+
+		P_NATIVE_END;
+	}
+}
+
+bool UMoverComponent::IsModifierActiveOrQueued(const FMovementModifierHandle& ModifierHandle) const
+{
+	return FindMovementModifier(ModifierHandle) ? true : false;
+}
+
+const FMovementModifierBase* UMoverComponent::FindMovementModifier(const FMovementModifierHandle& ModifierHandle) const
+{
+	if (bHasValidCachedState)
+	{
+		// Check active modifiers for modifier handle
+		for (auto ActiveModifierFromSyncStateIt = CachedLastSyncState.MovementModifiers.GetActiveModifiersIterator(); ActiveModifierFromSyncStateIt; ++ActiveModifierFromSyncStateIt)
+		{
+			const TSharedPtr<FMovementModifierBase> ActiveModifierFromSyncState = *ActiveModifierFromSyncStateIt;
+
+			if (ModifierHandle == ActiveModifierFromSyncState->GetHandle())
+			{
+				return ActiveModifierFromSyncState.Get();
+			}
+		}
+
+		// Check queued modifiers for modifier handle
+		for (auto QueuedModifierFromSyncStateIt = CachedLastSyncState.MovementModifiers.GetActiveModifiersIterator(); QueuedModifierFromSyncStateIt; ++QueuedModifierFromSyncStateIt)
+		{
+			const TSharedPtr<FMovementModifierBase> QueuedModifierFromSyncState = *QueuedModifierFromSyncStateIt;
+
+			if (ModifierHandle == QueuedModifierFromSyncState->GetHandle())
+			{
+				return QueuedModifierFromSyncState.Get();
+			}
+		}
+	}
+	
+	return nullptr;
+}
+
 void UMoverComponent::SetFrameStateFromContext(const FMoverSyncState* SyncState, const FMoverAuxStateContext* AuxState, bool bRebaseBasedState)
 {
 
@@ -626,6 +797,11 @@ void UMoverComponent::UpdateTickRegistration()
 	SetComponentTickEnabled(bHasUpdatedComponent && bAutoActivate);
 }
 
+void UMoverComponent::OnSimulationPreRollback(const FMoverSyncState* InvalidSyncState, const FMoverSyncState* SyncState, const FMoverAuxStateContext* InvalidAuxState, const FMoverAuxStateContext* AuxState)
+{
+	ModeFSM->OnSimulationPreRollback(InvalidSyncState, SyncState, InvalidAuxState, AuxState);
+}
+
 
 void UMoverComponent::OnSimulationRollback(const FMoverSyncState* SyncState, const FMoverAuxStateContext* AuxState)
 {
@@ -918,6 +1094,49 @@ void UMoverComponent::QueueLayeredMove(TSharedPtr<FLayeredMoveBase> LayeredMove)
 	ModeFSM->QueueLayeredMove(LayeredMove);
 }
 
+FMovementModifierHandle UMoverComponent::K2_QueueMovementModifier(const int32& MoveAsRawData)
+{
+	// This will never be called, the exec version below will be hit instead
+	checkNoEntry();
+	return 0;
+}
+
+DEFINE_FUNCTION(UMoverComponent::execK2_QueueMovementModifier)
+{
+	Stack.StepCompiledIn<FStructProperty>(nullptr);
+	void* MovePtr = Stack.MostRecentPropertyAddress;
+	FStructProperty* StructProp = CastField<FStructProperty>(Stack.MostRecentProperty);
+
+	P_FINISH;
+
+	P_NATIVE_BEGIN;
+
+	const bool bHasValidStructProp = StructProp && StructProp->Struct && StructProp->Struct->IsChildOf(FMovementModifierBase::StaticStruct());
+
+	if (ensureMsgf((bHasValidStructProp && MovePtr), TEXT("An invalid type (%s) was sent to a QueueMovementModifier node. A struct derived from FMovementModifierBase is required. No modifier will be queued."),
+		StructProp ? *GetNameSafe(StructProp->Struct) : *Stack.MostRecentProperty->GetClass()->GetName()))
+	{
+		// Could we steal this instead of cloning? (move semantics)
+		FMovementModifierBase* MoveAsBasePtr = reinterpret_cast<FMovementModifierBase*>(MovePtr);
+		FMovementModifierBase* ClonedMove = MoveAsBasePtr->Clone();
+
+		FMovementModifierHandle ModifierID = P_THIS->QueueMovementModifier(TSharedPtr<FMovementModifierBase>(ClonedMove));
+		*static_cast<FMovementModifierHandle*>(RESULT_PARAM) = ModifierID;
+	}
+
+	P_NATIVE_END;
+}
+
+FMovementModifierHandle UMoverComponent::QueueMovementModifier(TSharedPtr<FMovementModifierBase> Modifier)
+{
+	return ModeFSM->QueueMovementModifier(Modifier);
+}
+
+void UMoverComponent::CancelModifierFromHandle(FMovementModifierHandle ModifierHandle)
+{
+	ModeFSM->CancelModifierFromHandle(ModifierHandle);
+}
+
 void UMoverComponent::K2_QueueInstantMovementEffect(const int32& EffectAsRawData)
 {
 	// This will never be called, the exec version below will be hit instead
@@ -1499,6 +1718,11 @@ const FMoverInputCmdContext& UMoverComponent::GetLastInputCmd() const
 	return CachedLastUsedInputCmd;
 }
 
+const FMoverTimeStep& UMoverComponent::GetLastTimeStep() const
+{
+	return CachedLastSimTickTimeStep;
+}
+
 IMovementSettingsInterface* UMoverComponent::FindSharedSettings_Mutable(const UClass* ByType) const
 {
 	check(ByType);
diff --git a/Engine/Plugins/Experimental/Mover/Source/Mover/Public/MovementModifier.h b/Engine/Plugins/Experimental/Mover/Source/Mover/Public/MovementModifier.h
new file mode 100644
--- /dev/null
+++ b/Engine/Plugins/Experimental/Mover/Source/Mover/Public/MovementModifier.h
@@ -0,0 +1,233 @@
+﻿// Copyright Epic Games, Inc. All Rights Reserved.
+
+#pragma once
+
+#include "MoveLibrary/MovementUtilsTypes.h"
+#include "MovementModifier.generated.h"
+
+class UMoverComponent;
+struct FMoverAuxStateContext;
+struct FMoverTimeStep;
+struct FMoverSyncState;
+struct FMoverInputCmdContext;
+
+#define MODIFIER_HANDLE_TYPE uint16
+#define MODIFIER_HANDLE_MAX UINT16_MAX
+#define MODIFIER_INVALID_HANDLE 0
+
+/**
+ * Handle to identify movement modifiers.
+ */
+USTRUCT(BlueprintType)
+struct MOVER_API FMovementModifierHandle
+{
+	GENERATED_BODY()
+
+	FMovementModifierHandle()
+		: Handle(0)
+	{
+	}
+
+	FMovementModifierHandle(MODIFIER_HANDLE_TYPE InHandle)
+		: Handle(InHandle)
+	{
+	}
+
+	/** Creates a new handle */
+	void GenerateHandle();
+	
+	bool operator==(const FMovementModifierHandle& Other) const
+	{
+		return Handle == Other.Handle;
+	}
+
+	bool operator!=(const FMovementModifierHandle& Other) const
+	{
+		return Handle != Other.Handle;
+	}
+
+	FString ToString() const
+	{
+		return FString::Printf(TEXT("%d"), Handle);
+	}
+
+	bool IsValid() const
+	{
+		return Handle != MODIFIER_INVALID_HANDLE;
+	}
+
+	void Invalidate()
+	{
+		Handle = MODIFIER_INVALID_HANDLE;
+	}
+	
+private:
+	MODIFIER_HANDLE_TYPE Handle;
+};
+
+/**
+ * Movement Modifier: Used to apply changes that indirectly influence the movement simulation,
+ * without proposing or executing any movement, but still in sync with the sim.
+ * Example usages: changing groups of settings, movement mode re-mappings, etc.
+ * 
+ * Note: Currently mover expects to only have one type of modifier active at a time.
+ *		 This can be fixed by extending the Matches function to check more than just type,
+ *		 but make sure anything used to compare is synced through the NetSerialize function.
+ */
+USTRUCT(BlueprintType)
+struct MOVER_API FMovementModifierBase
+{
+	GENERATED_BODY()
+
+	FMovementModifierBase();
+	virtual ~FMovementModifierBase() {}
+	
+	// This modifier will expire after a set amount of time if > 0. If 0, it will be ticked only once, regardless of time step. It will need to be manually ended if < 0. 
+	UPROPERTY(EditAnywhere, BlueprintReadWrite, Category = Mover)
+	float DurationMs;
+
+	// The simulation time this move first ticked (< 0 means it hasn't started yet)
+	UPROPERTY(VisibleInstanceOnly, BlueprintReadOnly, Transient, Category = Mover)
+	float StartSimTimeMs;
+	
+	/** Fired when this modifier is activated. */
+	virtual void OnStart(UMoverComponent* MoverComp, const FMoverTimeStep& TimeStep, const FMoverSyncState& SyncState, const FMoverAuxStateContext& AuxState) {}
+	
+	/** Fired when this modifier is deactivated. */
+	virtual void OnEnd(UMoverComponent* MoverComp, const FMoverTimeStep& TimeStep, const FMoverSyncState& SyncState, const FMoverAuxStateContext& AuxState) {}
+	
+	/**
+	 * Fired just before a movement Substep
+	 */
+	virtual void OnPreMovement(UMoverComponent* MoverComp, const FMoverTimeStep& TimeStep) {}
+
+	/** Fired after a movement Substep */
+	virtual void OnPostMovement(UMoverComponent* MoverComp, const FMoverTimeStep& TimeStep, const FMoverSyncState& SyncState, const FMoverAuxStateContext& AuxState) {}
+
+	/** Kicks off this modifier, allowing any initialization to occur. */
+	void StartModifier(UMoverComponent* MoverComp, const FMoverTimeStep& TimeStep, const FMoverSyncState& SyncState, const FMoverAuxStateContext& AuxState);
+
+	/** Ends this move, allowing any cleanup to occur. */
+	void EndModifier(UMoverComponent* MoverComp, const FMoverTimeStep& TimeStep, const FMoverSyncState& SyncState, const FMoverAuxStateContext& AuxState);
+
+	/** Runtime query whether this modifier is finished and can be destroyed. The default implementation is based on DurationMs. */
+	virtual bool IsFinished(float CurrentSimTimeMs) const;
+	
+	// @return newly allocated copy of this FMovementModifier. Must be overridden by child classes
+	virtual FMovementModifierBase* Clone() const;
+
+	virtual void NetSerialize(FArchive& Ar);
+
+	virtual UScriptStruct* GetScriptStruct() const;
+
+	virtual FString ToSimpleString() const;
+
+	virtual void AddReferencedObjects(class FReferenceCollector& Collector) {}
+
+	/** 
+	 * Used to compare modifiers and check if they are the same instance of modifiers.
+	 * Doesn't need to be overridden but more specific data to match is best
+	 * Note: Current default implementation only checks type and nothing else.
+	 */ 
+	virtual bool Matches(const FMovementModifierBase* Other) const;
+
+	FMovementModifierHandle GetHandle() const;
+
+	void GenerateHandle();
+
+	/**
+	 * Used to write to a valid handle to a invalid handle
+	 * Currently used to write a local handle to modifiers that were added from a rollback so they don't have a handle yet
+	 * This is done to avoid queueing a modifier again as a local client applies potential input
+	 */
+	void SetHandleFromExistingModifier(const FMovementModifierHandle& OtherModifierHandle);
+	
+protected:
+	/**
+	 * Modifier handle local to this client or server instance. Used to cancel or query for a active/queued modifier.
+	 */
+	UPROPERTY()
+	FMovementModifierHandle LocalModifierHandle;
+};
+
+template<>
+struct TStructOpsTypeTraits< FMovementModifierBase > : public TStructOpsTypeTraitsBase2< FMovementModifierBase >
+{
+	enum
+	{
+		//WithNetSerializer = true,
+		WithCopy = true
+	};
+};
+
+
+// A collection of movement modifiers affecting a movable actor
+USTRUCT(BlueprintType)
+struct MOVER_API FMovementModifierGroup
+{
+	GENERATED_BODY()
+
+	FMovementModifierGroup() {}
+	virtual ~FMovementModifierGroup() {}
+	
+	bool HasAnyMoves() const { return (!ActiveModifiers.IsEmpty() || !QueuedModifiers.IsEmpty()); }
+
+	/** Serialize all modifiers and their states for this group */
+	void NetSerialize(FArchive& Ar, uint8 MaxNumModifiersToSerialize = MAX_uint8);
+	
+	void QueueMovementModifier(TSharedPtr<FMovementModifierBase> Modifier);
+
+	void CancelModifierFromHandle(const FMovementModifierHandle& HandleToCancel);
+	
+	// Generates active modifier list (by calling FlushModifierArrays) and returns the an array of all currently active modifiers
+	TArray<TSharedPtr<FMovementModifierBase>> GenerateActiveModifiers(UMoverComponent* MoverComp, const FMoverTimeStep& TimeStep, const FMoverSyncState& SyncState, const FMoverAuxStateContext& AuxState);
+	
+	/** Copy operator - deep copy so it can be used for archiving/saving off moves */
+	FMovementModifierGroup& operator=(const FMovementModifierGroup& Other);
+
+	/** Comparison operator - needs matching modifier along with identical states in those structs */
+	bool operator==(const FMovementModifierGroup& Other) const;
+
+	/** Comparison operator */
+	bool operator!=(const FMovementModifierGroup& Other) const;
+
+	/** Exposes references to GC system */
+	void AddStructReferencedObjects(FReferenceCollector& Collector) const;
+
+	/** Get a simplified string representation of this group. Typically for debugging. */
+	FString ToSimpleString() const;
+
+	/** Const access to active moves */
+	TArray<TSharedPtr<FMovementModifierBase>>::TConstIterator GetActiveModifiersIterator() const;
+	
+	/** Const access to queued moves */
+	TArray<TSharedPtr<FMovementModifierBase>>::TConstIterator GetQueuedModifiersIterator() const;
+	
+protected:
+	/** Movement modifiers that are currently active in this group */
+	TArray< TSharedPtr<FMovementModifierBase> > ActiveModifiers;
+
+	/** Movement modifiers that are queued to become active next sim frame */
+	TArray< TSharedPtr<FMovementModifierBase> > QueuedModifiers;
+	
+	// Clears out any finished or invalid modifiers and adds any queued modifiers to the active modifiers
+	void FlushModifierArrays(UMoverComponent* MoverComp, const FMoverTimeStep& TimeStep, const FMoverSyncState& SyncState, const FMoverAuxStateContext& AuxState);
+	
+	/** Helper function for serializing array of movement modifiers */
+	static void NetSerializeMovementModifierArray(FArchive& Ar, TArray< TSharedPtr<FMovementModifierBase> >& ModifierArray, uint8 MaxNumModifiersToSerialize = MAX_uint8);
+	
+
+};
+
+
+template<>
+struct TStructOpsTypeTraits<FMovementModifierGroup> : public TStructOpsTypeTraitsBase2<FMovementModifierGroup>
+{
+	enum
+	{
+		WithCopy = true,		// Necessary so that TSharedPtr<FMovementModifierGroup> Data is copied around
+		//WithNetSerializer = true,
+		WithIdenticalViaEquality = true,
+		WithAddStructReferencedObjects = true,
+	};
+};
diff --git a/Engine/Plugins/Experimental/Mover/Source/Mover/Public/MoverComponent.h b/Engine/Plugins/Experimental/Mover/Source/Mover/Public/MoverComponent.h
--- a/Engine/Plugins/Experimental/Mover/Source/Mover/Public/MoverComponent.h
+++ b/Engine/Plugins/Experimental/Mover/Source/Mover/Public/MoverComponent.h
@@ -15,6 +15,7 @@
 #if UE_ENABLE_INCLUDE_ORDER_DEPRECATED_IN_5_4
 #include "Engine/HitResult.h"
 #endif // UE_ENABLE_INCLUDE_ORDER_DEPRECATED_IN_5_4
+#include "MovementModifier.h"
 #include "Backends/MoverBackendLiaison.h"
 #include "UObject/WeakInterfacePtr.h"
 #include "MoverComponent.generated.h"
@@ -119,6 +120,9 @@ public:
 	// Take output for simulation. Called by Network Prediction system.
 	void FinalizeFrame(const FMoverSyncState* SyncState, const FMoverAuxStateContext* AuxState);
 
+	// This is an opportunity to run code on the code on the simproxy in interpolated mode - currently used to help activate and deactivate modifiers on the simproxy in interpolated mode
+	void TickInterpolatedSimProxy(const FMoverTimeStep& TimeStep, const FMoverInputCmdContext& InputCmd, UMoverComponent* MoverComp, const FMoverSyncState& CachedSyncState, const FMoverSyncState& SyncState, const FMoverAuxStateContext& AuxState);
+	
 	// Seed initial values based on component's state. Called by Network Prediction system.
 	void InitializeSimulationState(FMoverSyncState* OutSync, FMoverAuxStateContext* OutAux);
 
@@ -163,6 +167,24 @@ public:
 	// Queue a layered move to start during the next simulation frame
 	void QueueLayeredMove(TSharedPtr<FLayeredMoveBase> Move);
 	
+	/**
+ 	 * Queue a Movement Modifier to start during the next simulation frame. This will clone whatever move you pass in, so you'll need to fully set it up before queuing.
+ 	 * @param MovementModifier The modifier to queue, which must be a LayeredMoveBase sub-type.
+ 	 * @return Returns a Modifier handle that can be used to query or cancel the movement modifier
+ 	 */
+	UFUNCTION(BlueprintCallable, CustomThunk, Category = Mover, meta = (CustomStructureParam = "MoveAsRawData", AllowAbstract = "false", DisplayName = "Queue Movement Modifier"))
+	FMovementModifierHandle K2_QueueMovementModifier(UPARAM(DisplayName="Movement Modifier") const int32& MoveAsRawData);
+	DECLARE_FUNCTION(execK2_QueueMovementModifier);
+
+	// Queue a Movement Modifier to start during the next simulation frame.
+	FMovementModifierHandle QueueMovementModifier(TSharedPtr<FMovementModifierBase> Modifier);
+	
+	/**
+	 * Cancel any active or queued Modifiers with the handle passed in.
+	 */
+	UFUNCTION(BlueprintCallable, Category = Mover)
+	void CancelModifierFromHandle(FMovementModifierHandle ModifierHandle);
+	
 	/**
 	 * Queue a Instant Movement Effect to start at the end of this frame or start of the next subtick - whichever happens first. This will clone whatever move you pass in, so you'll need to fully set it up before queuing.
 	 * @param InstantMovementEffect			The effect to queue, which must be a FInstantMovementEffect sub-type. 
@@ -303,6 +325,10 @@ public:	// Queries
 	UFUNCTION(BlueprintPure, Category = Mover)
 	const FMoverInputCmdContext& GetLastInputCmd() const;
 
+	// Get the most recent TimeStep
+	UFUNCTION(BlueprintPure, Category = Mover)
+	const FMoverTimeStep& GetLastTimeStep() const;
+
 	// Access the most recent floor check hit result.
 	UFUNCTION(BlueprintPure, Category = Mover)
 	bool TryGetFloorCheckHitResult(FHitResult& OutHitResult) const;
@@ -363,6 +389,24 @@ public:	// Queries
 		return nullptr;
 	}
 
+	
+	/**
+	 * Retrieves Movement modifier by writing to a target instance if it is the matching type. Note: Writing to the struct returned will not modify the active struct.
+	 * @param ModifierHandle		Handle of the modifier we're trying to cancel
+	 * @param bFoundModifier		Flag indicating whether modifier was found and data was actually written to target struct instance
+	 * @param TargetAsRawBytes		The data struct instance to write to, which must be a FMovementModifierBase sub-type
+	 */
+	UFUNCTION(BlueprintCallable, CustomThunk, Category = Mover, meta = (CustomStructureParam = "TargetAsRawBytes", AllowAbstract = "false", DisplayName = "Find Movement Modifier"))
+	void K2_FindMovementModifier(FMovementModifierHandle ModifierHandle, bool& bFoundModifier, UPARAM(DisplayName = "Out Movement Modifier") int32& TargetAsRawBytes) const;
+	DECLARE_FUNCTION(execK2_FindMovementModifier);
+
+	// Checks if the modifier handle passed in is active or queued on this mover component
+	UFUNCTION(BlueprintPure, Category = Mover)
+	bool IsModifierActiveOrQueued(const FMovementModifierHandle& ModifierHandle) const;
+	
+	// Find movement modifier by it's handle. Returns nullptr if the modifier couldn't be found
+	const FMovementModifierBase* FindMovementModifier(const FMovementModifierHandle& ModifierHandle) const;
+	
 protected:
 
 	/** Makes this component and owner actor reflect the state of a particular frame snapshot. This occurs after simulation ticking, as well as during a rollback before we resimulate forward.
@@ -382,6 +426,9 @@ protected:
 	void SetUpdatedComponent(USceneComponent* NewUpdatedComponent);
 	void UpdateTickRegistration();
 
+	/** Called when a rollback occurs, before the simulation state has been restored */
+	void OnSimulationPreRollback(const FMoverSyncState* InvalidSyncState, const FMoverSyncState* SyncState, const FMoverAuxStateContext* InvalidAuxState, const FMoverAuxStateContext* AuxState);
+	
 	/** Called when a rollback occurs, after the simulation state has been restored */
 	void OnSimulationRollback(const FMoverSyncState* SyncState, const FMoverAuxStateContext* AuxState);
 
diff --git a/Engine/Plugins/Experimental/Mover/Source/Mover/Public/MoverSimulationTypes.h b/Engine/Plugins/Experimental/Mover/Source/Mover/Public/MoverSimulationTypes.h
--- a/Engine/Plugins/Experimental/Mover/Source/Mover/Public/MoverSimulationTypes.h
+++ b/Engine/Plugins/Experimental/Mover/Source/Mover/Public/MoverSimulationTypes.h
@@ -8,6 +8,7 @@
 #include "MoverTypes.h"
 #include "MoveLibrary/MovementRecord.h"
 #include "LayeredMove.h"
+#include "MovementModifier.h"
 #include "MoverDataModelTypes.h"
 #include "UObject/Interface.h"
 #include "MoverSimulationTypes.generated.h"
@@ -107,13 +108,17 @@ struct MOVER_API FMoverSyncState
 public:
 
 	// The mode we ended up in from the prior frame, and which we'll start in during the next frame
-	UPROPERTY(BlueprintReadWrite, EditAnywhere, Category = Mover)
+	UPROPERTY(BlueprintReadOnly, VisibleAnywhere, Category = Mover)
 	FName MovementMode;
 
 	// Additional moves influencing our proposed motion
-	UPROPERTY(BlueprintReadWrite, EditAnywhere, Category = Mover)
+	UPROPERTY(BlueprintReadOnly, VisibleAnywhere, Category = Mover)
 	FLayeredMoveGroup LayeredMoves;
 
+	// Additional modifiers influencing our simulation
+	UPROPERTY(BlueprintReadOnly, VisibleAnywhere, Category = Mover)
+	FMovementModifierGroup MovementModifiers;
+
 	UPROPERTY(BlueprintReadWrite, Category = Mover)
 	FMoverDataCollection SyncStateCollection;
 
@@ -129,6 +134,7 @@ public:
 	{
 		P.Ar << MovementMode;
 		LayeredMoves.NetSerialize(P.Ar);
+		MovementModifiers.NetSerialize(P.Ar);
 
 		bool bIgnoredResult(false);
 		SyncStateCollection.NetSerialize(P.Ar, nullptr, bIgnoredResult);
@@ -138,6 +144,7 @@ public:
 	{
 		Out.Appendf("MovementMode: %s\n", TCHAR_TO_ANSI(*MovementMode.ToString()));
 		Out.Appendf("Layered Moves: %s\n", TCHAR_TO_ANSI(*LayeredMoves.ToSimpleString()));
+		Out.Appendf("Movement Modifiers: %s\n", TCHAR_TO_ANSI(*MovementModifiers.ToSimpleString()));
 		SyncStateCollection.ToString(Out);
 	}
 
@@ -151,6 +158,7 @@ public:
 	{
 		MovementMode = To->MovementMode;
 		LayeredMoves = To->LayeredMoves;
+		MovementModifiers = To->MovementModifiers;
 
 		SyncStateCollection.Interpolate(From->SyncStateCollection, To->SyncStateCollection, Pct);
 	}
