Mover 2.0: Fixing iteration issue when removing item from queued modifiers

diff --git a/Engine/Plugins/Experimental/Mover/Source/Mover/Private/MovementModeStateMachine.cpp b/Engine/Plugins/Experimental/Mover/Source/Mover/Private/MovementModeStateMachine.cpp
--- a/Engine/Plugins/Experimental/Mover/Source/Mover/Private/MovementModeStateMachine.cpp
+++ b/Engine/Plugins/Experimental/Mover/Source/Mover/Private/MovementModeStateMachine.cpp
@@ -417,14 +417,24 @@ FMovementModifierHandle UMovementModeStateMachine::QueueMovementModifier(TShared
 
 void UMovementModeStateMachine::CancelModifierFromHandle(FMovementModifierHandle ModifierHandle)
 {
-	for (TSharedPtr<FMovementModifierBase> Modifier : QueuedMovementModifiers)
+	QueuedMovementModifiers.RemoveAll([ModifierHandle, this]
+	(const TSharedPtr<FMovementModifierBase>& Modifier)
 	{
-		if (Modifier->GetHandle() == ModifierHandle)
+		if (Modifier.IsValid())
 		{
-			QueuedMovementModifiers.Remove(Modifier);
+			if (Modifier->GetHandle() == ModifierHandle)
+			{
+				return true;
+			}
+		}
+		else
+		{
+			return true;	
 		}
-	}
 
+		return false;
+	});
+	
 	ModifiersToCancel.Add(ModifierHandle);
 }
 
