Mover: adding BP-callable function to set primary visual component, allowing it to be set before BeginPlay's attempt to auto-find the first mesh

diff --git a/Engine/Plugins/Experimental/Mover/Source/Mover/Private/MoverComponent.cpp b/Engine/Plugins/Experimental/Mover/Source/Mover/Private/MoverComponent.cpp
--- a/Engine/Plugins/Experimental/Mover/Source/Mover/Private/MoverComponent.cpp
+++ b/Engine/Plugins/Experimental/Mover/Source/Mover/Private/MoverComponent.cpp
@@ -196,9 +196,11 @@ void UMoverComponent::BeginPlay()
 		// If no primary visual component is already set, fall back to searching for any kind of mesh
 		if (!PrimaryVisualComponent)
 		{
-			PrimaryVisualComponent = MyActor->FindComponentByClass<UMeshComponent>();
+			SetPrimaryVisualComponent(MyActor->FindComponentByClass<UMeshComponent>());
 		}
 
+		ensureMsgf(UpdatedComponent && (PrimaryVisualComponent != UpdatedComponent), TEXT("A Mover actor must have an UpdatedComponent and cannot have a PrimaryVisualComponent that is the same as UpdatedComponent"));
+
 		// Optional motion warping support
 		if (UMotionWarpingComponent* WarpingComp = MyActor->FindComponentByClass<UMotionWarpingComponent>())
 		{
@@ -1401,6 +1403,13 @@ USceneComponent* UMoverComponent::GetPrimaryVisualComponent() const
 	return PrimaryVisualComponent.Get();
 }
 
+void UMoverComponent::SetPrimaryVisualComponent(USceneComponent* SceneComponent)
+{
+	if (ensure(SceneComponent->GetOwner() == GetOwner()))
+	{
+		PrimaryVisualComponent = SceneComponent;
+	}
+}
 
 FVector UMoverComponent::GetVelocity() const
 { 
diff --git a/Engine/Plugins/Experimental/Mover/Source/Mover/Public/MoverComponent.h b/Engine/Plugins/Experimental/Mover/Source/Mover/Public/MoverComponent.h
--- a/Engine/Plugins/Experimental/Mover/Source/Mover/Public/MoverComponent.h
+++ b/Engine/Plugins/Experimental/Mover/Source/Mover/Public/MoverComponent.h
@@ -272,6 +272,10 @@ public:	// Queries
 		return Cast<T>(GetPrimaryVisualComponent());
 	}
 
+	// Sets this Mover actor's primary visual component. Must be a descendant of the updated component that acts as our movement root. 
+	UFUNCTION(BlueprintCallable, Category=Mover)
+	void SetPrimaryVisualComponent(USceneComponent* SceneComponent);
+
 	// Get the current velocity (units per second, worldspace)
 	UFUNCTION(BlueprintPure, Category = Mover)
 	FVector GetVelocity() const;
