Mover: adding optional BP compile-time data validation for backend liaisons.  Note this is occurring on the CDO, so we have to pass in the MoverComponent as additional context.
This allows us to verify the Actor's "Replicate Movement" option, which is incompatible w/ Network Prediction but required for Chaos Networked Physics.

diff --git a/Engine/Plugins/Experimental/Mover/Source/Mover/Private/Backends/MoverNetworkPhysicsLiaison.cpp b/Engine/Plugins/Experimental/Mover/Source/Mover/Private/Backends/MoverNetworkPhysicsLiaison.cpp
--- a/Engine/Plugins/Experimental/Mover/Source/Mover/Private/Backends/MoverNetworkPhysicsLiaison.cpp
+++ b/Engine/Plugins/Experimental/Mover/Source/Mover/Private/Backends/MoverNetworkPhysicsLiaison.cpp
@@ -19,6 +19,12 @@
 #include "PhysicsMover/PhysicsMoverManager.h"
 #include "PhysicsProxy/CharacterGroundConstraintProxy.h"
 
+#if WITH_EDITOR
+#include "Misc/DataValidation.h"
+#endif
+
+#define LOCTEXT_NAMESPACE "Mover"
+
 //////////////////////////////////////////////////////////////////////////
 
 extern FPhysicsDrivenMotionDebugParams GPhysicsDrivenMotionDebugParams;
@@ -285,6 +291,24 @@ int32 UMoverNetworkPhysicsLiaisonComponent::GetCurrentSimFrame()
 	return 0;
 }
 
+#if WITH_EDITOR
+EDataValidationResult UMoverNetworkPhysicsLiaisonComponent::ValidateData(FDataValidationContext& Context, const UMoverComponent& ValidationMoverComp) const
+{
+	if (const AActor* OwnerActor = ValidationMoverComp.GetOwner())
+	{
+		if (!OwnerActor->IsReplicatingMovement())
+		{
+			Context.AddError(FText::Format(LOCTEXT("RequiresReplicateMovementProperty", "The owning actor ({0}) does not have the ReplicateMovement property enabled. This is required for use with Chaos Networked Physics and poor quality movement with occur without it. Please enable it."),
+				FText::FromString(GetNameSafe(OwnerActor))));
+
+			return EDataValidationResult::Invalid;
+		}
+	}
+
+	return EDataValidationResult::Valid;
+}
+#endif // WITH_EDITOR
+
 //////////////////////////////////////////////////////////////////////////
 // UMoverNetworkPhysicsLiaisonComponent UObject interface
 
@@ -1121,3 +1145,5 @@ void UMoverNetworkPhysicsLiaisonComponent::OnContactModification_Internal(const
 		}
 	}
 }
+
+#undef LOCTEXT_NAMESPACE
diff --git a/Engine/Plugins/Experimental/Mover/Source/Mover/Private/Backends/MoverNetworkPredictionLiaison.cpp b/Engine/Plugins/Experimental/Mover/Source/Mover/Private/Backends/MoverNetworkPredictionLiaison.cpp
--- a/Engine/Plugins/Experimental/Mover/Source/Mover/Private/Backends/MoverNetworkPredictionLiaison.cpp
+++ b/Engine/Plugins/Experimental/Mover/Source/Mover/Private/Backends/MoverNetworkPredictionLiaison.cpp
@@ -7,6 +7,13 @@
 #include "GameFramework/Actor.h"
 #include "MoverComponent.h"
 
+#if WITH_EDITOR
+#include "Misc/DataValidation.h"
+#endif
+
+
+#define LOCTEXT_NAMESPACE "Mover"
+
 // ----------------------------------------------------------------------------------------------------------
 //	FMoverActorModelDef: the piece that ties everything together that we use to register with the NP system.
 // ----------------------------------------------------------------------------------------------------------
@@ -123,17 +130,28 @@ bool UMoverNetworkPredictionLiaisonComponent::WritePendingSyncState(const FMover
 	return true;
 }
 
-void UMoverNetworkPredictionLiaisonComponent::BeginPlay()
+#if WITH_EDITOR
+EDataValidationResult UMoverNetworkPredictionLiaisonComponent::ValidateData(FDataValidationContext& Context, const UMoverComponent& ValidationMoverComp) const
 {
-	Super::BeginPlay();
-
-	if (const AActor* OwnerActor = GetOwner())
+	if (const AActor* OwnerActor = ValidationMoverComp.GetOwner())
 	{
-		ensureMsgf(!OwnerActor->IsReplicatingMovement(),
-			TEXT("MoverComponent owning actor %s has the ReplicateMovement property enabled. This will conflict with Network Prediction and cause poor quality movement. Please disable it."),
-			*GetNameSafe(GetOwner()));
+		if (OwnerActor->IsReplicatingMovement())
+		{
+			Context.AddError(FText::Format(LOCTEXT("ConflictingReplicateMovementProperty", "The owning actor ({0}) has the ReplicateMovement property enabled. This will conflict with Network Prediction and cause poor quality movement. Please disable it."),
+				FText::FromString(GetNameSafe(OwnerActor))));
+
+			return EDataValidationResult::Invalid;
+		}
 	}
 
+	return EDataValidationResult::Valid;
+}
+#endif // WITH_EDITOR
+
+void UMoverNetworkPredictionLiaisonComponent::BeginPlay()
+{
+	Super::BeginPlay();
+
 	if (StartingOutSync && StartingOutAux)
 	{
 		if (FMoverDefaultSyncState* StartingSyncState = StartingOutSync->SyncStateCollection.FindMutableDataByType<FMoverDefaultSyncState>())
@@ -178,4 +196,6 @@ void UMoverNetworkPredictionLiaisonComponent::InitializeNetworkPredictionProxy()
 
 		NetworkPredictionProxy.Init<FMoverActorModelDef>(GetWorld(), GetReplicationProxies(), this, this);
 	}
-}
\ No newline at end of file
+}
+
+#undef LOCTEXT_NAMESPACE
diff --git a/Engine/Plugins/Experimental/Mover/Source/Mover/Private/Backends/MoverStandaloneLiaison.cpp b/Engine/Plugins/Experimental/Mover/Source/Mover/Private/Backends/MoverStandaloneLiaison.cpp
--- a/Engine/Plugins/Experimental/Mover/Source/Mover/Private/Backends/MoverStandaloneLiaison.cpp
+++ b/Engine/Plugins/Experimental/Mover/Source/Mover/Private/Backends/MoverStandaloneLiaison.cpp
@@ -55,10 +55,6 @@ void UMoverStandaloneLiaisonComponent::BeginPlay()
 	{
 		ensureMsgf(OwnerActor->GetNetMode() == NM_Standalone, TEXT("UMoverStandaloneLiaisonComponent is only valid for use in Standalone projects. Movement will not work properly in networked play."));
 
-		ensureMsgf(!OwnerActor->IsReplicatingMovement(),
-			TEXT("MoverComponent owning actor %s has the ReplicateMovement property enabled. This will conflict with Network Prediction and cause poor quality movement. Please disable it."),
-			*GetNameSafe(GetOwner()));
-
 		if (UMoverComponent* FoundMoverComp = OwnerActor->FindComponentByClass<UMoverComponent>())
 		{
 			MoverComp = FoundMoverComp;
diff --git a/Engine/Plugins/Experimental/Mover/Source/Mover/Private/MovementMode.cpp b/Engine/Plugins/Experimental/Mover/Source/Mover/Private/MovementMode.cpp
--- a/Engine/Plugins/Experimental/Mover/Source/Mover/Private/MovementMode.cpp
+++ b/Engine/Plugins/Experimental/Mover/Source/Mover/Private/MovementMode.cpp
@@ -3,7 +3,11 @@
 #include "MovementMode.h"
 #include "MoverComponent.h"
 #include "Engine/BlueprintGeneratedClass.h"
+
+#if WITH_EDITOR
 #include "Misc/DataValidation.h"
+#endif
+
 
 #include UE_INLINE_GENERATED_CPP_BY_NAME(MovementMode)
 
diff --git a/Engine/Plugins/Experimental/Mover/Source/Mover/Private/MoverComponent.cpp b/Engine/Plugins/Experimental/Mover/Source/Mover/Private/MoverComponent.cpp
--- a/Engine/Plugins/Experimental/Mover/Source/Mover/Private/MoverComponent.cpp
+++ b/Engine/Plugins/Experimental/Mover/Source/Mover/Private/MoverComponent.cpp
@@ -966,6 +966,14 @@ bool UMoverComponent::ValidateSetup(FDataValidationContext& Context) const
 
 		bDidFindAnyProblems = true;
 	}
+	else
+	{
+		IMoverBackendLiaisonInterface* BackendCDOAsInterface = Cast<IMoverBackendLiaisonInterface>(BackendClass->GetDefaultObject());
+		if (BackendCDOAsInterface && (BackendCDOAsInterface->ValidateData(Context, *this) == EDataValidationResult::Invalid))
+		{
+			bDidFindAnyProblems = true;
+		}
+	}
 
 	// Verify persistent types
 	for (const FMoverDataPersistence& PersistentSyncEntry : PersistentSyncStateDataTypes)
diff --git a/Engine/Plugins/Experimental/Mover/Source/Mover/Public/Backends/MoverBackendLiaison.h b/Engine/Plugins/Experimental/Mover/Source/Mover/Public/Backends/MoverBackendLiaison.h
--- a/Engine/Plugins/Experimental/Mover/Source/Mover/Public/Backends/MoverBackendLiaison.h
+++ b/Engine/Plugins/Experimental/Mover/Source/Mover/Public/Backends/MoverBackendLiaison.h
@@ -6,6 +6,8 @@
 #include "MoverSimulationTypes.h"
 #include "MoverBackendLiaison.generated.h"
 
+class FDataValidationContext;
+class UMoverComponent;
 
 /**
  * MoverBackendLiaisonInterface: any object or system wanting to be the driver of Mover actors must implement this. The intent is to act as a
@@ -29,4 +31,9 @@ public:
 
 	virtual bool ReadPendingSyncState(OUT FMoverSyncState& OutSyncState) { return false; }
 	virtual bool WritePendingSyncState(const FMoverSyncState& SyncStateToWrite) { return false; }
+
+#if WITH_EDITOR
+	virtual EDataValidationResult ValidateData(FDataValidationContext& Context, const UMoverComponent& ValidationMoverComp) const { return EDataValidationResult::Valid; }
+#endif // WITH_EDITOR
+
 };
\ No newline at end of file
diff --git a/Engine/Plugins/Experimental/Mover/Source/Mover/Public/Backends/MoverNetworkPhysicsLiaison.h b/Engine/Plugins/Experimental/Mover/Source/Mover/Public/Backends/MoverNetworkPhysicsLiaison.h
--- a/Engine/Plugins/Experimental/Mover/Source/Mover/Public/Backends/MoverNetworkPhysicsLiaison.h
+++ b/Engine/Plugins/Experimental/Mover/Source/Mover/Public/Backends/MoverNetworkPhysicsLiaison.h
@@ -112,6 +112,9 @@ public:
 	// IMoverBackendLiaisonInterface
 	virtual float GetCurrentSimTimeMs() override;
 	virtual int32 GetCurrentSimFrame() override;
+#if WITH_EDITOR
+	virtual EDataValidationResult ValidateData(FDataValidationContext& Context, const UMoverComponent& ValidationMoverComp) const override;
+#endif
 
 	// UObject interface
 	virtual void OnRegister() override;
diff --git a/Engine/Plugins/Experimental/Mover/Source/Mover/Public/Backends/MoverNetworkPredictionLiaison.h b/Engine/Plugins/Experimental/Mover/Source/Mover/Public/Backends/MoverNetworkPredictionLiaison.h
--- a/Engine/Plugins/Experimental/Mover/Source/Mover/Public/Backends/MoverNetworkPredictionLiaison.h
+++ b/Engine/Plugins/Experimental/Mover/Source/Mover/Public/Backends/MoverNetworkPredictionLiaison.h
@@ -49,6 +49,9 @@ public:
 	virtual int32 GetCurrentSimFrame() override;
 	virtual bool ReadPendingSyncState(OUT FMoverSyncState& OutSyncState) override;
 	virtual bool WritePendingSyncState(const FMoverSyncState& SyncStateToWrite) override;
+#if WITH_EDITOR
+	virtual EDataValidationResult ValidateData(FDataValidationContext& Context, const UMoverComponent& ValidationMoverComp) const override;
+#endif
 	// End IMoverBackendLiaisonInterface
 
 	virtual void BeginPlay() override;
