Mover 2.0: Adding gameplay tag query functions to support querying gameplay tags across different mover systems

diff --git a/Engine/Plugins/Experimental/Mover/Source/Mover/Private/MovementMode.cpp b/Engine/Plugins/Experimental/Mover/Source/Mover/Private/MovementMode.cpp
--- a/Engine/Plugins/Experimental/Mover/Source/Mover/Private/MovementMode.cpp
+++ b/Engine/Plugins/Experimental/Mover/Source/Mover/Private/MovementMode.cpp
@@ -123,6 +123,16 @@ EDataValidationResult UBaseMovementMode::IsDataValid(FDataValidationContext& Con
 #endif // WITH_EDITOR
 
 
+bool UBaseMovementMode::HasGameplayTag(FGameplayTag TagToFind, bool bExactMatch) const
+{
+	if (bExactMatch)
+	{
+		return GameplayTags.HasTagExact(TagToFind);
+	}
+
+	return GameplayTags.HasTag(TagToFind);
+}
+
 void UBaseMovementMode::OnRegistered(const FName ModeName)
 {
 }
diff --git a/Engine/Plugins/Experimental/Mover/Source/Mover/Private/MoverComponent.cpp b/Engine/Plugins/Experimental/Mover/Source/Mover/Private/MoverComponent.cpp
--- a/Engine/Plugins/Experimental/Mover/Source/Mover/Private/MoverComponent.cpp
+++ b/Engine/Plugins/Experimental/Mover/Source/Mover/Private/MoverComponent.cpp
@@ -607,6 +607,43 @@ const FMovementModifierBase* UMoverComponent::FindMovementModifier(const FMoveme
 	return nullptr;
 }
 
+bool UMoverComponent::HasGameplayTag(FGameplayTag TagToFind, bool bExactMatch) const
+{
+	if (bHasValidCachedState)
+	{
+		// Search Movement Modes
+		if (const UBaseMovementMode* ActiveMovementMode = GetMovementMode())
+		{
+			if (ActiveMovementMode->HasGameplayTag(TagToFind, bExactMatch))
+			{
+				return true;
+			}
+		}
+		
+		// Search Movement Modifiers
+		for (auto ModifierFromSyncStateIt = CachedLastSyncState.MovementModifiers.GetActiveModifiersIterator(); ModifierFromSyncStateIt; ++ModifierFromSyncStateIt)
+		{
+			const TSharedPtr<FMovementModifierBase> ModifierFromSyncState = *ModifierFromSyncStateIt;
+			if (ModifierFromSyncState->HasGameplayTag(TagToFind, bExactMatch))
+			{
+				return true;
+			}
+		}
+		
+		// Search Layered Moves
+		for (auto LayeredMoveFromSyncStateIt = CachedLastSyncState.LayeredMoves.GetActiveMovesIterator(); LayeredMoveFromSyncStateIt; ++LayeredMoveFromSyncStateIt)
+		{
+			const TSharedPtr<FLayeredMoveBase> LayeredMoveFromSyncState = *LayeredMoveFromSyncStateIt;
+			if (LayeredMoveFromSyncState->HasGameplayTag(TagToFind, bExactMatch))
+			{
+				return true;
+			}
+		}
+	}
+	
+	return false;
+}
+
 void UMoverComponent::SetFrameStateFromContext(const FMoverSyncState* SyncState, const FMoverAuxStateContext* AuxState, bool bRebaseBasedState)
 {
 
diff --git a/Engine/Plugins/Experimental/Mover/Source/Mover/Public/LayeredMove.h b/Engine/Plugins/Experimental/Mover/Source/Mover/Public/LayeredMove.h
--- a/Engine/Plugins/Experimental/Mover/Source/Mover/Public/LayeredMove.h
+++ b/Engine/Plugins/Experimental/Mover/Source/Mover/Public/LayeredMove.h
@@ -2,6 +2,7 @@
 
 #pragma once
 
+#include "GameplayTagContainer.h"
 #include "MoveLibrary/MovementUtilsTypes.h"
 #include "LayeredMove.generated.h"
 
@@ -91,6 +92,19 @@ struct MOVER_API FLayeredMoveBase
 	// Settings related to velocity applied to the actor after a layered move has finished
 	UPROPERTY(EditAnywhere, BlueprintReadWrite, Category = Mover)
 	FLayeredMoveFinishVelocitySettings FinishVelocitySettings;
+
+	/**
+	 * Check Layered Move for a gameplay tag.
+	 *
+	 * @param TagToFind			Tag to check on the Mover systems
+	 * @param bExactMatch		If true, the tag has to be exactly present, if false then TagToFind will include it's parent tags while matching
+	 * 
+	 * @return True if the TagToFind was found
+	 */
+	virtual bool HasGameplayTag(FGameplayTag TagToFind, bool bExactMatch) const
+	{
+		return false;
+	}
 	
 	// Kicks off this move, allowing any initialization to occur.
 	void StartMove(const UMoverComponent* MoverComp, UMoverBlackboard* SimBlackboard, float CurrentSimTimeMs);
diff --git a/Engine/Plugins/Experimental/Mover/Source/Mover/Public/MovementMode.h b/Engine/Plugins/Experimental/Mover/Source/Mover/Public/MovementMode.h
--- a/Engine/Plugins/Experimental/Mover/Source/Mover/Public/MovementMode.h
+++ b/Engine/Plugins/Experimental/Mover/Source/Mover/Public/MovementMode.h
@@ -74,6 +74,16 @@ public:
 	UPROPERTY(EditDefaultsOnly, Category = Mover)
 	FGameplayTagContainer GameplayTags;
 
+	/**
+   	 * Check Movement Mode for a gameplay tag.
+   	 *
+   	 * @param TagToFind			Tag to check on the Mover systems
+   	 * @param bExactMatch		If true, the tag has to be exactly present, if false then TagToFind will include it's parent tags while matching
+   	 * 
+   	 * @return True if the TagToFind was found
+   	 */
+	virtual bool HasGameplayTag(FGameplayTag TagToFind, bool bExactMatch) const;
+	
 protected:
 
 	virtual void OnRegistered(const FName ModeName);
diff --git a/Engine/Plugins/Experimental/Mover/Source/Mover/Public/MovementModifier.h b/Engine/Plugins/Experimental/Mover/Source/Mover/Public/MovementModifier.h
--- a/Engine/Plugins/Experimental/Mover/Source/Mover/Public/MovementModifier.h
+++ b/Engine/Plugins/Experimental/Mover/Source/Mover/Public/MovementModifier.h
@@ -2,6 +2,7 @@
 
 #pragma once
 
+#include "GameplayTagContainer.h"
 #include "MoveLibrary/MovementUtilsTypes.h"
 #include "MovementModifier.generated.h"
 
@@ -141,7 +142,20 @@ struct MOVER_API FMovementModifierBase
 	 * This is done to avoid queueing a modifier again as a local client applies potential input
 	 */
 	void SetHandleFromExistingModifier(const FMovementModifierHandle& OtherModifierHandle);
-	
+
+	/**
+  	 * Check modifier for a gameplay tag.
+  	 *
+  	 * @param TagToFind			Tag to check on the Mover systems
+  	 * @param bExactMatch		If true, the tag has to be exactly present, if false then TagToFind will include it's parent tags while matching
+  	 * 
+  	 * @return True if the TagToFind was found
+  	 */
+	virtual bool HasGameplayTag(FGameplayTag TagToFind, bool bExactMatch) const
+	{
+		return false;
+	}
+
 protected:
 	/**
 	 * Modifier handle local to this client or server instance. Used to cancel or query for a active/queued modifier.
diff --git a/Engine/Plugins/Experimental/Mover/Source/Mover/Public/MoverComponent.h b/Engine/Plugins/Experimental/Mover/Source/Mover/Public/MoverComponent.h
--- a/Engine/Plugins/Experimental/Mover/Source/Mover/Public/MoverComponent.h
+++ b/Engine/Plugins/Experimental/Mover/Source/Mover/Public/MoverComponent.h
@@ -410,6 +410,17 @@ public:	// Queries
 	
 	// Find movement modifier by it's handle. Returns nullptr if the modifier couldn't be found
 	const FMovementModifierBase* FindMovementModifier(const FMovementModifierHandle& ModifierHandle) const;
+
+	/**
+ 	 * Check Mover systems for a gameplay tag.
+ 	 *
+ 	 * @param TagToFind			Tag to check on the Mover systems
+ 	 * @param bExactMatch		If true, the tag has to be exactly present, if false then TagToFind will include it's parent tags while matching
+ 	 * 
+ 	 * @return True if the TagToFind was found
+ 	 */
+	UFUNCTION(BlueprintPure, Category = Mover, meta = (Keywords = "HasTag"))
+	bool HasGameplayTag(FGameplayTag TagToFind, bool bExactMatch) const;
 	
 protected:
 
