Mover 2.0 fixing bug with instant effects not applying movement mode change if applied in mover systems

diff --git a/Engine/Plugins/Experimental/Mover/Source/Mover/Private/MovementModeStateMachine.cpp b/Engine/Plugins/Experimental/Mover/Source/Mover/Private/MovementModeStateMachine.cpp
--- a/Engine/Plugins/Experimental/Mover/Source/Mover/Private/MovementModeStateMachine.cpp
+++ b/Engine/Plugins/Experimental/Mover/Source/Mover/Private/MovementModeStateMachine.cpp
@@ -356,7 +356,13 @@ void UMovementModeStateMachine::OnSimulationTick(USceneComponent* UpdatedCompone
 	EffectParams.UpdatedPrimitive = UpdatedPrimitive;
 	
 	// Apply any instant effects that were queued up during this tick and didn't get handled in a substep
-	ApplyInstantEffects(EffectParams, OutputState.SyncState);
+	if (ApplyInstantEffects(EffectParams, OutputState.SyncState))
+	{
+		if (CurrentModeName != OutputState.SyncState.MovementMode)
+		{
+			SetModeImmediately(OutputState.SyncState.MovementMode);
+		}
+	}
 }
 
 void UMovementModeStateMachine::OnSimulationPreRollback(const FMoverSyncState* InvalidSyncState, const FMoverSyncState* SyncState, const FMoverAuxStateContext* InvalidAuxState, const FMoverAuxStateContext* AuxState)
