Mover 2.0 fixing issue with modifiers getting queued again even if they existed in the active modifier array already

diff --git a/Engine/Plugins/Experimental/Mover/Source/Mover/Private/MovementModifier.cpp b/Engine/Plugins/Experimental/Mover/Source/Mover/Private/MovementModifier.cpp
--- a/Engine/Plugins/Experimental/Mover/Source/Mover/Private/MovementModifier.cpp
+++ b/Engine/Plugins/Experimental/Mover/Source/Mover/Private/MovementModifier.cpp
@@ -84,11 +84,11 @@ void FMovementModifierBase::GenerateHandle()
 	LocalModifierHandle.GenerateHandle();
 }
 
-void FMovementModifierBase::SetHandleFromExistingModifier(const FMovementModifierHandle& OtherModifierHandle)
+void FMovementModifierBase::OverwriteHandleIfInvalid(const FMovementModifierHandle& ValidModifierHandle)
 {
-	if (OtherModifierHandle.IsValid() && !LocalModifierHandle.IsValid())
+	if (ValidModifierHandle.IsValid() && !LocalModifierHandle.IsValid())
 	{
-		LocalModifierHandle = OtherModifierHandle;
+		LocalModifierHandle = ValidModifierHandle;
 	}
 }
 
@@ -288,9 +288,9 @@ void FMovementModifierGroup::FlushModifierArrays(UMoverComponent* MoverComp, con
 			QueuedModifier->StartSimTimeMs = TimeStep.BaseSimTimeMs;
 
 			// We only want to queue this queued modifier if it wasn't already added from state received from authority. If we already have the modifier just assign it a handle since it's already been activated.
-			if (QueuedModifier->Matches(ActiveModifier.Get()) && !ActiveModifier->GetHandle().IsValid())
+			if (QueuedModifier->Matches(ActiveModifier.Get()))
 			{
-				ActiveModifier->SetHandleFromExistingModifier(QueuedModifier->GetHandle());
+				ActiveModifier->OverwriteHandleIfInvalid(QueuedModifier->GetHandle());
 				bModidiferAlreadyActive = true;
 				break;
 			}
diff --git a/Engine/Plugins/Experimental/Mover/Source/Mover/Public/MovementModifier.h b/Engine/Plugins/Experimental/Mover/Source/Mover/Public/MovementModifier.h
--- a/Engine/Plugins/Experimental/Mover/Source/Mover/Public/MovementModifier.h
+++ b/Engine/Plugins/Experimental/Mover/Source/Mover/Public/MovementModifier.h
@@ -141,7 +141,7 @@ struct MOVER_API FMovementModifierBase
 	 * Currently used to write a local handle to modifiers that were added from a rollback so they don't have a handle yet
 	 * This is done to avoid queueing a modifier again as a local client applies potential input
 	 */
-	void SetHandleFromExistingModifier(const FMovementModifierHandle& OtherModifierHandle);
+	void OverwriteHandleIfInvalid(const FMovementModifierHandle& ValidModifierHandle);
 
 	/**
   	 * Check modifier for a gameplay tag.
