ROV support for AI Mover Characters
- ROV for high LOD sim-proxy replication

diff --git a/Engine/Plugins/Experimental/Mover/Source/Mover/Private/Backends/MoverNetworkPhysicsLiaison.cpp b/Engine/Plugins/Experimental/Mover/Source/Mover/Private/Backends/MoverNetworkPhysicsLiaison.cpp
--- a/Engine/Plugins/Experimental/Mover/Source/Mover/Private/Backends/MoverNetworkPhysicsLiaison.cpp
+++ b/Engine/Plugins/Experimental/Mover/Source/Mover/Private/Backends/MoverNetworkPhysicsLiaison.cpp
@@ -935,8 +935,15 @@ void UMoverNetworkPhysicsLiaisonComponent::OnPreSimulate_Internal(const FPhysics
 		}
 	}
 
+	// Add AI Move if it exists
+	FVector AIMoveVelocity = FVector::ZeroVector;
+	if (const FMoverAIInputs* MoverAIInputs = Input.InputCmd.InputCollection.FindDataByType<FMoverAIInputs>())
+	{
+		AIMoveVelocity = MoverAIInputs->RVOVelocityDelta;
+	}
+
 	FMoverDefaultSyncState& SyncState = Input.SyncState.SyncStateCollection.FindOrAddMutableDataByType<FMoverDefaultSyncState>();
-	SyncState.SetTransforms_WorldSpace(CharacterParticle->GetX(), FRotator(CharacterParticle->GetR()), CharacterParticle->GetV() - LocalGroundVelocity);
+	SyncState.SetTransforms_WorldSpace(CharacterParticle->GetX(), FRotator(CharacterParticle->GetR()), CharacterParticle->GetV() - LocalGroundVelocity + AIMoveVelocity);
 
 	///////////////////////////////////////////////////////////////////////////////////////////////////////////////////////////
 	// Update the simulation
diff --git a/Engine/Plugins/Experimental/Mover/Source/Mover/Private/PhysicsMover/PhysicsMoverSimulationTypes.cpp b/Engine/Plugins/Experimental/Mover/Source/Mover/Private/PhysicsMover/PhysicsMoverSimulationTypes.cpp
--- a/Engine/Plugins/Experimental/Mover/Source/Mover/Private/PhysicsMover/PhysicsMoverSimulationTypes.cpp
+++ b/Engine/Plugins/Experimental/Mover/Source/Mover/Private/PhysicsMover/PhysicsMoverSimulationTypes.cpp
@@ -47,4 +47,31 @@ void FMovementSettingsInputs::ToString(FAnsiStringBuilderBase& Out) const
 
 	Out.Appendf("MaxSpeed=%.2f\n", MaxSpeed);
 	Out.Appendf("Acceleration=%.2f\n", Acceleration);
+}
+
+
+//////////////////////////////////////////////////////////////////////////
+// FMoverAIInputs
+
+FMoverDataStructBase* FMoverAIInputs::Clone() const
+{
+	FMoverAIInputs* CopyPtr = new FMoverAIInputs(*this);
+	return CopyPtr;
+}
+
+bool FMoverAIInputs::NetSerialize(FArchive& Ar, UPackageMap* Map, bool& bOutSuccess)
+{
+	Super::NetSerialize(Ar, Map, bOutSuccess);
+
+	Ar << RVOVelocityDelta;
+
+	bOutSuccess = true;
+	return true;
+}
+
+void FMoverAIInputs::ToString(FAnsiStringBuilderBase& Out) const
+{
+	Super::ToString(Out);
+
+	Out.Appendf("ROVVelDelta: X=%.2f Y=%.2f Z=%.2f\n", RVOVelocityDelta.X, RVOVelocityDelta.Y, RVOVelocityDelta.Z);
 }
\ No newline at end of file
diff --git a/Engine/Plugins/Experimental/Mover/Source/Mover/Public/PhysicsMover/PhysicsMoverSimulationTypes.h b/Engine/Plugins/Experimental/Mover/Source/Mover/Public/PhysicsMover/PhysicsMoverSimulationTypes.h
--- a/Engine/Plugins/Experimental/Mover/Source/Mover/Public/PhysicsMover/PhysicsMoverSimulationTypes.h
+++ b/Engine/Plugins/Experimental/Mover/Source/Mover/Public/PhysicsMover/PhysicsMoverSimulationTypes.h
@@ -99,7 +99,7 @@ public:
 };
 
 //////////////////////////////////////////////////////////////////////////
-// FMovementSettingsInput
+// FMovementSettingsInputs
 
 // Data block containing movement settings inputs that are networked from client to server.
 // This is useful if settings changes need to be predicted on the client and synced on the server.
@@ -140,6 +140,51 @@ public:
 
 template<>
 struct TStructOpsTypeTraits< FMovementSettingsInputs > : public TStructOpsTypeTraitsBase2< FMovementSettingsInputs >
+{
+	enum
+	{
+		WithNetSerializer = true,
+		WithCopy = true
+	};
+};
+
+
+//////////////////////////////////////////////////////////////////////////
+// FMoverAIInputs
+
+// Data block containing ROV Velocity that is networked from server to clients.
+// Also supports rewind/resimulation of the data.
+USTRUCT(BlueprintType)
+struct MOVER_API FMoverAIInputs : public FMoverDataStructBase
+{
+	GENERATED_USTRUCT_BODY()
+
+public:
+	// ROV Velocity calculated on the Server
+	UPROPERTY(BlueprintReadWrite, Category = Mover)
+	FVector RVOVelocityDelta;
+
+	FMoverAIInputs()
+		: RVOVelocityDelta(ForceInitToZero)
+	{
+	}
+
+	virtual ~FMoverAIInputs() {}
+
+	// @return newly allocated copy of this FMoverAIInputs. Must be overridden by child classes
+	virtual FMoverDataStructBase* Clone() const override;
+
+	virtual bool NetSerialize(FArchive& Ar, UPackageMap* Map, bool& bOutSuccess) override;
+
+	virtual UScriptStruct* GetScriptStruct() const override { return StaticStruct(); }
+
+	virtual void ToString(FAnsiStringBuilderBase& Out) const override;
+
+	virtual void AddReferencedObjects(FReferenceCollector& Collector) override { Super::AddReferencedObjects(Collector); }
+};
+
+template<>
+struct TStructOpsTypeTraits< FMoverAIInputs > : public TStructOpsTypeTraitsBase2< FMoverAIInputs >
 {
 	enum
 	{
