Mover: improved warning messaging. No other functional changes.

diff --git a/Engine/Plugins/Experimental/Mover/Source/Mover/Private/MoverComponent.cpp b/Engine/Plugins/Experimental/Mover/Source/Mover/Private/MoverComponent.cpp
--- a/Engine/Plugins/Experimental/Mover/Source/Mover/Private/MoverComponent.cpp
+++ b/Engine/Plugins/Experimental/Mover/Source/Mover/Private/MoverComponent.cpp
@@ -199,7 +199,7 @@ void UMoverComponent::BeginPlay()
 			SetPrimaryVisualComponent(MyActor->FindComponentByClass<UMeshComponent>());
 		}
 
-		ensureMsgf(UpdatedComponent && (PrimaryVisualComponent != UpdatedComponent), TEXT("A Mover actor must have an UpdatedComponent and cannot have a PrimaryVisualComponent that is the same as UpdatedComponent"));
+		ensureMsgf(UpdatedComponent && (PrimaryVisualComponent != UpdatedComponent), TEXT("A Mover actor (%s) must have an UpdatedComponent and cannot have a PrimaryVisualComponent that is the same as UpdatedComponent"), *GetNameSafe(MyActor));
 
 		// Optional motion warping support
 		if (UMotionWarpingComponent* WarpingComp = MyActor->FindComponentByClass<UMotionWarpingComponent>())
