Mover: adding layered move that supports root motion from a mesh's custom attribute. This is a first pass that only supports independent ticking.

diff --git a/Engine/Plugins/Experimental/Mover/Source/Mover/Private/DefaultMovementSet/LayeredMoves/RootMotionAttributeLayeredMove.cpp b/Engine/Plugins/Experimental/Mover/Source/Mover/Private/DefaultMovementSet/LayeredMoves/RootMotionAttributeLayeredMove.cpp
new file mode 100644
--- /dev/null
+++ b/Engine/Plugins/Experimental/Mover/Source/Mover/Private/DefaultMovementSet/LayeredMoves/RootMotionAttributeLayeredMove.cpp
@@ -0,0 +1,167 @@
+// Copyright Epic Games, Inc. All Rights Reserved.
+
+#include "DefaultMovementSet/LayeredMoves/RootMotionAttributeLayeredMove.h"
+#include "MoverComponent.h"
+#include "Animation/AttributesContainer.h"
+#include "Animation/BuiltInAttributeTypes.h"
+#include "Animation/AnimMontage.h"
+#include "Animation/AnimInstance.h"
+#include "Components/SkeletalMeshComponent.h"
+#include "MoveLibrary/MovementUtils.h"
+#include "MoverTypes.h"
+#include "MoverLog.h"
+
+
+#if !UE_BUILD_SHIPPING
+FAutoConsoleVariable CVarLogRootMotionAttrSteps(
+	TEXT("mover.debug.LogRootMotionAttrSteps"),
+	false,
+	TEXT("Whether to log detailed information about root motion attribute layered moves. 0: Disable, 1: Enable"),
+	ECVF_Cheat);
+#endif	// !UE_BUILD_SHIPPING
+
+FLayeredMove_RootMotionAttribute::FLayeredMove_RootMotionAttribute()
+{
+	DurationMs = -1.f;
+	MixMode = EMoveMixMode::OverrideAll;
+}
+
+bool FLayeredMove_RootMotionAttribute::GenerateMove(const FMoverTickStartData& SimState, const FMoverTimeStep& TimeStep, const UMoverComponent* MoverComp, UMoverBlackboard* SimBlackboard, FProposedMove& OutProposedMove)
+{
+	const float DeltaSeconds = TimeStep.StepMs / 1000.f;
+
+	bool bDidAttrHaveRootMotion = false;
+	FTransform LocalRootMotion;
+	bool bHasValidWarpingContext = false;
+	FMotionWarpingUpdateContext WarpingContext;
+
+	if (!TimeStep.bIsResimulating)
+	{
+		// Clear resim values
+		bDidAttrHaveRootMotionForResim = false;
+		LocalRootMotionForResim = FTransform::Identity;
+		WarpingContextForResim = FMotionWarpingUpdateContext();
+
+		// Sample root motion from the mesh attribute
+		const USkeletalMeshComponent* Mesh = Cast<const USkeletalMeshComponent>(MoverComp->GetPrimaryVisualComponent());
+
+		if (Mesh)
+		{
+			// TODO: support options for different interpretations, such as velocity
+			static const FName RootMotionAttributeName = "RootMotionDelta";
+			static const UE::Anim::FAttributeId RootMotionAttributeId = { RootMotionAttributeName, FCompactPoseBoneIndex(0) };
+
+			if (const FTransformAnimationAttribute* RootMotionAttribute = Mesh->GetCustomAttributes().Find<FTransformAnimationAttribute>(RootMotionAttributeId))
+			{
+				// NOTE this will only work for ticking modes that tick in time with the world tick, because it relies on the Mesh ticking at the same rate as the movement simulation. 
+				// For fixed-tick modes, the attribute would be better as an accumulator for the movement sim to consume, along with a time accumulation attribute.
+				LocalRootMotion = RootMotionAttribute->Value;
+				bDidAttrHaveRootMotion = true;
+			}
+		}
+
+		static const bool bShouldWarpFromMontage = true;
+
+		if (bDidAttrHaveRootMotion && bShouldWarpFromMontage)
+		{
+			// Not resimulating: we are following along with any root motion montages
+			const FAnimMontageInstance* RootMotionMontageInstance = Mesh && Mesh->GetAnimInstance() ? Mesh->GetAnimInstance()->GetRootMotionMontageInstance() : nullptr;
+
+			if (RootMotionMontageInstance)
+			{
+				const UAnimMontage* Montage = RootMotionMontageInstance->Montage;
+
+				WarpingContext.DeltaSeconds = DeltaSeconds;
+				WarpingContext.Animation = Montage;
+				WarpingContext.CurrentPosition = RootMotionMontageInstance->GetPosition();
+				WarpingContext.PreviousPosition = RootMotionMontageInstance->GetPreviousPosition();
+				WarpingContext.Weight = RootMotionMontageInstance->GetWeight();
+				WarpingContext.PlayRate = RootMotionMontageInstance->Montage->RateScale * RootMotionMontageInstance->GetPlayRate();
+
+				bHasValidWarpingContext = true;
+			}
+		}
+
+		// Save values for resim
+		{
+			bDidAttrHaveRootMotionForResim = bDidAttrHaveRootMotion;
+			LocalRootMotionForResim = LocalRootMotion;
+
+			if (bHasValidWarpingContext)
+			{
+				WarpingContextForResim = WarpingContext;
+			}
+		}
+
+	}
+	else   // resimulating...
+	{
+		// restore the cached transform and warping parameters (if set)
+		bDidAttrHaveRootMotion	= bDidAttrHaveRootMotionForResim;
+		LocalRootMotion			= LocalRootMotionForResim;
+		bHasValidWarpingContext = WarpingContextForResim.Animation != nullptr;
+		WarpingContext			= WarpingContextForResim;
+	}
+
+	
+	if (bDidAttrHaveRootMotion)
+	{
+		const FMoverDefaultSyncState* SyncState = SimState.SyncState.SyncStateCollection.FindDataByType<FMoverDefaultSyncState>();
+		const FCharacterDefaultInputs* InputCmd = SimState.InputCmd.InputCollection.FindDataByType<FCharacterDefaultInputs>();
+
+		// TODO: Provide better options to control when we should skip proposed movement
+		if (!bHasValidWarpingContext &&
+				(InputCmd->bIsJumpJustPressed || MoverComp->HasGameplayTag(Mover_IsFalling, false)))
+		{
+			return false;	// do not perform root motion while jumping or falling, so that we can have air control, unless we're under montage influence
+		}
+
+		// Note that we're forcing the use of the sync state's actor transform data. This is necessary when the movement simulation 
+		// is running ahead of the actor's visual representation and may be rotated differently, such as in an async physics sim.
+		const FTransform SimActorTransform = FTransform(SyncState->GetOrientation_WorldSpace().Quaternion(), SyncState->GetLocation_WorldSpace());
+		const FTransform WorldSpaceRootMotion = MoverComp->ConvertLocalRootMotionToWorld(LocalRootMotion, DeltaSeconds, &SimActorTransform, (bHasValidWarpingContext ? &WarpingContext : nullptr));
+
+		OutProposedMove = FProposedMove();
+		OutProposedMove.MixMode = MixMode;
+
+		// Convert the transform into linear and angular velocities
+		OutProposedMove.LinearVelocity  = WorldSpaceRootMotion.GetTranslation() * (1.f / DeltaSeconds);
+		OutProposedMove.AngularVelocity = WorldSpaceRootMotion.GetRotation().Rotator() * (1.f / DeltaSeconds);
+
+#if !UE_BUILD_SHIPPING
+		UE_CLOG(CVarLogRootMotionAttrSteps->GetBool(), LogMover, Log, TEXT("RootMotionAttr. SimF %i (dt %.3f) => LocalT: %s (WST: %s)  Vel: %.3f"),
+			TimeStep.ServerFrame, DeltaSeconds,
+			*LocalRootMotion.GetTranslation().ToString(), *WorldSpaceRootMotion.GetTranslation().ToString(), OutProposedMove.LinearVelocity.Length());
+#endif // !UE_BUILD_SHIPPING
+
+		return true;
+	}
+
+	return false;
+}
+
+FLayeredMoveBase* FLayeredMove_RootMotionAttribute::Clone() const
+{
+	FLayeredMove_RootMotionAttribute* CopyPtr = new FLayeredMove_RootMotionAttribute(*this);
+	return CopyPtr;
+}
+
+void FLayeredMove_RootMotionAttribute::NetSerialize(FArchive& Ar)
+{
+	Super::NetSerialize(Ar);
+}
+
+UScriptStruct* FLayeredMove_RootMotionAttribute::GetScriptStruct() const
+{
+	return FLayeredMove_RootMotionAttribute::StaticStruct();
+}
+
+FString FLayeredMove_RootMotionAttribute::ToSimpleString() const
+{
+	return FString::Printf(TEXT("RootMotionAttribute"));
+}
+
+void FLayeredMove_RootMotionAttribute::AddReferencedObjects(class FReferenceCollector& Collector)
+{
+	Super::AddReferencedObjects(Collector);
+}
diff --git a/Engine/Plugins/Experimental/Mover/Source/Mover/Public/DefaultMovementSet/LayeredMoves/RootMotionAttributeLayeredMove.h b/Engine/Plugins/Experimental/Mover/Source/Mover/Public/DefaultMovementSet/LayeredMoves/RootMotionAttributeLayeredMove.h
new file mode 100644
--- /dev/null
+++ b/Engine/Plugins/Experimental/Mover/Source/Mover/Public/DefaultMovementSet/LayeredMoves/RootMotionAttributeLayeredMove.h
@@ -0,0 +1,39 @@
+// Copyright Epic Games, Inc. All Rights Reserved.
+
+#pragma once
+
+#include "LayeredMove.h"
+#include "RootMotionModifier.h"
+#include "RootMotionAttributeLayeredMove.generated.h"
+
+/** 
+ * Root Motion Attribute Move: handles root motion from a mesh's custom attribute. 
+ * Currently only supports Independent ticking mode, and allows air control while jumping/falling.
+ */
+USTRUCT(BlueprintType)
+struct MOVER_API FLayeredMove_RootMotionAttribute : public FLayeredMoveBase
+{
+	GENERATED_USTRUCT_BODY()
+
+	FLayeredMove_RootMotionAttribute();
+	virtual ~FLayeredMove_RootMotionAttribute() {}
+
+protected:
+	// These member variables are NOT replicated. They are used if we rollback and resimulate when the root motion attribute is no longer in sync.
+	bool bDidAttrHaveRootMotionForResim = false;
+	FTransform LocalRootMotionForResim;
+	FMotionWarpingUpdateContext WarpingContextForResim;
+
+	// Generate a movement 
+	virtual bool GenerateMove(const FMoverTickStartData& StartState, const FMoverTimeStep& TimeStep, const UMoverComponent* MoverComp, UMoverBlackboard* SimBlackboard, FProposedMove& OutProposedMove) override;
+
+	virtual FLayeredMoveBase* Clone() const override;
+
+	virtual void NetSerialize(FArchive& Ar) override;
+
+	virtual UScriptStruct* GetScriptStruct() const override;
+
+	virtual FString ToSimpleString() const override;
+
+	virtual void AddReferencedObjects(class FReferenceCollector& Collector) override;
+};
