Mover: adding a stance movement modifier meant to handle different stance changes such as crouch or prone - currently only crouch is implemented

#virtualized

diff --git a/Engine/Plugins/Experimental/Mover/Source/Mover/Private/DefaultMovementSet/CharacterMoverComponent.cpp b/Engine/Plugins/Experimental/Mover/Source/Mover/Private/DefaultMovementSet/CharacterMoverComponent.cpp
--- a/Engine/Plugins/Experimental/Mover/Source/Mover/Private/DefaultMovementSet/CharacterMoverComponent.cpp
+++ b/Engine/Plugins/Experimental/Mover/Source/Mover/Private/DefaultMovementSet/CharacterMoverComponent.cpp
@@ -8,9 +8,11 @@
 #include "DefaultMovementSet/Modes/WalkingMode.h"
 #include "DefaultMovementSet/Settings/CommonLegacyMovementSettings.h"
 #include "MoveLibrary/FloorQueryUtils.h"
+#include "MoveLibrary/MovementUtils.h"
 
 UCharacterMoverComponent::UCharacterMoverComponent()
 	: bHandleJump(true)
+	, bWantsToCrouch(false)
 {
 	// Default movement modes
 	MovementModes.Add(DefaultModeNames::Walking, CreateDefaultSubobject<UWalkingMode>(TEXT("DefaultWalkingMode")));
@@ -29,57 +31,32 @@ void UCharacterMoverComponent::BeginPlay()
 
 bool UCharacterMoverComponent::IsCrouching() const
 {
-	return false; // Crouching not yet supported by Mover
+	return HasGameplayTag(Mover_IsCrouching, true);
 }
 
 bool UCharacterMoverComponent::IsFlying() const
 {
-	if (const UBaseMovementMode* Mode = GetMovementMode())
-	{
-		return Mode->GameplayTags.HasTag(Mover_IsFlying);
-	}
-
-	return false;
+	return HasGameplayTag(Mover_IsFlying, true);
 }
 
 bool UCharacterMoverComponent::IsFalling() const
 {
-	if (const UBaseMovementMode* Mode = GetMovementMode())
-	{
-		return Mode->GameplayTags.HasTag(Mover_IsFalling);
-	}
-
-	return false;
+	return HasGameplayTag(Mover_IsFalling, true);
 }
 
 bool UCharacterMoverComponent::IsAirborne() const
 {
-	if (const UBaseMovementMode* Mode = GetMovementMode())
-	{
-		return Mode->GameplayTags.HasTag(Mover_IsInAir);
-	}
-
-	return false;
+	return HasGameplayTag(Mover_IsInAir, true);
 }
 
 bool UCharacterMoverComponent::IsOnGround() const
 {
-	if (const UBaseMovementMode* Mode = GetMovementMode())
-	{
-		return Mode->GameplayTags.HasTag(Mover_IsOnGround);
-	}
-
-	return false;
+	return HasGameplayTag(Mover_IsOnGround, true);
 }
 
 bool UCharacterMoverComponent::IsSwimming() const
 {
-	if (const UBaseMovementMode* Mode = GetMovementMode())
-	{
-		return Mode->GameplayTags.HasTag(Mover_IsSwimming);
-	}
-
-	return false;
+	return HasGameplayTag(Mover_IsSwimming, true);
 }
 
 bool UCharacterMoverComponent::IsSlopeSliding() const
@@ -117,15 +94,75 @@ bool UCharacterMoverComponent::Jump()
 	return false;
 }
 
+bool UCharacterMoverComponent::CanCrouch()
+{
+	return true;
+}
+
+void UCharacterMoverComponent::Crouch()
+{
+	if (CanCrouch())
+	{
+		bWantsToCrouch = true;
+	}
+}
+
+void UCharacterMoverComponent::UnCrouch()
+{
+	bWantsToCrouch = false;
+}
+
 void UCharacterMoverComponent::OnMoverPreSimulationTick(const FMoverTimeStep& TimeStep, const FMoverInputCmdContext& InputCmd)
 {
 	if (bHandleJump)
 	{
 		const FCharacterDefaultInputs* CharacterInputs = InputCmd.InputCollection.FindDataByType<FCharacterDefaultInputs>();
-	
 		if (CharacterInputs && CharacterInputs->bIsJumpJustPressed && CanActorJump())
 		{
 			Jump();
 		}
 	}
+	
+	const FStanceModifier* StanceModifier = static_cast<const FStanceModifier*>(FindMovementModifier(StanceModifierHandle));
+	// This is a fail safe in case our handle was bad - try finding the modifier by type if we can
+	if (!StanceModifier)
+	{
+		StanceModifier = FindMovementModifierByType<FStanceModifier>();
+	}
+	
+	EStanceMode OldActiveStance = EStanceMode::Invalid;
+	if (StanceModifier)
+	{
+		OldActiveStance = StanceModifier->ActiveStance;
+	}
+	
+	const bool bIsCrouching = HasGameplayTag(Mover_IsCrouching, true);
+	if (bIsCrouching && (!bWantsToCrouch || !CanCrouch()))
+	{	
+		if (StanceModifier && StanceModifier->CanExpand(this))
+		{
+			CancelModifierFromHandle(StanceModifier->GetHandle());
+			StanceModifierHandle.Invalidate();
+
+			StanceModifier = nullptr;
+		}
+	}
+	else if (!bIsCrouching && bWantsToCrouch && CanCrouch())
+	{
+		TSharedPtr<FStanceModifier> NewStanceModifier = MakeShared<FStanceModifier>();
+		StanceModifierHandle = QueueMovementModifier(NewStanceModifier);
+
+		StanceModifier = NewStanceModifier.Get();
+	}
+	
+	EStanceMode NewActiveStance = EStanceMode::Invalid;
+	if (StanceModifier)
+	{
+		NewActiveStance = StanceModifier->ActiveStance;
+	}
+
+	if (OldActiveStance != NewActiveStance)
+	{
+		OnStanceChanged.Broadcast(OldActiveStance, NewActiveStance);
+	}
 }
diff --git a/Engine/Plugins/Experimental/Mover/Source/Mover/Private/DefaultMovementSet/MovementModifiers/StanceModifier.cpp b/Engine/Plugins/Experimental/Mover/Source/Mover/Private/DefaultMovementSet/MovementModifiers/StanceModifier.cpp
new file mode 100644
--- /dev/null
+++ b/Engine/Plugins/Experimental/Mover/Source/Mover/Private/DefaultMovementSet/MovementModifiers/StanceModifier.cpp
@@ -0,0 +1,306 @@
+﻿// Copyright Epic Games, Inc. All Rights Reserved.
+
+#include "DefaultMovementSet/MovementModifiers/StanceModifier.h"
+
+#include "MoverComponent.h"
+#include "Components/CapsuleComponent.h"
+#include "DefaultMovementSet/CharacterMoverComponent.h"
+#include "DefaultMovementSet/Settings/CommonLegacyMovementSettings.h"
+#include "DefaultMovementSet/InstantMovementEffects/BasicInstantMovementEffects.h"
+#include "DefaultMovementSet/Settings/StanceSettings.h"
+#include "Engine/Blueprint.h"
+#include "Engine/SCS_Node.h"
+#include "Engine/SimpleConstructionScript.h"
+#include "GameFramework/Pawn.h"
+#include "MoveLibrary/MovementUtils.h"
+
+FStanceModifier::FStanceModifier()
+	: ActiveStance(EStanceMode::Crouch)
+{
+	DurationMs = -1.0f;
+}
+
+bool FStanceModifier::HasGameplayTag(FGameplayTag TagToFind, bool bExactMatch) const
+{
+	// TODO: Eventually check for other stance tags here like prone
+	if (bExactMatch)
+	{
+		return TagToFind.MatchesTagExact(Mover_IsCrouching);
+	}
+
+	return TagToFind.MatchesTag(Mover_IsCrouching);
+}
+
+void FStanceModifier::OnStart(UMoverComponent* MoverComp, const FMoverTimeStep& TimeStep, const FMoverSyncState& SyncState, const FMoverAuxStateContext& AuxState)
+{
+	if (UStanceSettings* StanceSettings = MoverComp->FindSharedSettings_Mutable<UStanceSettings>())
+	{
+		if (const UCapsuleComponent* CapsuleComponent = Cast<UCapsuleComponent>(MoverComp->GetUpdatedComponent()))
+		{
+			float OldHalfHeight = CapsuleComponent->GetScaledCapsuleHalfHeight();
+			float NewHalfHeight = 0;
+			float NewEyeHeight = 0;
+
+			switch (ActiveStance)
+			{
+			default:
+			case EStanceMode::Crouch:
+				NewHalfHeight = StanceSettings->CrouchHalfHeight;
+				NewEyeHeight = StanceSettings->CrouchedEyeHeight;
+				break;
+
+			// Prone isn't currently implemented
+			case EStanceMode::Prone:
+				UE_LOG(LogMover, Warning, TEXT("Stance got into prone stance - That stance is not currently implemented."));
+				// TODO: returning here so we don't apply any bad state to actor in case prone was set. Eventually, the return should be removed once prone is implemented properly
+				DurationMs = 0;
+				return;
+			}
+			
+			AdjustCapsule(MoverComp, OldHalfHeight, NewHalfHeight, NewEyeHeight);
+			ApplyMovementSettings(MoverComp);
+		}
+	}
+}
+
+void FStanceModifier::OnEnd(UMoverComponent* MoverComp, const FMoverTimeStep& TimeStep, const FMoverSyncState& SyncState, const FMoverAuxStateContext& AuxState)
+{
+	const AActor* OwnerCDO = Cast<AActor>(MoverComp->GetOwner()->GetClass()->GetDefaultObject());
+
+	if (UCapsuleComponent* CapsuleComponent = Cast<UCapsuleComponent>(MoverComp->GetUpdatedComponent()))
+	{
+		if (const UCapsuleComponent* OriginalCapsule = GetOriginalCapsuleComponent(MoverComp))
+		{
+			if (const APawn* OwnerCDOAsPawn = Cast<APawn>(OwnerCDO))
+			{
+				AdjustCapsule(MoverComp, CapsuleComponent->GetScaledCapsuleHalfHeight(), OriginalCapsule->GetScaledCapsuleHalfHeight(), OwnerCDOAsPawn->BaseEyeHeight);
+				RevertMovementSettings(MoverComp);
+			}
+		}
+	}
+}
+
+void FStanceModifier::OnPreMovement(UMoverComponent* MoverComp, const FMoverTimeStep& TimeStep)
+{
+	// TODO: Check for different inputs/state here and manage swapping between stances - use AdjustCapsule and Apply/Revert movement settings.
+
+	// TODO: Prone isn't currently implemented - so we're just going to cancel the modifier if we got into that state
+	if (ActiveStance == EStanceMode::Prone)
+	{
+		UE_LOG(LogMover, Warning, TEXT("Stance got into prone stance - That stance is not currently implemented."));
+		DurationMs = 0;
+	}
+}
+
+void FStanceModifier::OnPostMovement(UMoverComponent* MoverComp, const FMoverTimeStep& TimeStep, const FMoverSyncState& SyncState, const FMoverAuxStateContext& AuxState)
+{
+	FMovementModifierBase::OnPostMovement(MoverComp, TimeStep, SyncState, AuxState);
+}
+
+FMovementModifierBase* FStanceModifier::Clone() const
+{
+	FStanceModifier* CopyPtr = new FStanceModifier(*this);
+	return CopyPtr;
+}
+
+void FStanceModifier::NetSerialize(FArchive& Ar)
+{
+	Super::NetSerialize(Ar);
+}
+
+UScriptStruct* FStanceModifier::GetScriptStruct() const
+{
+	return FStanceModifier::StaticStruct();
+}
+
+FString FStanceModifier::ToSimpleString() const
+{
+	return FString::Printf(TEXT("Stance Modifier"));
+}
+
+void FStanceModifier::AddReferencedObjects(FReferenceCollector& Collector)
+{
+	Super::AddReferencedObjects(Collector);
+}
+
+bool FStanceModifier::CanExpand(const UCharacterMoverComponent* MoverComp) const
+{
+	float StandingHalfHeight = 90;
+	float CurrentHalfHeight = 55;
+
+	USceneComponent* UpdatedComponent = MoverComp->GetUpdatedComponent();
+	UPrimitiveComponent* UpdatedCompAsPrimitive = Cast<UPrimitiveComponent>(UpdatedComponent);
+	
+	if (const UCapsuleComponent* OriginalCapsule = GetOriginalCapsuleComponent(MoverComp))
+	{
+		StandingHalfHeight = OriginalCapsule->GetScaledCapsuleHalfHeight();
+	}
+
+	if (UCapsuleComponent* CapsuleComponent = Cast<UCapsuleComponent>(UpdatedComponent))
+	{
+		CurrentHalfHeight = CapsuleComponent->GetScaledCapsuleHalfHeight();
+	}
+
+	float HalfHeightDifference = StandingHalfHeight - CurrentHalfHeight;
+	
+	// TODO: pluggable shapes
+	FCollisionQueryParams CapsuleParams(SCENE_QUERY_STAT(CrouchTrace), false, MoverComp->GetOwner());
+	FCollisionResponseParams ResponseParam;
+	UMovementUtils::InitCollisionParams(UpdatedCompAsPrimitive, CapsuleParams, ResponseParam);
+
+	const FMoverDefaultSyncState* SyncState = MoverComp->GetSyncState().SyncStateCollection.FindDataByType<FMoverDefaultSyncState>();
+	
+	FVector PawnLocation = SyncState->GetLocation_WorldSpace();
+	FQuat PawnRot = SyncState->GetOrientation_WorldSpace().Quaternion();
+	float PawnRadius = 0.0f;
+	float PawnHalfHeight = 0.0f;
+	UpdatedCompAsPrimitive->CalcBoundingCylinder(PawnRadius, PawnHalfHeight);
+	
+	// TODO: Compensate for the difference between current capsule size and standing size
+	FCollisionShape StandingCapsuleShape = FCollisionShape::MakeCapsule(PawnRadius, StandingHalfHeight);
+	const ECollisionChannel CollisionChannel = UpdatedCompAsPrimitive->GetCollisionObjectType();
+	bool bEncroached = true;
+
+	// TODO: We may need to expand this check to look at more than just the initial overlap - see CMC Uncrouch for details
+	if (!ShouldExpandingMaintainBase(MoverComp))
+	{
+		// Expand in place
+		bEncroached = UMovementUtils::OverlapTest(UpdatedComponent, UpdatedCompAsPrimitive, PawnLocation, PawnRot, CollisionChannel, StandingCapsuleShape, MoverComp->GetOwner());
+	}
+	else
+	{
+		// Expand while keeping base location the same.
+		FVector StandingLocation = PawnLocation + (HalfHeightDifference + .01f) * MoverComp->GetUpDirection();
+		bEncroached = UMovementUtils::OverlapTest(UpdatedComponent, UpdatedCompAsPrimitive, StandingLocation, PawnRot, CollisionChannel, StandingCapsuleShape, MoverComp->GetOwner());
+	}
+
+	return !bEncroached;
+}
+
+bool FStanceModifier::ShouldExpandingMaintainBase(const UCharacterMoverComponent* MoverComp) const
+{
+	if (MoverComp->IsOnGround())
+	{
+		return true;
+	}
+
+	return false;
+}
+
+const UCapsuleComponent* FStanceModifier::GetOriginalCapsuleComponent(const UMoverComponent* MoverComp)
+{
+	const UCapsuleComponent* OriginalCapsule = nullptr;
+
+	if (const AActor* OwnerCDO = Cast<AActor>(MoverComp->GetOwner()->GetClass()->GetDefaultObject()))
+	{
+		// Check if native CDO has Capsule component
+		OriginalCapsule = OwnerCDO->FindComponentByClass<UCapsuleComponent>();
+		
+		// check if it comes from a BP
+		if (!OriginalCapsule)
+		{
+			if (const UBlueprintGeneratedClass* OwnerClassAsBP = Cast<UBlueprintGeneratedClass>(OwnerCDO->GetClass()))
+			{
+				TArray<const UBlueprintGeneratedClass*> BlueprintClasses;
+				UBlueprintGeneratedClass::GetGeneratedClassesHierarchy(OwnerClassAsBP, BlueprintClasses);
+				for (const UBlueprintGeneratedClass* BlueprintClass : BlueprintClasses)
+				{
+					if (BlueprintClass->SimpleConstructionScript)
+					{
+						// Check Simple construction script
+						const TArray<USCS_Node*>& SCSNodes = BlueprintClass->SimpleConstructionScript->GetAllNodes();
+						for (USCS_Node* SCSNode : SCSNodes)
+						{
+							if (SCSNode)
+							{
+								if (const UCapsuleComponent* BPComponent = Cast<UCapsuleComponent>(SCSNode->ComponentTemplate))
+								{
+									OriginalCapsule = BPComponent;
+									break;
+								}
+							}
+						}
+					}
+				}
+			}
+		}
+	}
+
+	return OriginalCapsule;
+}
+
+void FStanceModifier::AdjustCapsule(UMoverComponent* MoverComp, float OldHalfHeight, float NewHalfHeight, float NewEyeHeight)
+{
+	const float HalfHeightDifference = FMath::Abs(NewHalfHeight - OldHalfHeight);
+	const bool bExpanding = OldHalfHeight < NewHalfHeight;
+	
+	// Set capsule size to crouching size
+	if (UCapsuleComponent* CapsuleComponent = Cast<UCapsuleComponent>(MoverComp->GetOwner()->FindComponentByClass(UCapsuleComponent::StaticClass())))
+	{
+		if (CapsuleComponent->GetUnscaledCapsuleHalfHeight() == NewHalfHeight)
+		{
+			return;
+		}
+		
+		CapsuleComponent->SetCapsuleSize(CapsuleComponent->GetUnscaledCapsuleRadius(), NewHalfHeight);
+	}
+
+	// update eye height on pawn
+	if (APawn* MoverCompOwnerAsPawn = Cast<APawn>(MoverComp->GetOwner()))
+	{
+		MoverCompOwnerAsPawn->BaseEyeHeight = NewEyeHeight;
+	}
+
+	const FVector CapsuleOffset = MoverComp->GetUpDirection() * (bExpanding ? HalfHeightDifference : -HalfHeightDifference);
+	const FVector VisualOffset = MoverComp->GetUpDirection() * (bExpanding ? -OldHalfHeight : NewHalfHeight);
+	
+	// Adjust location of capsule as setting it's size left it floating
+	if (!bExpanding || MoverComp->GetVelocity().Length() <= 0)
+	{
+		TSharedPtr<FTeleportEffect> TeleportEffect = MakeShared<FTeleportEffect>();
+		TeleportEffect->TargetLocation = MoverComp->GetUpdatedComponentTransform().GetLocation() + (CapsuleOffset);
+		MoverComp->QueueInstantMovementEffect(TeleportEffect);
+	}
+	
+	// Add offset to visual component as the base location has changed
+	MoverComp->GetPrimaryVisualComponent()->AddRelativeLocation(VisualOffset);
+}
+
+void FStanceModifier::ApplyMovementSettings(UMoverComponent* MoverComp)
+{
+	switch (ActiveStance)
+	{
+	default:
+	case EStanceMode::Crouch:
+		if (UStanceSettings* StanceSettings = MoverComp->FindSharedSettings_Mutable<UStanceSettings>())
+		{
+			// Update relevant movement settings
+			if (UCommonLegacyMovementSettings* MovementSettings = MoverComp->FindSharedSettings_Mutable<UCommonLegacyMovementSettings>())
+			{
+				MovementSettings->Acceleration = StanceSettings->CrouchingMaxAcceleration;
+				MovementSettings->MaxSpeed = StanceSettings->CrouchingMaxSpeed;
+			}
+		}
+		
+		break;
+
+	// Prone isn't currently implemented properly so we're just doing nothing for now
+	case EStanceMode::Prone:
+		UE_LOG(LogMover, Warning, TEXT("Stance got into prone stance - That mode is not currently implemented fully."));
+		break;
+	}
+}
+
+void FStanceModifier::RevertMovementSettings(UMoverComponent* MoverComp)
+{
+	const AActor* OwnerCDO = Cast<AActor>(MoverComp->GetOwner()->GetClass()->GetDefaultObject());
+	const UMoverComponent* CDOMoverComp = OwnerCDO->FindComponentByClass<UMoverComponent>();
+	const UCommonLegacyMovementSettings* OriginalMovementSettings = CDOMoverComp->FindSharedSettings<UCommonLegacyMovementSettings>();
+	
+	// Revert movement settings back to original settings
+	if (UCommonLegacyMovementSettings* MovementSettings = MoverComp->FindSharedSettings_Mutable<UCommonLegacyMovementSettings>())
+	{
+		MovementSettings->Acceleration = OriginalMovementSettings->Acceleration;
+		MovementSettings->MaxSpeed = OriginalMovementSettings->MaxSpeed;
+	}
+}
diff --git a/Engine/Plugins/Experimental/Mover/Source/Mover/Private/DefaultMovementSet/NavMoverComponent.cpp b/Engine/Plugins/Experimental/Mover/Source/Mover/Private/DefaultMovementSet/NavMoverComponent.cpp
--- a/Engine/Plugins/Experimental/Mover/Source/Mover/Private/DefaultMovementSet/NavMoverComponent.cpp
+++ b/Engine/Plugins/Experimental/Mover/Source/Mover/Private/DefaultMovementSet/NavMoverComponent.cpp
@@ -183,48 +183,27 @@ bool UNavMoverComponent::CanStartPathFollowing() const
 
 bool UNavMoverComponent::IsCrouching() const
 {
-	// TODO: Mover doesn't support crouching yet so return false for now
-	return false;
+	return MoverComponent->HasGameplayTag(Mover_IsCrouching, true);
 }
 
 bool UNavMoverComponent::IsFalling() const
 {
-	if (const UBaseMovementMode* Mode = MoverComponent->GetMovementMode())
-	{
-		return Mode->GameplayTags.HasTag(Mover_IsFalling);
-	}
-
-	return false;
+	return MoverComponent->HasGameplayTag(Mover_IsFalling, true);
 }
 
 bool UNavMoverComponent::IsMovingOnGround() const
 {
-	if (const UBaseMovementMode* Mode = MoverComponent->GetMovementMode())
-	{
-		return Mode->GameplayTags.HasTag(Mover_IsOnGround);
-	}
-
-	return false;
+	return MoverComponent->HasGameplayTag(Mover_IsOnGround, true);
 }
 
 bool UNavMoverComponent::IsSwimming() const
 {
-	if (const UBaseMovementMode* Mode = MoverComponent->GetMovementMode())
-	{
-		return Mode->GameplayTags.HasTag(Mover_IsSwimming);
-	}
-
-	return false;
+	return MoverComponent->HasGameplayTag(Mover_IsSwimming, true);
 }
 
 bool UNavMoverComponent::IsFlying() const
 {
-	if (const UBaseMovementMode* Mode = MoverComponent->GetMovementMode())
-	{
-		return Mode->GameplayTags.HasTag(Mover_IsFlying);
-	}
-
-	return false;
+	return MoverComponent->HasGameplayTag(Mover_IsFlying, true);
 }
 
 void UNavMoverComponent::GetSimpleCollisionCylinder(float& CollisionRadius, float& CollisionHalfHeight) const
diff --git a/Engine/Plugins/Experimental/Mover/Source/Mover/Private/MoverComponent.cpp b/Engine/Plugins/Experimental/Mover/Source/Mover/Private/MoverComponent.cpp
--- a/Engine/Plugins/Experimental/Mover/Source/Mover/Private/MoverComponent.cpp
+++ b/Engine/Plugins/Experimental/Mover/Source/Mover/Private/MoverComponent.cpp
@@ -607,6 +607,36 @@ const FMovementModifierBase* UMoverComponent::FindMovementModifier(const FMoveme
 	return nullptr;
 }
 
+const FMovementModifierBase* UMoverComponent::FindMovementModifierByType(const UScriptStruct* DataStructType) const
+{
+	if (bHasValidCachedState)
+	{
+		// Check active modifiers for modifier handle
+		for (auto ActiveModifierFromSyncStateIt = CachedLastSyncState.MovementModifiers.GetActiveModifiersIterator(); ActiveModifierFromSyncStateIt; ++ActiveModifierFromSyncStateIt)
+		{
+			const TSharedPtr<FMovementModifierBase> ActiveModifierFromSyncState = *ActiveModifierFromSyncStateIt;
+
+			if (DataStructType == ActiveModifierFromSyncState->GetScriptStruct())
+			{
+				return ActiveModifierFromSyncState.Get();
+			}
+		}
+
+		// Check queued modifiers for modifier handle
+		for (auto QueuedModifierFromSyncStateIt = CachedLastSyncState.MovementModifiers.GetActiveModifiersIterator(); QueuedModifierFromSyncStateIt; ++QueuedModifierFromSyncStateIt)
+		{
+			const TSharedPtr<FMovementModifierBase> QueuedModifierFromSyncState = *QueuedModifierFromSyncStateIt;
+
+			if (DataStructType == QueuedModifierFromSyncState->GetScriptStruct())
+			{
+				return QueuedModifierFromSyncState.Get();
+			}
+		}
+	}
+	
+	return nullptr;
+}
+
 bool UMoverComponent::HasGameplayTag(FGameplayTag TagToFind, bool bExactMatch) const
 {
 	if (bHasValidCachedState)
diff --git a/Engine/Plugins/Experimental/Mover/Source/Mover/Private/MoverTypes.cpp b/Engine/Plugins/Experimental/Mover/Source/Mover/Private/MoverTypes.cpp
--- a/Engine/Plugins/Experimental/Mover/Source/Mover/Private/MoverTypes.cpp
+++ b/Engine/Plugins/Experimental/Mover/Source/Mover/Private/MoverTypes.cpp
@@ -12,6 +12,7 @@ UE_DEFINE_GAMEPLAY_TAG_COMMENT(Mover_IsInAir, "Mover.IsInAir", "Default Mover st
 UE_DEFINE_GAMEPLAY_TAG_COMMENT(Mover_IsFalling, "Mover.IsFalling", "Default Mover state flag indicating character is falling.");
 UE_DEFINE_GAMEPLAY_TAG_COMMENT(Mover_IsFlying, "Mover.IsFlying", "Default Mover state flag indicating character is flying.");
 UE_DEFINE_GAMEPLAY_TAG_COMMENT(Mover_IsSwimming, "Mover.IsSwimming", "Default Mover state flag indicating character is swimming.");
+UE_DEFINE_GAMEPLAY_TAG_COMMENT(Mover_IsCrouching, "Mover.Stance.IsCrouching", "Default Mover state flag indicating character is crouching.");
 
 FMoverOnImpactParams::FMoverOnImpactParams() 
 	: AttemptedMoveDelta(0) 
diff --git a/Engine/Plugins/Experimental/Mover/Source/Mover/Public/DefaultMovementSet/CharacterMoverComponent.h b/Engine/Plugins/Experimental/Mover/Source/Mover/Public/DefaultMovementSet/CharacterMoverComponent.h
--- a/Engine/Plugins/Experimental/Mover/Source/Mover/Public/DefaultMovementSet/CharacterMoverComponent.h
+++ b/Engine/Plugins/Experimental/Mover/Source/Mover/Public/DefaultMovementSet/CharacterMoverComponent.h
@@ -3,8 +3,16 @@
 #pragma once
 
 #include "MoverComponent.h"
+#include "MovementModifiers/StanceModifier.h"
 #include "CharacterMoverComponent.generated.h"
 
+/**
+ * Fires when a stance is changed
+ * Note: If a stance was just Activated it will fire with an invalid OldStance
+ *		 If a stance was just Deactivated it will fire with an invalid NewStance
+ */
+DECLARE_DYNAMIC_MULTICAST_DELEGATE_TwoParams(FMover_OnStanceChanged, EStanceMode, OldStance, EStanceMode, NewStance);
+
 UCLASS(BlueprintType, meta = (BlueprintSpawnableComponent))
 class MOVER_API UCharacterMoverComponent : public UMoverComponent
 {
@@ -55,7 +63,30 @@ public:
 	UPROPERTY(EditAnywhere, BlueprintReadWrite, Category=Mover)
     bool bHandleJump;
 	
+	// Whether this actor can currently crouch or not 
+	UFUNCTION(BlueprintCallable, Category = Mover)
+	virtual bool CanCrouch();
+	
+	// Perform crouch on actor
+	UFUNCTION(BlueprintCallable, Category = Mover)
+	virtual void Crouch();
+
+	// Perform uncrouch on actor
+	UFUNCTION(BlueprintCallable, Category = Mover)
+	virtual void UnCrouch();
+
+	// Broadcast when this actor changes stances.
+	UPROPERTY(BlueprintAssignable, Category = Mover)
+	FMover_OnStanceChanged OnStanceChanged;
+	
 protected:
 	UFUNCTION()
 	virtual void OnMoverPreSimulationTick(const FMoverTimeStep& TimeStep, const FMoverInputCmdContext& InputCmd);
+	
+	// ID used to keep track of the modifier responsible for crouching
+	FMovementModifierHandle StanceModifierHandle;
+
+	/** If true, try to crouch (or keep crouching) on next update. If false, try to stop crouching on next update. */
+	UPROPERTY(Category = "Mover|Crouch", VisibleInstanceOnly, BlueprintReadOnly)
+	bool bWantsToCrouch;
 };
diff --git a/Engine/Plugins/Experimental/Mover/Source/Mover/Public/DefaultMovementSet/MovementModifiers/StanceModifier.h b/Engine/Plugins/Experimental/Mover/Source/Mover/Public/DefaultMovementSet/MovementModifiers/StanceModifier.h
new file mode 100644
--- /dev/null
+++ b/Engine/Plugins/Experimental/Mover/Source/Mover/Public/DefaultMovementSet/MovementModifiers/StanceModifier.h
@@ -0,0 +1,80 @@
+﻿// Copyright Epic Games, Inc. All Rights Reserved.
+
+#pragma once
+
+#include "MovementModifier.h"
+#include "StanceModifier.generated.h"
+
+class UCapsuleComponent;
+class UCharacterMoverComponent;
+
+UENUM(BlueprintType)
+enum class EStanceMode : uint8
+{
+	// Invalid default stance
+	Invalid = 0,
+	// Actor goes into crouch
+	Crouch,
+	// Actor goes into prone - not currently implemented
+	Prone,
+};
+
+/**
+ * Stances: Applies settings to the actor to make them go into different stances like crouch or prone(not implemented), affects actor maxacceleration and capsule height
+ * Note: This modifier currently uses the CDO of the actor to reset values to "standing" values.
+ *		 This modifier also assumes the actor is using a capsule as it's updated component for now
+ */
+USTRUCT(BlueprintType)
+struct MOVER_API FStanceModifier : public FMovementModifierBase
+{
+	GENERATED_BODY()
+
+public:
+	FStanceModifier();
+	virtual ~FStanceModifier() override {}
+
+	EStanceMode ActiveStance;
+	
+	virtual bool HasGameplayTag(FGameplayTag TagToFind, bool bExactMatch) const override;
+	
+	/** Fired when this modifier is activated. */
+	virtual void OnStart(UMoverComponent* MoverComp, const FMoverTimeStep& TimeStep, const FMoverSyncState& SyncState, const FMoverAuxStateContext& AuxState) override;
+	
+	/** Fired when this modifier is deactivated. */
+	virtual void OnEnd(UMoverComponent* MoverComp, const FMoverTimeStep& TimeStep, const FMoverSyncState& SyncState, const FMoverAuxStateContext& AuxState) override;
+	
+	/** Fired just before a Substep */
+	virtual void OnPreMovement(UMoverComponent* MoverComp, const FMoverTimeStep& TimeStep) override;
+
+	/** Fired after a Substep */
+	virtual void OnPostMovement(UMoverComponent* MoverComp, const FMoverTimeStep& TimeStep, const FMoverSyncState& SyncState, const FMoverAuxStateContext& AuxState) override;
+	
+	// @return newly allocated copy of this FMovementModifier. Must be overridden by child classes
+	virtual FMovementModifierBase* Clone() const override;
+
+	virtual void NetSerialize(FArchive& Ar) override;
+
+	virtual UScriptStruct* GetScriptStruct() const override;
+
+	virtual FString ToSimpleString() const override;
+
+	virtual void AddReferencedObjects(class FReferenceCollector& Collector) override;
+
+	virtual bool CanExpand(const UCharacterMoverComponent* MoverComp) const;
+	
+	// Whether expanding should be from the base of the capsule or not
+	virtual bool ShouldExpandingMaintainBase(const UCharacterMoverComponent* MoverComp) const;
+
+	// Gets CDO capsule component - useful for getting original values
+	static const UCapsuleComponent* GetOriginalCapsuleComponent(const UMoverComponent* MoverComp);
+
+protected:
+	// Modifies the updated component casted to a capsule component
+	void AdjustCapsule(UMoverComponent* MoverComp, float OldHalfHeight, float NewHalfHeight, float NewEyeHeight);
+
+	// Applies any movement settings like acceleration or max speed changes
+	void ApplyMovementSettings(UMoverComponent* MoverComp);
+	
+	// Reverts any movement settings like acceleration or max speed changes
+	void RevertMovementSettings(UMoverComponent* MoverComp);
+};
diff --git a/Engine/Plugins/Experimental/Mover/Source/Mover/Public/DefaultMovementSet/Settings/StanceSettings.h b/Engine/Plugins/Experimental/Mover/Source/Mover/Public/DefaultMovementSet/Settings/StanceSettings.h
new file mode 100644
--- /dev/null
+++ b/Engine/Plugins/Experimental/Mover/Source/Mover/Public/DefaultMovementSet/Settings/StanceSettings.h
@@ -0,0 +1,34 @@
+// Copyright Epic Games, Inc. All Rights Reserved.
+
+#pragma once
+
+#include "MovementMode.h"
+#include "StanceSettings.generated.h"
+
+/**
+ * StanceSettings: collection of settings that are shared through the mover component and contain stance related settings.
+ */ 
+UCLASS(BlueprintType)
+class MOVER_API UStanceSettings : public UObject, public IMovementSettingsInterface
+{
+	GENERATED_BODY()
+
+	virtual FString GetDisplayName() const override { return GetName(); }
+
+public:
+	// New max acceleration while in crouching
+	UPROPERTY(EditAnywhere, BlueprintReadWrite, Category = "Crouch", meta = (ClampMin = "0", UIMin = "0", ForceUnits = "cm/s^2"))
+	float CrouchingMaxAcceleration = 2000;
+
+	// New max speed while in crouching
+	UPROPERTY(EditAnywhere, BlueprintReadWrite, Category = "Crouch", meta = (ClampMin = "0", UIMin = "0", ForceUnits = "cm/s"))
+	float CrouchingMaxSpeed = 200;
+	
+	// New capsule half height while crouching
+	UPROPERTY(EditAnywhere, BlueprintReadWrite, Category = "Crouch", meta = (ClampMin = "0", UIMin = "0", ForceUnits = "cm"))
+	float CrouchHalfHeight = 55;
+
+	// New eye height while crouching
+	UPROPERTY(EditAnywhere, BlueprintReadWrite, Category = "Crouch")
+	float CrouchedEyeHeight = 50;
+};
\ No newline at end of file
diff --git a/Engine/Plugins/Experimental/Mover/Source/Mover/Public/MoverComponent.h b/Engine/Plugins/Experimental/Mover/Source/Mover/Public/MoverComponent.h
--- a/Engine/Plugins/Experimental/Mover/Source/Mover/Public/MoverComponent.h
+++ b/Engine/Plugins/Experimental/Mover/Source/Mover/Public/MoverComponent.h
@@ -411,6 +411,21 @@ public:	// Queries
 	// Find movement modifier by it's handle. Returns nullptr if the modifier couldn't be found
 	const FMovementModifierBase* FindMovementModifier(const FMovementModifierHandle& ModifierHandle) const;
 
+	// Find movement modifier by type (returns the first modifier it finds). Returns nullptr if the modifier couldn't be found
+	const FMovementModifierBase* FindMovementModifierByType(const UScriptStruct* DataStructType) const;
+	
+	/** Find a movement modifier of a specific type in this components movement modifiers. If not found, null will be returned. */
+	template <typename T>
+	const T* FindMovementModifierByType() const
+	{
+		if (const FMovementModifierBase* FoundData = FindMovementModifierByType(T::StaticStruct()))
+		{
+			return static_cast<const T*>(FoundData);
+		}
+
+		return nullptr;
+	}
+	
 	/**
  	 * Check Mover systems for a gameplay tag.
  	 *
diff --git a/Engine/Plugins/Experimental/Mover/Source/Mover/Public/MoverTypes.h b/Engine/Plugins/Experimental/Mover/Source/Mover/Public/MoverTypes.h
--- a/Engine/Plugins/Experimental/Mover/Source/Mover/Public/MoverTypes.h
+++ b/Engine/Plugins/Experimental/Mover/Source/Mover/Public/MoverTypes.h
@@ -16,6 +16,7 @@ MOVER_API UE_DECLARE_GAMEPLAY_TAG_EXTERN(Mover_IsInAir);
 MOVER_API UE_DECLARE_GAMEPLAY_TAG_EXTERN(Mover_IsFalling);
 MOVER_API UE_DECLARE_GAMEPLAY_TAG_EXTERN(Mover_IsFlying);
 MOVER_API UE_DECLARE_GAMEPLAY_TAG_EXTERN(Mover_IsSwimming);
+MOVER_API UE_DECLARE_GAMEPLAY_TAG_EXTERN(Mover_IsCrouching);
 
 // Struct to hold params for when an impact happens. This contains all of the data for impacts including what gets passed to the FMover_OnImpact delegate
 USTRUCT(BlueprintType, meta = (DisplayName = "Impact Data"))
diff --git a/Engine/Plugins/Experimental/MoverExamples/Source/MoverExamples/Public/CharacterVariants/AbilityInputs.h b/Engine/Plugins/Experimental/MoverExamples/Source/MoverExamples/Public/CharacterVariants/AbilityInputs.h
--- a/Engine/Plugins/Experimental/MoverExamples/Source/MoverExamples/Public/CharacterVariants/AbilityInputs.h
+++ b/Engine/Plugins/Experimental/MoverExamples/Source/MoverExamples/Public/CharacterVariants/AbilityInputs.h
@@ -24,6 +24,8 @@ struct MOVEREXAMPLES_API FMoverExampleAbilityInputs : public FMoverDataStructBas
 	UPROPERTY(BlueprintReadWrite, Category = Mover)
 	bool bWantsToStartZiplining = false;
 
+	UPROPERTY(BlueprintReadWrite, Category = Mover)
+	bool bWantsToBeCrouched = false;
 
 	// @return newly allocated copy of this FMoverExampleAbilityInputs. Must be overridden by child classes
 	virtual FMoverDataStructBase* Clone() const override
@@ -41,6 +43,7 @@ struct MOVEREXAMPLES_API FMoverExampleAbilityInputs : public FMoverDataStructBas
 		Ar.SerializeBits(&bIsAimPressed, 1);
 		Ar.SerializeBits(&bIsVaultJustPressed, 1);
 		Ar.SerializeBits(&bWantsToStartZiplining, 1);
+		Ar.SerializeBits(&bWantsToBeCrouched, 1);
 
 		bOutSuccess = true;
 		return true;
@@ -55,6 +58,7 @@ struct MOVEREXAMPLES_API FMoverExampleAbilityInputs : public FMoverDataStructBas
 		Out.Appendf("bIsAimPressed: %i\n", bIsAimPressed);
 		Out.Appendf("bIsVaultJustPressed: %i\n", bIsVaultJustPressed);
 		Out.Appendf("bWantsToStartZiplining: %i\n", bWantsToStartZiplining);
+		Out.Appendf("bWantsToBeCrouched: %i\n", bWantsToBeCrouched);
 	}
 
 	virtual void AddReferencedObjects(FReferenceCollector& Collector) override { Super::AddReferencedObjects(Collector); }
