Mover: Making layered move preferred moves only set the mode the first time they tick to avoid a mover system constantly setting movement mode

diff --git a/Engine/Plugins/Experimental/Mover/Source/Mover/Private/MovementModeStateMachine.cpp b/Engine/Plugins/Experimental/Mover/Source/Mover/Private/MovementModeStateMachine.cpp
--- a/Engine/Plugins/Experimental/Mover/Source/Mover/Private/MovementModeStateMachine.cpp
+++ b/Engine/Plugins/Experimental/Mover/Source/Mover/Private/MovementModeStateMachine.cpp
@@ -172,6 +172,7 @@ void UMovementModeStateMachine::OnSimulationTick(USceneComponent* UpdatedCompone
 		EffectParams.UpdatedComponent = UpdatedComponent;
 		EffectParams.UpdatedPrimitive = UpdatedPrimitive;
 		
+		bool bModeSetFromInstantEffect = false;
 		// Apply any instant effects that were queued up between ticks
 		if (ApplyInstantEffects(EffectParams, OutputState.SyncState))
 		{
@@ -186,6 +187,7 @@ void UMovementModeStateMachine::OnSimulationTick(USceneComponent* UpdatedCompone
 
 			if (CurrentModeName != OutputState.SyncState.MovementMode)
 			{
+				bModeSetFromInstantEffect = true;
 				SetModeImmediately(OutputState.SyncState.MovementMode);
 				SubstepStartData.SyncState.MovementMode = CurrentModeName;
 			}
@@ -218,12 +220,18 @@ void UMovementModeStateMachine::OnSimulationTick(USceneComponent* UpdatedCompone
 			FProposedMove MoveStep;
 			if (ActiveMove->GenerateMove(SubstepStartData, TimeStep, MoverComp, SimBlackboard, MoveStep))
 			{
+				// If this active move is already past it's first tick we don't need to set the preferred mode again
+				if (ActiveMove->StartSimTimeMs <= TimeStep.BaseSimTimeMs)
+				{
+					MoveStep.PreferredMode = NAME_None;
+				}
+				
 				bHasLayeredMoveContributions = true;
 				MoverComp->MovementMixer->MixLayeredMove(*ActiveMove, MoveStep, CombinedLayeredMove);
 			}
 		}
 
-		if (bHasLayeredMoveContributions && !CombinedLayeredMove.PreferredMode.IsNone())
+		if (bHasLayeredMoveContributions && !CombinedLayeredMove.PreferredMode.IsNone() && !bModeSetFromInstantEffect)
 		{
 			SetModeImmediately(CombinedLayeredMove.PreferredMode);
 			OutputState.SyncState.MovementMode = CurrentModeName;
diff --git a/Engine/Plugins/Experimental/Mover/Source/Mover/Public/LayeredMove.h b/Engine/Plugins/Experimental/Mover/Source/Mover/Public/LayeredMove.h
--- a/Engine/Plugins/Experimental/Mover/Source/Mover/Public/LayeredMove.h
+++ b/Engine/Plugins/Experimental/Mover/Source/Mover/Public/LayeredMove.h
@@ -62,6 +62,9 @@ struct FLayeredMoveFinishVelocitySettings
 * are aggregated and applied to the overall attempted move.
 * Multiple layered moves can be active at any time, and may produce additive motion or motion that overrides
 * what the current Movement Mode may intend.
+* Layered moves can also set a preferred movement mode that only changes the movement mode at the start of
+* the move. Any movement mode changes that need to happen as part of the layered move after the start of the move
+* need to be queued through an Instant Effect or the QueueNextMode function
 */
 
 // Base class for all layered moves
diff --git a/Engine/Plugins/Experimental/Mover/Source/Mover/Public/MoveLibrary/MovementUtilsTypes.h b/Engine/Plugins/Experimental/Mover/Source/Mover/Public/MoveLibrary/MovementUtilsTypes.h
--- a/Engine/Plugins/Experimental/Mover/Source/Mover/Public/MoveLibrary/MovementUtilsTypes.h
+++ b/Engine/Plugins/Experimental/Mover/Source/Mover/Public/MoveLibrary/MovementUtilsTypes.h
@@ -31,22 +31,32 @@ struct MOVER_API FProposedMove
 		bHasDirIntent(false)
 	{}
 
+	// Determines how this move should resolve with other moves
 	UPROPERTY(EditAnywhere, BlueprintReadWrite, Category = Mover)
-	EMoveMixMode MixMode = EMoveMixMode::AdditiveVelocity;		// Determines how this move should resolve with other moves
+	EMoveMixMode MixMode = EMoveMixMode::AdditiveVelocity;
 
+	/**
+	 * Indicates that we should switch to a particular movement mode before the next simulation step is performed.
+	 * Note: If this is set from a layered move the preferred mode will only be set at the beginning of the layered move, not continuously.
+	 */
 	UPROPERTY(EditAnywhere, BlueprintReadWrite, Category = Mover)
-	FName	PreferredMode = NAME_None;					// Indicates that we should switch to a particular movement mode before the next simulation step is performed.
+	FName	PreferredMode = NAME_None;
 
+	// Signals whether there was any directional intent specified
 	UPROPERTY(EditAnywhere, BlueprintReadWrite, Category = Mover)
-	uint8	 bHasDirIntent : 1;							// Signals whether there was any directional intent specified
+	uint8	 bHasDirIntent : 1;
 
+	// Directional, per-axis magnitude [-1, 1] in world space (length of 1 indicates max speed intent). Only valid if bHasDirIntent is set.
 	UPROPERTY(EditAnywhere, BlueprintReadWrite, Category = Mover)
-	FVector  DirectionIntent = FVector::ZeroVector;		// Directional, per-axis magnitude [-1, 1] in world space (length of 1 indicates max speed intent). Only valid if bHasDirIntent is set.
+	FVector  DirectionIntent = FVector::ZeroVector;
 
+	// Units per second, world space, possibly mapped onto walking surface
 	UPROPERTY(EditAnywhere, BlueprintReadWrite, Category = Mover)
-	FVector  LinearVelocity = FVector::ZeroVector;		// Units per second, world space, possibly mapped onto walking surface
+	FVector  LinearVelocity = FVector::ZeroVector;
+	
+	// Degrees per second, local space
 	UPROPERTY(EditAnywhere, BlueprintReadWrite, Category = Mover)
-	FRotator AngularVelocity = FRotator::ZeroRotator;	// Degrees per second, local space
+	FRotator AngularVelocity = FRotator::ZeroRotator;
 };
 
 
