Mover: Fixing visual offset for crouching to properly set visual component on the ground

diff --git a/Engine/Plugins/Experimental/Mover/Source/Mover/Private/DefaultMovementSet/MovementModifiers/StanceModifier.cpp b/Engine/Plugins/Experimental/Mover/Source/Mover/Private/DefaultMovementSet/MovementModifiers/StanceModifier.cpp
--- a/Engine/Plugins/Experimental/Mover/Source/Mover/Private/DefaultMovementSet/MovementModifiers/StanceModifier.cpp
+++ b/Engine/Plugins/Experimental/Mover/Source/Mover/Private/DefaultMovementSet/MovementModifiers/StanceModifier.cpp
@@ -252,7 +252,7 @@ void FStanceModifier::AdjustCapsule(UMoverComponent* MoverComp, float OldHalfHei
 	}
 
 	const FVector CapsuleOffset = MoverComp->GetUpDirection() * (bExpanding ? HalfHeightDifference : -HalfHeightDifference);
-	const FVector VisualOffset = MoverComp->GetUpDirection() * (bExpanding ? -OldHalfHeight : NewHalfHeight);
+	const FVector VisualOffset = MoverComp->GetUpDirection() * (bExpanding ? -HalfHeightDifference : HalfHeightDifference);
 	
 	// Adjust location of capsule as setting it's size left it floating
 	if (!bExpanding || MoverComp->GetVelocity().Length() <= 0)
