Mover: Fixing timing issue with layered move preferred mode not getting applied

diff --git a/Engine/Plugins/Experimental/Mover/Source/Mover/Private/MovementModeStateMachine.cpp b/Engine/Plugins/Experimental/Mover/Source/Mover/Private/MovementModeStateMachine.cpp
--- a/Engine/Plugins/Experimental/Mover/Source/Mover/Private/MovementModeStateMachine.cpp
+++ b/Engine/Plugins/Experimental/Mover/Source/Mover/Private/MovementModeStateMachine.cpp
@@ -221,7 +221,7 @@ void UMovementModeStateMachine::OnSimulationTick(USceneComponent* UpdatedCompone
 			if (ActiveMove->GenerateMove(SubstepStartData, TimeStep, MoverComp, SimBlackboard, MoveStep))
 			{
 				// If this active move is already past it's first tick we don't need to set the preferred mode again
-				if (ActiveMove->StartSimTimeMs <= TimeStep.BaseSimTimeMs)
+				if (ActiveMove->StartSimTimeMs < TimeStep.BaseSimTimeMs)
 				{
 					MoveStep.PreferredMode = NAME_None;
 				}
