Mover: fixed issue with interpolating between default sync states that had different movement bases. Locations would be in different spaces, typically leading to unintentional teleportation.

diff --git a/Engine/Plugins/Experimental/Mover/Source/Mover/Private/MoverDataModelTypes.cpp b/Engine/Plugins/Experimental/Mover/Source/Mover/Private/MoverDataModelTypes.cpp
--- a/Engine/Plugins/Experimental/Mover/Source/Mover/Private/MoverDataModelTypes.cpp
+++ b/Engine/Plugins/Experimental/Mover/Source/Mover/Private/MoverDataModelTypes.cpp
@@ -206,7 +206,7 @@ void FMoverDefaultSyncState::Interpolate(const FMoverDataStructBase& From, const
 
 	// TODO: investigate replacing this threshold with a flag indicating that the state (or parts thereof) isn't intended to be interpolated
 	static constexpr float TeleportThreshold = 500.f * 500.f;
-	if (FVector::DistSquared(FromState->Location, ToState->Location) > TeleportThreshold)
+	if (FVector::DistSquared(FromState->GetLocation_WorldSpace(), ToState->GetLocation_WorldSpace()) > TeleportThreshold)
 	{
 		*this = *ToState;
 	}
