Mover: fixing issue where a turning rate that was a multiple of 360 degrees/second could prevent rotation in the default character movement modes.

diff --git a/Engine/Plugins/Experimental/Mover/Source/Mover/Private/DefaultMovementSet/Modes/FallingMode.cpp b/Engine/Plugins/Experimental/Mover/Source/Mover/Private/DefaultMovementSet/Modes/FallingMode.cpp
--- a/Engine/Plugins/Experimental/Mover/Source/Mover/Private/DefaultMovementSet/Modes/FallingMode.cpp
+++ b/Engine/Plugins/Experimental/Mover/Source/Mover/Private/DefaultMovementSet/Modes/FallingMode.cpp
@@ -166,7 +166,7 @@ void UFallingMode::OnSimulationTick(const FSimulationTickParams& Params, FMoverT
 	FRotator TargetOrient = StartingSyncState->GetOrientation_WorldSpace();
 
 	// Apply orientation changes (if any)
-	if (!ProposedMove.AngularVelocity.IsZero())
+	if (!UMovementUtils::IsAngularVelocityZero(ProposedMove.AngularVelocity))
 	{
 		TargetOrient += (ProposedMove.AngularVelocity * DeltaSeconds);
 	}
diff --git a/Engine/Plugins/Experimental/Mover/Source/Mover/Private/DefaultMovementSet/Modes/FlyingMode.cpp b/Engine/Plugins/Experimental/Mover/Source/Mover/Private/DefaultMovementSet/Modes/FlyingMode.cpp
--- a/Engine/Plugins/Experimental/Mover/Source/Mover/Private/DefaultMovementSet/Modes/FlyingMode.cpp
+++ b/Engine/Plugins/Experimental/Mover/Source/Mover/Private/DefaultMovementSet/Modes/FlyingMode.cpp
@@ -94,7 +94,7 @@ void UFlyingMode::OnSimulationTick(const FSimulationTickParams& Params, FMoverTi
 	bool bIsOrientationChanging = false;
 
 	// Apply orientation changes (if any)
-	if (!ProposedMove.AngularVelocity.IsZero())
+	if (!UMovementUtils::IsAngularVelocityZero(ProposedMove.AngularVelocity))
 	{
 		TargetOrient += (ProposedMove.AngularVelocity * DeltaSeconds);
 		bIsOrientationChanging = (TargetOrient != StartingOrient);
diff --git a/Engine/Plugins/Experimental/Mover/Source/Mover/Private/DefaultMovementSet/Modes/WalkingMode.cpp b/Engine/Plugins/Experimental/Mover/Source/Mover/Private/DefaultMovementSet/Modes/WalkingMode.cpp
--- a/Engine/Plugins/Experimental/Mover/Source/Mover/Private/DefaultMovementSet/Modes/WalkingMode.cpp
+++ b/Engine/Plugins/Experimental/Mover/Source/Mover/Private/DefaultMovementSet/Modes/WalkingMode.cpp
@@ -155,7 +155,7 @@ void UWalkingMode::OnSimulationTick(const FSimulationTickParams& Params, FMoverT
 	bool bIsOrientationChanging = false;
 
 	// Apply orientation changes (if any)
-	if (!ProposedMove.AngularVelocity.IsNearlyZero())
+	if (!UMovementUtils::IsAngularVelocityZero(ProposedMove.AngularVelocity))
 	{ 
 		TargetOrient += (ProposedMove.AngularVelocity * DeltaSeconds);
 		bIsOrientationChanging = (TargetOrient != StartingOrient);
diff --git a/Engine/Plugins/Experimental/Mover/Source/Mover/Private/MoveLibrary/MovementUtils.cpp b/Engine/Plugins/Experimental/Mover/Source/Mover/Private/MoveLibrary/MovementUtils.cpp
--- a/Engine/Plugins/Experimental/Mover/Source/Mover/Private/MoveLibrary/MovementUtils.cpp
+++ b/Engine/Plugins/Experimental/Mover/Source/Mover/Private/MoveLibrary/MovementUtils.cpp
@@ -492,6 +492,12 @@ FVector UMovementUtils::ComputeDirectionIntent(const FVector& MoveInput, EMoveIn
 	return ResultDirIntent;
 }
 
+bool UMovementUtils::IsAngularVelocityZero(const FRotator& AngularVelocity)
+{
+	return (AngularVelocity.Yaw == 0.0 && AngularVelocity.Pitch == 0.0 && AngularVelocity.Roll == 0.0);
+}
+
+
 bool UMovementUtils::TryMoveUpdatedComponent_Internal(const FMovingComponentSet& MovingComps, const FVector& Delta, const FQuat& NewRotation, bool bSweep, EMoveComponentFlags MoveComponentFlags, FHitResult* OutHit, ETeleportType Teleport)
 {
 	if (MovingComps.UpdatedComponent.IsValid())
diff --git a/Engine/Plugins/Experimental/Mover/Source/Mover/Private/PhysicsMover/Modes/PhysicsDrivenFallingMode.cpp b/Engine/Plugins/Experimental/Mover/Source/Mover/Private/PhysicsMover/Modes/PhysicsDrivenFallingMode.cpp
--- a/Engine/Plugins/Experimental/Mover/Source/Mover/Private/PhysicsMover/Modes/PhysicsDrivenFallingMode.cpp
+++ b/Engine/Plugins/Experimental/Mover/Source/Mover/Private/PhysicsMover/Modes/PhysicsDrivenFallingMode.cpp
@@ -112,7 +112,7 @@ void UPhysicsDrivenFallingMode::OnSimulationTick(const FSimulationTickParams& Pa
 	// In air steering
 
 	FRotator TargetOrient = StartingSyncState->GetOrientation_WorldSpace();
-	if (!ProposedMove.AngularVelocity.IsZero())
+	if (!UMovementUtils::IsAngularVelocityZero(ProposedMove.AngularVelocity)) 
 	{
 		TargetOrient += (ProposedMove.AngularVelocity * DeltaSeconds);
 	}
diff --git a/Engine/Plugins/Experimental/Mover/Source/Mover/Private/PhysicsMover/Modes/PhysicsDrivenFlyingMode.cpp b/Engine/Plugins/Experimental/Mover/Source/Mover/Private/PhysicsMover/Modes/PhysicsDrivenFlyingMode.cpp
--- a/Engine/Plugins/Experimental/Mover/Source/Mover/Private/PhysicsMover/Modes/PhysicsDrivenFlyingMode.cpp
+++ b/Engine/Plugins/Experimental/Mover/Source/Mover/Private/PhysicsMover/Modes/PhysicsDrivenFlyingMode.cpp
@@ -74,7 +74,7 @@ void UPhysicsDrivenFlyingMode::OnSimulationTick(const FSimulationTickParams& Par
 	// In air steering
 
 	FRotator TargetOrient = StartingSyncState->GetOrientation_WorldSpace();
-	if (!ProposedMove.AngularVelocity.IsZero())
+	if (!UMovementUtils::IsAngularVelocityZero(ProposedMove.AngularVelocity))
 	{
 		TargetOrient += (ProposedMove.AngularVelocity * DeltaSeconds);
 	}
diff --git a/Engine/Plugins/Experimental/Mover/Source/Mover/Private/PhysicsMover/Modes/PhysicsDrivenSwimmingMode.cpp b/Engine/Plugins/Experimental/Mover/Source/Mover/Private/PhysicsMover/Modes/PhysicsDrivenSwimmingMode.cpp
--- a/Engine/Plugins/Experimental/Mover/Source/Mover/Private/PhysicsMover/Modes/PhysicsDrivenSwimmingMode.cpp
+++ b/Engine/Plugins/Experimental/Mover/Source/Mover/Private/PhysicsMover/Modes/PhysicsDrivenSwimmingMode.cpp
@@ -8,6 +8,7 @@
 #include "GameFramework/PhysicsVolume.h"
 #include "Math/UnitConversion.h"
 #include "MoveLibrary/FloorQueryUtils.h"
+#include "MoveLibrary/MovementUtils.h"
 #include "MoveLibrary/WaterMovementUtils.h"
 #include "MoverComponent.h"
 #include "PhysicsMover/PhysicsMovementUtils.h"
@@ -109,7 +110,7 @@ void UPhysicsDrivenSwimmingMode::OnSimulationTick(const FSimulationTickParams& P
 		const bool bFallTrigger = FMath::Clamp((WaterResult.WaterSplineData.ImmersionDepth + TargetHeight) / (2 * TargetHeight), -2.f, 2.f) < -1.f;
 	
 		FRotator TargetOrient = StartingSyncState->GetOrientation_WorldSpace();
-		if (!ProposedMove.AngularVelocity.IsZero())
+		if (!UMovementUtils::IsAngularVelocityZero(ProposedMove.AngularVelocity))
 		{
 			TargetOrient += (ProposedMove.AngularVelocity * DeltaSeconds);
 		}
diff --git a/Engine/Plugins/Experimental/Mover/Source/Mover/Private/PhysicsMover/Modes/PhysicsDrivenWalkingMode.cpp b/Engine/Plugins/Experimental/Mover/Source/Mover/Private/PhysicsMover/Modes/PhysicsDrivenWalkingMode.cpp
--- a/Engine/Plugins/Experimental/Mover/Source/Mover/Private/PhysicsMover/Modes/PhysicsDrivenWalkingMode.cpp
+++ b/Engine/Plugins/Experimental/Mover/Source/Mover/Private/PhysicsMover/Modes/PhysicsDrivenWalkingMode.cpp
@@ -13,6 +13,7 @@
 #include "GameFramework/PhysicsVolume.h"
 #include "Math/UnitConversion.h"
 #include "MoverComponent.h"
+#include "MoveLibrary/MovementUtils.h"
 #include "MoveLibrary/GroundMovementUtils.h"
 #include "MoveLibrary/WaterMovementUtils.h"
 #include "PhysicsMover/PhysicsMovementUtils.h"
@@ -389,7 +390,7 @@ void UPhysicsDrivenWalkingMode::OnSimulationTick(const FSimulationTickParams& Pa
 		// Target orientation
 		// This is always applied regardless of whether the character is supported
 		FRotator TargetOrientation = StartingSyncState->GetOrientation_WorldSpace();
-		if (!ProposedMove.AngularVelocity.IsZero())
+		if (!UMovementUtils::IsAngularVelocityZero(ProposedMove.AngularVelocity)) 
 		{
 			TargetOrientation += (ProposedMove.AngularVelocity * DeltaSeconds);
 		}
diff --git a/Engine/Plugins/Experimental/Mover/Source/Mover/Public/MoveLibrary/MovementUtils.h b/Engine/Plugins/Experimental/Mover/Source/Mover/Public/MoveLibrary/MovementUtils.h
--- a/Engine/Plugins/Experimental/Mover/Source/Mover/Public/MoveLibrary/MovementUtils.h
+++ b/Engine/Plugins/Experimental/Mover/Source/Mover/Public/MoveLibrary/MovementUtils.h
@@ -181,6 +181,11 @@ public:
 	UFUNCTION(BlueprintCallable, Category = Mover)
 	static FVector ComputeDirectionIntent(const FVector& MoveInput, EMoveInputType MoveInputType);
 
+	/** Returns whether this rotator representing angular velocity has any non-zero values. This function exists due to FRotator::IsZero queries performing undesired wrapping and clamping. */
+	UFUNCTION(BlueprintCallable, Category = Mover)
+	static bool IsAngularVelocityZero(const FRotator& AngularVelocity);
+
+
 	// Internal functions - not meant to be called outside of this library
 	
 	/** Internal function that other move functions use to perform all actual component movement and retrieve results
