Fix for StanceModifier crashing when uncrouching if the MoverComponent was placed manually in the blueprint.

diff --git a/Engine/Plugins/Experimental/Mover/Source/Mover/Private/DefaultMovementSet/MovementModifiers/StanceModifier.cpp b/Engine/Plugins/Experimental/Mover/Source/Mover/Private/DefaultMovementSet/MovementModifiers/StanceModifier.cpp
--- a/Engine/Plugins/Experimental/Mover/Source/Mover/Private/DefaultMovementSet/MovementModifiers/StanceModifier.cpp
+++ b/Engine/Plugins/Experimental/Mover/Source/Mover/Private/DefaultMovementSet/MovementModifiers/StanceModifier.cpp
@@ -14,6 +14,50 @@
 #include "GameFramework/Pawn.h"
 #include "MoveLibrary/MovementUtils.h"
 
+// Gets CDO component type - useful for getting original values
+template <class ComponentType>
+static const ComponentType* GetOriginalComponentType(const UMoverComponent* MoverComp)
+{
+	const ComponentType* OriginalComponent = nullptr;
+
+	if (const AActor* OwnerCDO = Cast<AActor>(MoverComp->GetOwner()->GetClass()->GetDefaultObject()))
+	{
+		// Check if native CDO has Capsule component
+		OriginalComponent = OwnerCDO->FindComponentByClass<ComponentType>();
+
+		// check if it comes from a BP
+		if (!OriginalComponent)
+		{
+			if (const UBlueprintGeneratedClass* OwnerClassAsBP = Cast<UBlueprintGeneratedClass>(OwnerCDO->GetClass()))
+			{
+				TArray<const UBlueprintGeneratedClass*> BlueprintClasses;
+				UBlueprintGeneratedClass::GetGeneratedClassesHierarchy(OwnerClassAsBP, BlueprintClasses);
+				for (const UBlueprintGeneratedClass* BlueprintClass : BlueprintClasses)
+				{
+					if (BlueprintClass->SimpleConstructionScript)
+					{
+						// Check Simple construction script
+						const TArray<USCS_Node*>& SCSNodes = BlueprintClass->SimpleConstructionScript->GetAllNodes();
+						for (USCS_Node* SCSNode : SCSNodes)
+						{
+							if (SCSNode)
+							{
+								if (const ComponentType* BPComponent = Cast<ComponentType>(SCSNode->ComponentTemplate))
+								{
+									OriginalComponent = BPComponent;
+									break;
+								}
+							}
+						}
+					}
+				}
+			}
+		}
+	}
+
+	return OriginalComponent;
+}
+
 FStanceModifier::FStanceModifier()
 	: ActiveStance(EStanceMode::Crouch)
 {
@@ -69,7 +113,7 @@ void FStanceModifier::OnEnd(UMoverComponent* MoverComp, const FMoverTimeStep& Ti
 
 	if (UCapsuleComponent* CapsuleComponent = Cast<UCapsuleComponent>(MoverComp->GetUpdatedComponent()))
 	{
-		if (const UCapsuleComponent* OriginalCapsule = GetOriginalCapsuleComponent(MoverComp))
+		if (const UCapsuleComponent* OriginalCapsule = GetOriginalComponentType<UCapsuleComponent>(MoverComp))
 		{
 			if (const APawn* OwnerCDOAsPawn = Cast<APawn>(OwnerCDO))
 			{
@@ -131,7 +175,7 @@ bool FStanceModifier::CanExpand(const UCharacterMoverComponent* MoverComp) const
 	USceneComponent* UpdatedComponent = MoverComp->GetUpdatedComponent();
 	UPrimitiveComponent* UpdatedCompAsPrimitive = Cast<UPrimitiveComponent>(UpdatedComponent);
 	
-	if (const UCapsuleComponent* OriginalCapsule = GetOriginalCapsuleComponent(MoverComp))
+	if (const UCapsuleComponent* OriginalCapsule = GetOriginalComponentType<UCapsuleComponent>(MoverComp))
 	{
 		StandingHalfHeight = OriginalCapsule->GetScaledCapsuleHalfHeight();
 	}
@@ -187,48 +231,6 @@ bool FStanceModifier::ShouldExpandingMaintainBase(const UCharacterMoverComponent
 	return false;
 }
 
-const UCapsuleComponent* FStanceModifier::GetOriginalCapsuleComponent(const UMoverComponent* MoverComp)
-{
-	const UCapsuleComponent* OriginalCapsule = nullptr;
-
-	if (const AActor* OwnerCDO = Cast<AActor>(MoverComp->GetOwner()->GetClass()->GetDefaultObject()))
-	{
-		// Check if native CDO has Capsule component
-		OriginalCapsule = OwnerCDO->FindComponentByClass<UCapsuleComponent>();
-		
-		// check if it comes from a BP
-		if (!OriginalCapsule)
-		{
-			if (const UBlueprintGeneratedClass* OwnerClassAsBP = Cast<UBlueprintGeneratedClass>(OwnerCDO->GetClass()))
-			{
-				TArray<const UBlueprintGeneratedClass*> BlueprintClasses;
-				UBlueprintGeneratedClass::GetGeneratedClassesHierarchy(OwnerClassAsBP, BlueprintClasses);
-				for (const UBlueprintGeneratedClass* BlueprintClass : BlueprintClasses)
-				{
-					if (BlueprintClass->SimpleConstructionScript)
-					{
-						// Check Simple construction script
-						const TArray<USCS_Node*>& SCSNodes = BlueprintClass->SimpleConstructionScript->GetAllNodes();
-						for (USCS_Node* SCSNode : SCSNodes)
-						{
-							if (SCSNode)
-							{
-								if (const UCapsuleComponent* BPComponent = Cast<UCapsuleComponent>(SCSNode->ComponentTemplate))
-								{
-									OriginalCapsule = BPComponent;
-									break;
-								}
-							}
-						}
-					}
-				}
-			}
-		}
-	}
-
-	return OriginalCapsule;
-}
-
 void FStanceModifier::AdjustCapsule(UMoverComponent* MoverComp, float OldHalfHeight, float NewHalfHeight, float NewEyeHeight)
 {
 	const float HalfHeightDifference = FMath::Abs(NewHalfHeight - OldHalfHeight);
@@ -293,8 +295,7 @@ void FStanceModifier::ApplyMovementSettings(UMoverComponent* MoverComp)
 
 void FStanceModifier::RevertMovementSettings(UMoverComponent* MoverComp)
 {
-	const AActor* OwnerCDO = Cast<AActor>(MoverComp->GetOwner()->GetClass()->GetDefaultObject());
-	const UMoverComponent* CDOMoverComp = OwnerCDO->FindComponentByClass<UMoverComponent>();
+	const UMoverComponent* CDOMoverComp = GetOriginalComponentType<UMoverComponent>(MoverComp);
 	const UCommonLegacyMovementSettings* OriginalMovementSettings = CDOMoverComp->FindSharedSettings<UCommonLegacyMovementSettings>();
 	
 	// Revert movement settings back to original settings
diff --git a/Engine/Plugins/Experimental/Mover/Source/Mover/Public/DefaultMovementSet/MovementModifiers/StanceModifier.h b/Engine/Plugins/Experimental/Mover/Source/Mover/Public/DefaultMovementSet/MovementModifiers/StanceModifier.h
--- a/Engine/Plugins/Experimental/Mover/Source/Mover/Public/DefaultMovementSet/MovementModifiers/StanceModifier.h
+++ b/Engine/Plugins/Experimental/Mover/Source/Mover/Public/DefaultMovementSet/MovementModifiers/StanceModifier.h
@@ -65,9 +65,6 @@ public:
 	// Whether expanding should be from the base of the capsule or not
 	virtual bool ShouldExpandingMaintainBase(const UCharacterMoverComponent* MoverComp) const;
 
-	// Gets CDO capsule component - useful for getting original values
-	static const UCapsuleComponent* GetOriginalCapsuleComponent(const UMoverComponent* MoverComp);
-
 protected:
 	// Modifies the updated component casted to a capsule component
 	void AdjustCapsule(UMoverComponent* MoverComp, float OldHalfHeight, float NewHalfHeight, float NewEyeHeight);
