Mover: fixing crash that could occur in standlone games where an invalid scene component could be discovered during initialization, but could become garbage and replaced with another instance as play begins

diff --git a/Engine/Plugins/Experimental/Mover/Source/Mover/Private/MoverComponent.cpp b/Engine/Plugins/Experimental/Mover/Source/Mover/Private/MoverComponent.cpp
--- a/Engine/Plugins/Experimental/Mover/Source/Mover/Private/MoverComponent.cpp
+++ b/Engine/Plugins/Experimental/Mover/Source/Mover/Private/MoverComponent.cpp
@@ -68,24 +68,8 @@ UMoverComponent::UMoverComponent()
 void UMoverComponent::InitializeComponent()
 {
 	TGuardValue<bool> InInitializeComponentGuard(bInInitializeComponent, true);
-
-	// RootComponent is null in OnRegister for blueprint (non-native) root components.
-	if (!UpdatedComponent)
-	{
-		// Auto-register owner's root component if found.
-		if (AActor* MyActor = GetOwner())
-		{
-			if (USceneComponent* NewUpdatedComponent = MyActor->GetRootComponent())
-			{
-				SetUpdatedComponent(NewUpdatedComponent);
-			}
-			else
-			{
-				ensureMsgf(false, TEXT("No root component found on %s. Simulation initialization will most likely fail."), *GetPathNameSafe(MyActor));
-			}
-		}
-	}
-
+	FindDefaultUpdatedComponent();
+	
 	// Instantiate out sister backend component that will actually talk to the system driving the simulation
 	if (BackendClass)
 	{
@@ -112,6 +96,7 @@ void UMoverComponent::UninitializeComponent()
 	if (UActorComponent* LiaisonAsComp = Cast<UActorComponent>(BackendLiaisonComp.Get()))
 	{
 		LiaisonAsComp->DestroyComponent();
+		BackendLiaisonComp = nullptr;
 	}
 
 	Super::UninitializeComponent();
@@ -122,27 +107,9 @@ void UMoverComponent::OnRegister()
 {
 	TGuardValue<bool> InOnRegisterGuard(bInOnRegister, true);
 
-	UpdatedCompAsPrimitive = Cast<UPrimitiveComponent>(UpdatedComponent);
 	Super::OnRegister();
 
-	const UWorld* MyWorld = GetWorld();
-
-	if (MyWorld && MyWorld->IsGameWorld())
-	{
-		const AActor* MyActor = GetOwner();
-
-		USceneComponent* NewUpdatedComponent = UpdatedComponent;
-		if (!UpdatedComponent)
-		{
-			// Auto-register owner's root component if found.
-			if (MyActor)
-			{
-				NewUpdatedComponent = MyActor->GetRootComponent();
-			}
-		}
-
-		SetUpdatedComponent(NewUpdatedComponent);
-	}
+	FindDefaultUpdatedComponent();
 }
 
 
@@ -190,6 +157,9 @@ void UMoverComponent::PostLoad()
 void UMoverComponent::BeginPlay()
 {
 	Super::BeginPlay();
+	
+	FindDefaultUpdatedComponent();
+	ensureMsgf(UpdatedComponent != nullptr, TEXT("No root component found on %s. Simulation initialization will most likely fail."), *GetPathNameSafe(GetOwner()));
 
 	if (const AActor* MyActor = GetOwner())
 	{
@@ -859,6 +829,24 @@ void UMoverComponent::SetUpdatedComponent(USceneComponent* NewUpdatedComponent)
 	UpdateTickRegistration();
 }
 
+void UMoverComponent::FindDefaultUpdatedComponent()
+{
+	if (!IsValid(UpdatedComponent))
+	{
+		USceneComponent* NewUpdatedComponent = nullptr;
+
+		const AActor* MyActor = GetOwner();
+		const UWorld* MyWorld = GetWorld();
+
+		if (MyActor && MyWorld && MyWorld->IsGameWorld())
+		{
+			NewUpdatedComponent = MyActor->GetRootComponent();
+		}
+
+		SetUpdatedComponent(NewUpdatedComponent);
+	}
+}
+
 
 void UMoverComponent::UpdateTickRegistration()
 {
diff --git a/Engine/Plugins/Experimental/Mover/Source/Mover/Public/MoverComponent.h b/Engine/Plugins/Experimental/Mover/Source/Mover/Public/MoverComponent.h
--- a/Engine/Plugins/Experimental/Mover/Source/Mover/Public/MoverComponent.h
+++ b/Engine/Plugins/Experimental/Mover/Source/Mover/Public/MoverComponent.h
@@ -454,6 +454,7 @@ public:
 protected:
 	// Basic "Update Component/Ticking"
 	void SetUpdatedComponent(USceneComponent* NewUpdatedComponent);
+	void FindDefaultUpdatedComponent();
 	void UpdateTickRegistration();
 
 	/** Called when a rollback occurs, before the simulation state has been restored */
