Mover: adding debug CVar to disable root motion attribute layered moves, to assist with comparison testing between root motion and the locomotion model

diff --git a/Engine/Plugins/Experimental/Mover/Source/Mover/Private/DefaultMovementSet/LayeredMoves/RootMotionAttributeLayeredMove.cpp b/Engine/Plugins/Experimental/Mover/Source/Mover/Private/DefaultMovementSet/LayeredMoves/RootMotionAttributeLayeredMove.cpp
--- a/Engine/Plugins/Experimental/Mover/Source/Mover/Private/DefaultMovementSet/LayeredMoves/RootMotionAttributeLayeredMove.cpp
+++ b/Engine/Plugins/Experimental/Mover/Source/Mover/Private/DefaultMovementSet/LayeredMoves/RootMotionAttributeLayeredMove.cpp
@@ -18,6 +18,13 @@ FAutoConsoleVariable CVarLogRootMotionAttrSteps(
 	false,
 	TEXT("Whether to log detailed information about root motion attribute layered moves. 0: Disable, 1: Enable"),
 	ECVF_Cheat);
+
+FAutoConsoleVariable CVarDisableRootMotionAttrContributions(
+	TEXT("mover.debug.DisableRootMotionAttributes"),
+	false,
+	TEXT("If enabled, contributions from root motion attributes will be ignored in favor of other Mover influences"),
+	ECVF_Cheat);
+
 #endif	// !UE_BUILD_SHIPPING
 
 FLayeredMove_RootMotionAttribute::FLayeredMove_RootMotionAttribute()
@@ -28,6 +35,13 @@ FLayeredMove_RootMotionAttribute::FLayeredMove_RootMotionAttribute()
 
 bool FLayeredMove_RootMotionAttribute::GenerateMove(const FMoverTickStartData& SimState, const FMoverTimeStep& TimeStep, const UMoverComponent* MoverComp, UMoverBlackboard* SimBlackboard, FProposedMove& OutProposedMove)
 {
+#if !UE_BUILD_SHIPPING
+	if (CVarDisableRootMotionAttrContributions->GetBool())
+	{
+		return false;	// do not contribute any movement
+	}
+#endif // !UE_BUILD_SHIPPING
+
 	const float DeltaSeconds = TimeStep.StepMs / 1000.f;
 
 	bool bDidAttrHaveRootMotion = false;
