Mover
- Make MovementModeStateMachine header public
- Transitions can access the mover component and shared settings
- Allow derived MoverComponent classes to access state machine and blackboard

diff --git a/Engine/Plugins/Experimental/Mover/Source/Mover/Private/MovementMode.cpp b/Engine/Plugins/Experimental/Mover/Source/Mover/Private/MovementMode.cpp
--- a/Engine/Plugins/Experimental/Mover/Source/Mover/Private/MovementMode.cpp
+++ b/Engine/Plugins/Experimental/Mover/Source/Mover/Private/MovementMode.cpp
@@ -1,6 +1,7 @@
 // Copyright Epic Games, Inc. All Rights Reserved.
 
 #include "MovementMode.h"
+#include "MovementModeTransition.h"
 #include "MoverComponent.h"
 #include "Engine/BlueprintGeneratedClass.h"
 
@@ -35,6 +36,11 @@ void UBaseMovementMode::DoRegister(const FName ModeName)
 	UFunction* GenMoveFunction = GetClass()->FindFunctionByName(GenMoveFuncName);
 	bHasBlueprintGenerateMove = IsImplementedInBlueprint(GenMoveFunction);
 
+	for (TObjectPtr<UBaseMovementModeTransition>& Transition : Transitions)
+	{
+		Transition->DoRegister();
+	}
+
 	OnRegistered(ModeName);
 }
 
diff --git a/Engine/Plugins/Experimental/Mover/Source/Mover/Private/MovementModeStateMachine.cpp b/Engine/Plugins/Experimental/Mover/Source/Mover/Private/MovementModeStateMachine.cpp
--- a/Engine/Plugins/Experimental/Mover/Source/Mover/Private/MovementModeStateMachine.cpp
+++ b/Engine/Plugins/Experimental/Mover/Source/Mover/Private/MovementModeStateMachine.cpp
@@ -2,6 +2,7 @@
 
 
 #include "MovementModeStateMachine.h"
+#include "MovementModeTransition.h"
 #include "MoverDeveloperSettings.h"
 #include "MoverLog.h"
 #include "MoveLibrary/MovementUtils.h"
@@ -86,6 +87,29 @@ void UMovementModeStateMachine::SetDefaultMode(FName NewDefaultModeName)
 	DefaultModeName = NewDefaultModeName;
 }
 
+void UMovementModeStateMachine::RegisterGlobalTransition(TObjectPtr<UBaseMovementModeTransition> Transition)
+{
+	GlobalTransitions.Add(Transition);
+
+	Transition->DoRegister();
+}
+
+void UMovementModeStateMachine::UnregisterGlobalTransition(TObjectPtr<UBaseMovementModeTransition> Transition)
+{
+	Transition->DoUnregister();
+
+	GlobalTransitions.Remove(Transition);
+}
+
+void UMovementModeStateMachine::ClearAllGlobalTransitions()
+{
+	for (TObjectPtr<UBaseMovementModeTransition>& Transition : GlobalTransitions)
+	{
+		Transition->DoUnregister();
+	}
+
+	GlobalTransitions.Empty();
+}
 
 void UMovementModeStateMachine::QueueNextMode(FName DesiredNextModeName, bool bShouldReenter)
 {
@@ -291,7 +315,7 @@ void UMovementModeStateMachine::OnSimulationTick(USceneComponent* UpdatedCompone
 
 				if (EvalResult.NextMode.IsNone())
 				{
-					for (UBaseMovementModeTransition* Transition : MoverComp->Transitions)
+					for (UBaseMovementModeTransition* Transition : GlobalTransitions)
 					{
 						if (IsValid(Transition))
 						{
diff --git a/Engine/Plugins/Experimental/Mover/Source/Mover/Private/MovementModeStateMachine.h b/Engine/Plugins/Experimental/Mover/Source/Mover/Private/MovementModeStateMachine.h
deleted file mode 100644
diff --git a/Engine/Plugins/Experimental/Mover/Source/Mover/Private/MovementModeTransition.cpp b/Engine/Plugins/Experimental/Mover/Source/Mover/Private/MovementModeTransition.cpp
--- a/Engine/Plugins/Experimental/Mover/Source/Mover/Private/MovementModeTransition.cpp
+++ b/Engine/Plugins/Experimental/Mover/Source/Mover/Private/MovementModeTransition.cpp
@@ -2,6 +2,7 @@
 
 #include "MovementModeTransition.h"
 #include "Engine/BlueprintGeneratedClass.h"
+#include "MoverComponent.h"
 #include "MoverSimulationTypes.h"
 #include "MoverTypes.h"
 
@@ -31,10 +32,34 @@ void UBaseMovementModeTransition::DoRegister()
 	UFunction* TriggerFunction = GetClass()->FindFunctionByName(TriggerFuncName);
 	bHasBlueprintTrigger = IsImplementedInBlueprint(TriggerFunction);
 
+	OnRegistered();
 }
 
 void UBaseMovementModeTransition::DoUnregister()
 {
+	OnUnregistered();
+}
+
+void UBaseMovementModeTransition::OnRegistered()
+{
+}
+
+void UBaseMovementModeTransition::OnUnregistered()
+{
+}
+
+UMoverComponent* UBaseMovementModeTransition::GetMoverComponent() const
+{
+	// Transitions can belong to either the mover component or to
+	// a particular movement mode
+	if (UBaseMovementMode* MovementMode = Cast<UBaseMovementMode>(GetOuter()))
+	{
+		return MovementMode->GetMoverComponent();
+	}
+	else
+	{
+		return Cast<UMoverComponent>(GetOuter());
+	}
 }
 
 FTransitionEvalResult UBaseMovementModeTransition::DoEvaluate(const FSimulationTickParams& Params) const
diff --git a/Engine/Plugins/Experimental/Mover/Source/Mover/Private/MoverComponent.cpp b/Engine/Plugins/Experimental/Mover/Source/Mover/Private/MoverComponent.cpp
--- a/Engine/Plugins/Experimental/Mover/Source/Mover/Private/MoverComponent.cpp
+++ b/Engine/Plugins/Experimental/Mover/Source/Mover/Private/MoverComponent.cpp
@@ -686,6 +686,7 @@ bool UMoverComponent::InitMoverSimulation()
 	// Set up FSM and initial movement states
 	ModeFSM = NewObject<UMovementModeStateMachine>(this, TEXT("MoverStateMachine"), RF_Transient);
 	ModeFSM->ClearAllMovementModes();
+	ModeFSM->ClearAllGlobalTransitions();
 
 	bool bHasMatchingStartingState = false;
 
@@ -703,6 +704,11 @@ bool UMoverComponent::InitMoverSimulation()
 		bHasMatchingStartingState |= (StartingMovementMode == Element.Key);
 	}
 
+	for (TObjectPtr<UBaseMovementModeTransition>& Transition : Transitions)
+	{
+		ModeFSM->RegisterGlobalTransition(Transition);
+	}
+
 	if (!bHasMatchingStartingState)
 	{
 		UE_LOG(LogMover, Warning, TEXT("Invalid StartingMovementMode '%s' specified on %s. Mover actor will not function."),
@@ -980,6 +986,30 @@ bool UMoverComponent::ValidateSetup(FDataValidationContext& Context) const
 					bDidFindAnyProblems = true;
 				}
 			}
+
+			for (const UBaseMovementModeTransition* Transition : Element.Value->Transitions)
+			{
+				for (const TSubclassOf<UObject>& Type : Transition->SharedSettingsClasses)
+				{
+					if (Type.Get() == nullptr)
+					{
+						Context.AddError(FText::Format(LOCTEXT("InvalidModeSettingsError", "Transition on mode {0} on {1}, has an invalid SharedSettingsClass. You may need to remove the invalid settings class."),
+							FText::FromName(Element.Key),
+							FText::FromString(GetNameSafe(GetOwner()))));
+
+						bDidFindAnyProblems = true;
+					}
+					else if (FindSharedSettings(Type) == nullptr)
+					{
+						Context.AddError(FText::Format(LOCTEXT("MissingModeSettingsError", "Transition on mode {0} on {1}, is missing its desired SharedSettingsClass {2}. You may need to save the asset and/or recompile."),
+							FText::FromName(Element.Key),
+							FText::FromString(GetNameSafe(GetOwner())),
+							FText::FromString(Type->GetName())));
+
+						bDidFindAnyProblems = true;
+					}
+				}
+			}
 		}
 	}
 
@@ -1003,6 +1033,25 @@ bool UMoverComponent::ValidateSetup(FDataValidationContext& Context) const
 
 			bDidFindAnyProblems = true;
 		}
+
+		for (const TSubclassOf<UObject>& Type : Transition->SharedSettingsClasses)
+		{
+			if (Type.Get() == nullptr)
+			{
+				Context.AddError(FText::Format(LOCTEXT("InvalidModeSettingsError", "Transition on {0}, has an invalid SharedSettingsClass. You may need to remove the invalid settings class."),
+					FText::FromString(GetNameSafe(GetOwner()))));
+
+				bDidFindAnyProblems = true;
+			}
+			else if (FindSharedSettings(Type) == nullptr)
+			{
+				Context.AddError(FText::Format(LOCTEXT("MissingModeSettingsError", "Transition on {0}, is missing its desired SharedSettingsClass {2}. You may need to save the asset and/or recompile."),
+					FText::FromString(GetNameSafe(GetOwner())),
+					FText::FromString(Type->GetName())));
+
+				bDidFindAnyProblems = true;
+			}
+		}
 	}
 
 	// Verify backend liaison
@@ -1107,6 +1156,64 @@ void UMoverComponent::RefreshSharedSettings()
 					SharedSettings.Add(NewSettings);
 				}
 			}
+
+			for (const UBaseMovementModeTransition* Transition : Mode->Transitions)
+			{
+				for (const TSubclassOf<UObject>& SharedSettingsType : Transition->SharedSettingsClasses)
+				{
+					if (SharedSettingsType.Get() == nullptr)
+					{
+						UE_LOG(LogMover, Warning, TEXT("Invalid shared setting class detected on Transition on Movement Mode %s."), *Mode->GetName());
+						continue;
+					}
+
+					bool bFoundMatchingClass = false;
+					for (const TObjectPtr<UObject>& SettingsObj : SharedSettings)
+					{
+						if (SettingsObj && SettingsObj->IsA(SharedSettingsType))
+						{
+							bFoundMatchingClass = true;
+							UnreferencedSettingsObjs.Remove(SettingsObj);
+							break;
+						}
+					}
+
+					if (!bFoundMatchingClass)
+					{
+						UObject* NewSettings = NewObject<UObject>(this, SharedSettingsType, NAME_None, GetMaskedFlags(RF_PropagateToSubObjects) | RF_Transactional);
+						SharedSettings.Add(NewSettings);
+					}
+				}
+			}
+		}
+	}
+
+	for (const UBaseMovementModeTransition* Transition : Transitions)
+	{
+		for (const TSubclassOf<UObject>& SharedSettingsType : Transition->SharedSettingsClasses)
+		{
+			if (SharedSettingsType.Get() == nullptr)
+			{
+				UE_LOG(LogMover, Warning, TEXT("Invalid shared setting class detected on Transition."));
+				continue;
+			}
+
+			bool bFoundMatchingClass = false;
+			for (const TObjectPtr<UObject>& SettingsObj : SharedSettings)
+			{
+				if (SettingsObj && SettingsObj->IsA(SharedSettingsType))
+				{
+					bFoundMatchingClass = true;
+					UnreferencedSettingsObjs.Remove(SettingsObj);
+					break;
+				}
+			}
+
+			if (!bFoundMatchingClass)
+			{
+				UObject* NewSettings = NewObject<UObject>(this, SharedSettingsType, NAME_None, GetMaskedFlags(RF_PropagateToSubObjects) | RF_Transactional);
+				SharedSettings.Add(NewSettings);
+			}
 		}
 	}
 
diff --git a/Engine/Plugins/Experimental/Mover/Source/Mover/Public/MovementModeStateMachine.h b/Engine/Plugins/Experimental/Mover/Source/Mover/Public/MovementModeStateMachine.h
new file mode 100644
--- /dev/null
+++ b/Engine/Plugins/Experimental/Mover/Source/Mover/Public/MovementModeStateMachine.h
@@ -0,0 +1,140 @@
+// Copyright Epic Games, Inc. All Rights Reserved.
+
+#pragma once
+
+#include "CoreMinimal.h"
+#include "Templates/SubclassOf.h"
+#include "MovementMode.h"
+#include "InstantMovementEffect.h"
+#include "MovementModeStateMachine.generated.h"
+
+struct FProposedMove;
+class UImmediateMovementModeTransition;
+class UMovementModeTransition;
+
+/**
+ * NullMovementMode: a default do-nothing mode used as a placeholder when no other mode is active
+ */
+ UCLASS(NotBlueprintable)
+class UNullMovementMode : public UBaseMovementMode
+{
+	GENERATED_UCLASS_BODY()
+
+public:
+	virtual void OnSimulationTick(const FSimulationTickParams& Params, FMoverTickEndData& OutputState) override;
+
+	const static FName NullModeName;
+};
+
+/**
+ * - Any movement modes registered are co-owned by the state machine
+ * - There is always an active mode, falling back to a do-nothing 'null' mode
+ * - Queuing a mode that is already active will cause it to exit and re-enter
+ * - Modes only switch during simulation tick
+ */
+ UCLASS()
+class MOVER_API UMovementModeStateMachine : public UObject
+{
+	 GENERATED_UCLASS_BODY()
+
+public:
+	void RegisterMovementMode(FName ModeName, TObjectPtr<UBaseMovementMode> Mode, bool bIsDefaultMode=false);
+	void RegisterMovementMode(FName ModeName, TSubclassOf<UBaseMovementMode> ModeType, bool bIsDefaultMode=false);
+
+	void UnregisterMovementMode(FName ModeName);
+	void UnregisterMovementMode(TObjectPtr<UBaseMovementMode> ModeToUnregister);
+	void ClearAllMovementModes();
+
+	void RegisterGlobalTransition(TObjectPtr<UBaseMovementModeTransition> Transition);
+	void UnregisterGlobalTransition(TObjectPtr<UBaseMovementModeTransition> Transition);
+	void ClearAllGlobalTransitions();
+
+	void SetDefaultMode(FName NewDefaultModeName);
+
+	void QueueNextMode(FName DesiredNextModeName, bool bShouldReenter=false);
+	void SetModeImmediately(FName DesiredModeName, bool bShouldReenter=false);
+	void ClearQueuedMode();
+
+	void OnSimulationTick(USceneComponent* UpdatedComponent, UPrimitiveComponent* UpdatedPrimitive, UMoverBlackboard* SimBlackboard, const FMoverTickStartData& StartState, const FMoverTimeStep& TimeStep, FMoverTickEndData& OutputState);
+ 	void OnSimulationPreRollback(const FMoverSyncState* InvalidSyncState, const FMoverSyncState* SyncState, const FMoverAuxStateContext* InvalidAuxState, const FMoverAuxStateContext* AuxState);
+	void OnSimulationRollback(const FMoverSyncState* SyncState, const FMoverAuxStateContext* AuxState);
+
+	FName GetCurrentModeName() const { return CurrentModeName; }
+
+	const UBaseMovementMode* GetCurrentMode() const;
+
+	const UBaseMovementMode* FindMovementMode(FName ModeName) const;
+
+	void QueueLayeredMove(TSharedPtr<FLayeredMoveBase> Move);
+	
+ 	void QueueInstantMovementEffect(TSharedPtr<FInstantMovementEffect> Effect);
+
+	FMovementModifierHandle QueueMovementModifier(TSharedPtr<FMovementModifierBase> Modifier);
+
+ 	void CancelModifierFromHandle(FMovementModifierHandle ModifierHandle);
+ 	
+protected:
+
+	virtual void PostInitProperties() override;
+
+	UPROPERTY()
+	TMap<FName, TObjectPtr<UBaseMovementMode>> Modes;
+	TArray<TObjectPtr<UBaseMovementModeTransition>> GlobalTransitions;
+
+	UPROPERTY(Transient)
+	TObjectPtr<UImmediateMovementModeTransition> QueuedModeTransition;
+
+	FName DefaultModeName = NAME_None;
+	FName CurrentModeName = NAME_None;
+
+	/** Moves that are queued to be added to the simulation at the start of the next sim subtick */
+	TArray<TSharedPtr<FLayeredMoveBase>> QueuedLayeredMoves;
+
+ 	/** Effects that are queued to be applied to the simulation at the start of the next sim subtick or at the end of this tick */
+ 	TArray<TSharedPtr<FInstantMovementEffect>> QueuedInstantEffects;
+
+ 	/** Modifiers that are queued to be added to the simulation at the start of the next sim subtick */
+ 	TArray<TSharedPtr<FMovementModifierBase>> QueuedMovementModifiers;
+
+ 	/** Modifiers that are to be canceled at the start of the next sim subtick */
+ 	TArray<FMovementModifierHandle> ModifiersToCancel;
+ 	
+private:
+	void ConstructDefaultModes();
+	void AdvanceToNextMode();
+	void FlushQueuedMovesToGroup(FLayeredMoveGroup& Group);
+ 	void FlushQueuedModifiersToGroup(FMovementModifierGroup& ModifierGroup);
+ 	void FlushModifierCancellationsToGroup(FMovementModifierGroup& ActiveModifierGroup);
+ 	void RollbackModifiers(const FMoverSyncState* InvalidSyncState, const FMoverSyncState* SyncState, const FMoverAuxStateContext* InvalidAuxState, const FMoverAuxStateContext* AuxState);
+ 	bool ApplyInstantEffects(FApplyMovementEffectParams& ApplyEffectParams, FMoverSyncState& OutputState);
+	AActor* GetOwnerActor() const;
+};
+
+
+
+ /**
+  * Simple transition that evaluates true if a "next mode" is set. Used internally only by the Mover plugin. 
+  */
+ UCLASS()
+ class UImmediateMovementModeTransition : public UBaseMovementModeTransition
+ {
+	 GENERATED_UCLASS_BODY()
+
+ public:
+	 virtual FTransitionEvalResult OnEvaluate(const FSimulationTickParams& Params) const override;
+	 virtual void OnTrigger(const FSimulationTickParams& Params) override;
+
+
+	 bool IsSet() const { return !NextMode.IsNone(); }
+	 void SetNextMode(FName DesiredModeName, bool bShouldReenter = false);
+	 void Clear();
+
+
+
+	 FName GetNextModeName() const { return NextMode; }
+	 bool ShoulReenter() const { return bShouldNextModeReenter; }
+
+ private:
+	 FName NextMode = NAME_None;
+	 bool bShouldNextModeReenter = false;
+ };
diff --git a/Engine/Plugins/Experimental/Mover/Source/Mover/Public/MovementModeTransition.h b/Engine/Plugins/Experimental/Mover/Source/Mover/Public/MovementModeTransition.h
--- a/Engine/Plugins/Experimental/Mover/Source/Mover/Public/MovementModeTransition.h
+++ b/Engine/Plugins/Experimental/Mover/Source/Mover/Public/MovementModeTransition.h
@@ -7,6 +7,7 @@
 #include "MovementModeTransition.generated.h"
 
 struct FSimulationTickParams;
+class UMoverComponent;
 
 /** 
  * Results from a movement mode transition evaluation
@@ -37,11 +38,20 @@ class MOVER_API UBaseMovementModeTransition : public UObject
 	GENERATED_UCLASS_BODY()
 
 public:
-	// TODO : add hooks for shared settings, so that transitions can refer to the same settings as the modes without needing to know about the mode's implementation
-
 	void DoRegister();
 	void DoUnregister();
 
+	UFUNCTION(BlueprintCallable, Category = Mover, meta = (DisplayName = "Get Mover Component"))
+	UMoverComponent* GetMoverComponent() const;
+
+	/** Templated convenience version of GetMoverComponent which checks the type is as presumed. */
+	template<class T>
+	T* GetMoverComponentChecked() const
+	{
+		static_assert(TPointerIsConvertibleFromTo<T, const UMoverComponent>::Value, "'T' template parameter to GetMoverComponentChecked must be derived from UMoverComponent");
+		return CastChecked<T>(GetMoverComponent());
+	}
+
 	FTransitionEvalResult DoEvaluate(const FSimulationTickParams& Params) const;
 
 	void DoTrigger(const FSimulationTickParams& Params);
@@ -50,7 +60,14 @@ public:
 	virtual EDataValidationResult IsDataValid(class FDataValidationContext& Context) const override;
 #endif // WITH_EDITOR
 
+	/** Settings object type that this mode depends on. May be shared with other transitions and movement modes. When the transition is added to a Mover Component, it will create a shared instance of this settings class. */
+	UPROPERTY(EditDefaultsOnly, Category = Mover, meta = (MustImplement = "/Script/Mover.MovementSettingsInterface"))
+	TArray<TSubclassOf<UObject>> SharedSettingsClasses;
+
 protected:
+	virtual void OnRegistered();
+
+	virtual void OnUnregistered();
 
 	virtual FTransitionEvalResult OnEvaluate(const FSimulationTickParams& Params) const;
 
diff --git a/Engine/Plugins/Experimental/Mover/Source/Mover/Public/MoverComponent.h b/Engine/Plugins/Experimental/Mover/Source/Mover/Public/MoverComponent.h
--- a/Engine/Plugins/Experimental/Mover/Source/Mover/Public/MoverComponent.h
+++ b/Engine/Plugins/Experimental/Mover/Source/Mover/Public/MoverComponent.h
@@ -69,7 +69,7 @@ public:
 	virtual void BeginPlay() override;
 
 	virtual void TickComponent(float DeltaTime, enum ELevelTick TickType, FActorComponentTickFunction* ThisTickFunction) override;
-	
+
 	// Broadcast before each simulation tick
 	UPROPERTY(BlueprintAssignable, Category = Mover)
 	FMover_OnPreSimTick OnPreSimulationTick;
@@ -520,6 +520,17 @@ protected:
 	/** Tick function that may be called anytime after this actor's movement step, useful as a way to support based movement on objects that are not */
 	FMoverDynamicBasedMovementTickFunction BasedMovementTickFunction;
 
+	UPROPERTY(Transient)
+	TObjectPtr<UMovementModeStateMachine> ModeFSM;	// JAH TODO: Also consider allowing a type property on the component to allow an alternative machine implementation to be allocated/used
+
+	/** Used to store cached data & computations between decoupled systems, that can be referenced by name */
+	UPROPERTY(Transient)
+	TObjectPtr<UMoverBlackboard> SimBlackboard;
+
+	/** If enabled, this actor will be moved to follow a base actor that it's standing on. Typically disabled for physics-based movement, which handles based movement internally. */
+	UPROPERTY(EditDefaultsOnly, Category = "Mover")
+	bool bSupportsKinematicBasedMovement = true;
+
 private:
 	/** Collection of settings objects that are shared between movement modes. This list is automatically managed based on the @MovementModes contents. */
 	UPROPERTY(EditDefaultsOnly, EditFixedSize, Instanced, Category = Mover, meta = (NoResetToDefault, MustImplement = "/Script/Mover.MovementSettingsInterface"))
@@ -537,10 +548,6 @@ private:
 	UPROPERTY(EditDefaultsOnly, Category = "Mover|Constraints")
 	FPlanarConstraint PlanarConstraint;
 
-	/** If enabled, this actor will be moved to follow a base actor that it's standing on. Typically disabled for physics-based movement, which handles based movement internally. */
-	UPROPERTY(EditDefaultsOnly, Category = "Mover")
-	bool bSupportsKinematicBasedMovement = true;
-
 	/** Transient flag indicating whether we are executing OnRegister(). */
 	bool bInOnRegister = false;
 
@@ -550,13 +557,6 @@ private:
 	// Transient flag indicating we've had a rollback and haven't started simulating forward again yet
 	bool bHasRolledBack = false;
 
-	UPROPERTY(Transient)
-	TObjectPtr<UMovementModeStateMachine> ModeFSM;	// JAH TODO: Also consider allowing a type property on the component to allow an alternative machine implementation to be allocated/used
-
-	/** Used to store cached data & computations between decoupled systems, that can be referenced by name */
-	UPROPERTY(Transient)
-	TObjectPtr<UMoverBlackboard> SimBlackboard;
-
 	friend class UBaseMovementMode;
 	friend class UMoverNetworkPhysicsLiaisonComponent;
 	friend class UMoverDebugComponent;
