Mover: fixed crash that could occur when using a root motion attribute layered move on a Mover actor that hasn't produced any input.

diff --git a/Engine/Plugins/Experimental/Mover/Source/Mover/Private/DefaultMovementSet/LayeredMoves/RootMotionAttributeLayeredMove.cpp b/Engine/Plugins/Experimental/Mover/Source/Mover/Private/DefaultMovementSet/LayeredMoves/RootMotionAttributeLayeredMove.cpp
--- a/Engine/Plugins/Experimental/Mover/Source/Mover/Private/DefaultMovementSet/LayeredMoves/RootMotionAttributeLayeredMove.cpp
+++ b/Engine/Plugins/Experimental/Mover/Source/Mover/Private/DefaultMovementSet/LayeredMoves/RootMotionAttributeLayeredMove.cpp
@@ -125,7 +125,7 @@ bool FLayeredMove_RootMotionAttribute::GenerateMove(const FMoverTickStartData& S
 
 		// TODO: Provide better options to control when we should skip proposed movement
 		if (!bHasValidWarpingContext &&
-				(InputCmd->bIsJumpJustPressed || MoverComp->HasGameplayTag(Mover_IsFalling, false)))
+				((InputCmd && InputCmd->bIsJumpJustPressed) || MoverComp->HasGameplayTag(Mover_IsFalling, false)))
 		{
 			return false;	// do not perform root motion while jumping or falling, so that we can have air control, unless we're under montage influence
 		}
