Mover - Fix localization error

diff --git a/Engine/Plugins/Experimental/Mover/Source/Mover/Private/MoverComponent.cpp b/Engine/Plugins/Experimental/Mover/Source/Mover/Private/MoverComponent.cpp
--- a/Engine/Plugins/Experimental/Mover/Source/Mover/Private/MoverComponent.cpp
+++ b/Engine/Plugins/Experimental/Mover/Source/Mover/Private/MoverComponent.cpp
@@ -993,7 +993,7 @@ bool UMoverComponent::ValidateSetup(FDataValidationContext& Context) const
 				{
 					if (Type.Get() == nullptr)
 					{
-						Context.AddError(FText::Format(LOCTEXT("InvalidModeSettingsError", "Transition on mode {0} on {1}, has an invalid SharedSettingsClass. You may need to remove the invalid settings class."),
+						Context.AddError(FText::Format(LOCTEXT("InvalidModeTransitionSettingsError", "Transition on mode {0} on {1}, has an invalid SharedSettingsClass. You may need to remove the invalid settings class."),
 							FText::FromName(Element.Key),
 							FText::FromString(GetNameSafe(GetOwner()))));
 
@@ -1001,7 +1001,7 @@ bool UMoverComponent::ValidateSetup(FDataValidationContext& Context) const
 					}
 					else if (FindSharedSettings(Type) == nullptr)
 					{
-						Context.AddError(FText::Format(LOCTEXT("MissingModeSettingsError", "Transition on mode {0} on {1}, is missing its desired SharedSettingsClass {2}. You may need to save the asset and/or recompile."),
+						Context.AddError(FText::Format(LOCTEXT("MissingModeTransitionSettingsError", "Transition on mode {0} on {1}, is missing its desired SharedSettingsClass {2}. You may need to save the asset and/or recompile."),
 							FText::FromName(Element.Key),
 							FText::FromString(GetNameSafe(GetOwner())),
 							FText::FromString(Type->GetName())));
@@ -1038,14 +1038,14 @@ bool UMoverComponent::ValidateSetup(FDataValidationContext& Context) const
 		{
 			if (Type.Get() == nullptr)
 			{
-				Context.AddError(FText::Format(LOCTEXT("InvalidModeSettingsError", "Transition on {0}, has an invalid SharedSettingsClass. You may need to remove the invalid settings class."),
+				Context.AddError(FText::Format(LOCTEXT("InvalidTransitionSettingsError", "Transition on {0}, has an invalid SharedSettingsClass. You may need to remove the invalid settings class."),
 					FText::FromString(GetNameSafe(GetOwner()))));
 
 				bDidFindAnyProblems = true;
 			}
 			else if (FindSharedSettings(Type) == nullptr)
 			{
-				Context.AddError(FText::Format(LOCTEXT("MissingModeSettingsError", "Transition on {0}, is missing its desired SharedSettingsClass {2}. You may need to save the asset and/or recompile."),
+				Context.AddError(FText::Format(LOCTEXT("MissingTransitionSettingsError", "Transition on {0}, is missing its desired SharedSettingsClass {2}. You may need to save the asset and/or recompile."),
 					FText::FromString(GetNameSafe(GetOwner())),
 					FText::FromString(Type->GetName())));
 
