Mover - Ensure consistent turns when rotating by 180 degrees

diff --git a/Engine/Plugins/Experimental/Mover/Source/Mover/Private/MoveLibrary/MovementUtils.cpp b/Engine/Plugins/Experimental/Mover/Source/Mover/Private/MoveLibrary/MovementUtils.cpp
--- a/Engine/Plugins/Experimental/Mover/Source/Mover/Private/MoveLibrary/MovementUtils.cpp
+++ b/Engine/Plugins/Experimental/Mover/Source/Mover/Private/MoveLibrary/MovementUtils.cpp
@@ -458,6 +458,22 @@ FRotator UMovementUtils::ComputeAngularVelocity(const FRotator& FromOrientation,
 
 		AngularDelta.GetWindingAndRemainder(Winding, Remainder);	// to find the fastest turn, just get the (-180,180] remainder
 
+		// Ensure consistent turn direction when rotating by 180 degrees
+		// This can be important in avoiding client/server desyncs for small
+		// angular differences
+		if (Remainder.Yaw < (-180.0 + UE_KINDA_SMALL_NUMBER))
+		{
+			Remainder.Yaw = 180.0f;
+		}
+		if (Remainder.Pitch < (-180.0 + UE_KINDA_SMALL_NUMBER))
+		{
+			Remainder.Pitch = 180.0f;
+		}
+		if (Remainder.Roll < (-180.0 + UE_KINDA_SMALL_NUMBER))
+		{
+			Remainder.Roll = 180.0f;
+		}
+
 		AngularVelocityDpS = Remainder * (1.f / DeltaSeconds);
 
 		if (TurningRateLimit >= 0.f)
