Mover: adding NavWalkingMode to support more efficient walking movement for actors using a navmesh

diff --git a/Engine/Plugins/Experimental/Mover/Source/Mover/Private/DefaultMovementSet/Modes/NavWalkingMode.cpp b/Engine/Plugins/Experimental/Mover/Source/Mover/Private/DefaultMovementSet/Modes/NavWalkingMode.cpp
new file mode 100644
--- /dev/null
+++ b/Engine/Plugins/Experimental/Mover/Source/Mover/Private/DefaultMovementSet/Modes/NavWalkingMode.cpp
@@ -0,0 +1,597 @@
+﻿// Copyright Epic Games, Inc. All Rights Reserved.
+
+
+#include "DefaultMovementSet/Modes/NavWalkingMode.h"
+#include "MoverComponent.h"
+#include "AI/NavigationSystemBase.h"
+#include "AI/Navigation/NavigationDataInterface.h"
+#include "Components/ShapeComponent.h"
+#include "DefaultMovementSet/NavMoverComponent.h"
+#include "DefaultMovementSet/Settings/CommonLegacyMovementSettings.h"
+#include "MoveLibrary/FloorQueryUtils.h"
+#include "MoveLibrary/GroundMovementUtils.h"
+#include "MoveLibrary/ModularMovement.h"
+#include "MoveLibrary/MovementUtils.h"
+
+
+UNavWalkingMode::UNavWalkingMode()
+	: bSweepWhileNavWalking(true)
+	, bProjectNavMeshWalking(false)
+	, NavMeshProjectionHeightScaleUp(0.67f)
+	, NavMeshProjectionHeightScaleDown(1.0f)
+	, NavMeshProjectionInterval(0.1f)
+	, NavMeshProjectionInterpSpeed(12.f)
+	, NavMeshProjectionTimer(0)
+	, NavMoverComponent(nullptr)
+	, bProjectNavMeshOnBothWorldChannels(true)
+{
+	GameplayTags.AddTag(Mover_IsOnGround);
+	GameplayTags.AddTag(Mover_IsNavWalking);
+}
+
+void UNavWalkingMode::OnGenerateMove(const FMoverTickStartData& StartState, const FMoverTimeStep& TimeStep, FProposedMove& OutProposedMove) const
+{
+	const UMoverComponent* MoverComp = GetMoverComponent();
+	const FCharacterDefaultInputs* CharacterInputs = StartState.InputCmd.InputCollection.FindDataByType<FCharacterDefaultInputs>();
+	const FMoverDefaultSyncState* StartingSyncState = StartState.SyncState.SyncStateCollection.FindDataByType<FMoverDefaultSyncState>();
+	check(StartingSyncState);
+
+	const float DeltaSeconds = TimeStep.StepMs * 0.001f;
+	FFloorCheckResult LastFloorResult;
+	FVector MovementNormal;
+
+	UMoverBlackboard* SimBlackboard = MoverComp->GetSimBlackboard_Mutable();
+
+	// Try to use the floor as the basis for the intended move direction (i.e. try to walk along slopes, rather than into them)
+	if (SimBlackboard && SimBlackboard->TryGet(CommonBlackboard::LastFloorResult, LastFloorResult) && LastFloorResult.IsWalkableFloor())
+	{
+		MovementNormal = LastFloorResult.HitResult.ImpactNormal;
+	}
+	else
+	{
+		MovementNormal = MoverComp->GetUpDirection();
+	}
+	
+	FRotator IntendedOrientation_WorldSpace;
+	// If there's no intent from input to change orientation, use the current orientation
+	if (!CharacterInputs || CharacterInputs->OrientationIntent.IsNearlyZero())
+	{
+		IntendedOrientation_WorldSpace = StartingSyncState->GetOrientation_WorldSpace();
+	}
+	else
+	{
+		IntendedOrientation_WorldSpace = CharacterInputs->GetOrientationIntentDir_WorldSpace().ToOrientationRotator();
+	}
+	
+	FGroundMoveParams Params;
+
+	if (CharacterInputs)
+	{
+		Params.MoveInputType = CharacterInputs->GetMoveInputType();
+		Params.MoveInput = CharacterInputs->GetMoveInput_WorldSpace();
+	}
+	else
+	{
+		Params.MoveInputType = EMoveInputType::Invalid;
+		Params.MoveInput = FVector::ZeroVector;
+	}
+
+	Params.OrientationIntent = IntendedOrientation_WorldSpace;
+	Params.PriorVelocity = FVector::VectorPlaneProject(StartingSyncState->GetVelocity_WorldSpace(), MovementNormal);
+	Params.PriorOrientation = StartingSyncState->GetOrientation_WorldSpace();
+	Params.GroundNormal = MovementNormal;
+	Params.TurningRate = CommonLegacySettings->TurningRate;
+	Params.TurningBoost = CommonLegacySettings->TurningBoost;
+	Params.MaxSpeed = CommonLegacySettings->MaxSpeed;
+	Params.Acceleration = CommonLegacySettings->Acceleration;
+	Params.Deceleration = CommonLegacySettings->Deceleration;
+	Params.DeltaSeconds = DeltaSeconds;
+	
+	if (Params.MoveInput.SizeSquared() > 0.f && !UMovementUtils::IsExceedingMaxSpeed(Params.PriorVelocity, CommonLegacySettings->MaxSpeed))
+	{
+		Params.Friction = CommonLegacySettings->GroundFriction;
+	}
+	else
+	{
+		Params.Friction = CommonLegacySettings->bUseSeparateBrakingFriction ? CommonLegacySettings->BrakingFriction : CommonLegacySettings->GroundFriction;
+		Params.Friction *= CommonLegacySettings->BrakingFrictionFactor;
+	}
+
+	OutProposedMove = UGroundMovementUtils::ComputeControlledGroundMove(Params);
+
+	if (TurnGenerator)
+	{
+		OutProposedMove.AngularVelocity = ITurnGeneratorInterface::Execute_GetTurn(TurnGenerator, IntendedOrientation_WorldSpace, StartState, *StartingSyncState, TimeStep, OutProposedMove, SimBlackboard);
+	}
+}
+
+void UNavWalkingMode::OnSimulationTick(const FSimulationTickParams& Params, FMoverTickEndData& OutputState)
+{
+	const UMoverComponent* MoverComp = Cast<UMoverComponent>(GetMoverComponent());
+	ensureMsgf(MoverComp, TEXT("Nav Walking Mode couldn't find a valid MoverComponent!"));
+	
+	const FMoverTickStartData& StartState = Params.StartState;
+	USceneComponent* UpdatedComponent = Params.MovingComps.UpdatedComponent.Get();
+	UPrimitiveComponent* UpdatedPrimitive = Params.MovingComps.UpdatedPrimitive.Get();
+	FProposedMove ProposedMove = Params.ProposedMove;
+	FVector UpDirection = MoverComp->GetUpDirection();
+
+	if (!UpdatedComponent || !UpdatedPrimitive)
+	{
+		return;
+	}
+
+	const FCharacterDefaultInputs* CharacterInputs = StartState.InputCmd.InputCollection.FindDataByType<FCharacterDefaultInputs>();
+	const FMoverDefaultSyncState* StartingSyncState = StartState.SyncState.SyncStateCollection.FindDataByType<FMoverDefaultSyncState>();
+	check(StartingSyncState);
+
+	FMoverDefaultSyncState& OutputSyncState = OutputState.SyncState.SyncStateCollection.FindOrAddMutableDataByType<FMoverDefaultSyncState>();
+	
+	const float DeltaSeconds = Params.TimeStep.StepMs * 0.001f;
+	const FVector OrigMoveDelta = ProposedMove.LinearVelocity * DeltaSeconds;
+
+	TObjectPtr<AActor> OwnerActor = UpdatedComponent->GetOwner();
+	check(OwnerActor);
+
+	FMovementRecord MoveRecord;
+	MoveRecord.SetDeltaSeconds(DeltaSeconds);
+
+	OutputSyncState.MoveDirectionIntent = (ProposedMove.bHasDirIntent ? ProposedMove.DirectionIntent : FVector::ZeroVector);
+
+	const FRotator StartingOrient = StartingSyncState->GetOrientation_WorldSpace();
+	FRotator TargetOrient = StartingOrient;
+
+	// Apply orientation changes (if any)
+	if (!UMovementUtils::IsAngularVelocityZero(ProposedMove.AngularVelocity))
+	{ 
+		TargetOrient += (ProposedMove.AngularVelocity * DeltaSeconds);
+	}
+	
+	const FQuat TargetOrientQuat = TargetOrient.Quaternion();
+
+	
+	FVector StartingLocation = NavMoverComponent->GetFeetLocation();
+	FVector AdjustedDest = StartingLocation + OrigMoveDelta;
+	
+	FFloorCheckResult CurrentFloor; 
+	const bool bDeltaMoveNearlyZero = OrigMoveDelta.IsNearlyZero();
+	FNavLocation DestNavLocation;
+
+	float SimpleRadius = 0;
+	float SimpleHalfHeight = 0;
+	NavMoverComponent->GetSimpleCollisionCylinder(SimpleRadius, SimpleHalfHeight);
+	
+	bool bSameNavLocation = false;
+	if (CachedNavLocation.NodeRef != INVALID_NAVNODEREF)
+	{
+		if (bProjectNavMeshWalking)
+		{
+			const float DistSq = UMovementUtils::ProjectToGravityFloor(StartingLocation - CachedNavLocation.Location, UpDirection).SizeSquared();
+			const float DistDot = FMath::Abs((StartingLocation - CachedNavLocation.Location).Dot(UpDirection));
+			
+			const float TotalCapsuleHeight = SimpleHalfHeight * 2.0f;
+			const float ProjectionScale = (StartingLocation.Dot(UpDirection) > CachedNavLocation.Location.Dot(UpDirection)) ? NavMeshProjectionHeightScaleUp : NavMeshProjectionHeightScaleDown;
+			const float DistThr = TotalCapsuleHeight * FMath::Max(0.f, ProjectionScale);
+
+			bSameNavLocation = (DistSq <= UE_KINDA_SMALL_NUMBER) && (DistDot < DistThr);
+		}
+		else
+		{
+			bSameNavLocation = CachedNavLocation.Location.Equals(StartingLocation);
+		}
+
+		if (bDeltaMoveNearlyZero && bSameNavLocation)
+		{
+			if (const INavigationDataInterface* NavData = GetNavData())
+			{
+				if (!NavData->IsNodeRefValid(CachedNavLocation.NodeRef))
+				{
+					CachedNavLocation.NodeRef = INVALID_NAVNODEREF;
+					bSameNavLocation = false;
+				}
+			}
+		}
+	}
+
+	if (bDeltaMoveNearlyZero && bSameNavLocation)
+	{
+		DestNavLocation = CachedNavLocation;
+		UE_LOG(LogMover, VeryVerbose, TEXT("%s using cached navmesh location! (bProjectNavMeshWalking = %d)"), *GetNameSafe(GetMoverComponent()->GetOwner()), bProjectNavMeshWalking);
+	}
+	else
+	{
+		// Start the trace from the vertical location of the last valid trace.
+		// Otherwise if we are projecting our location to the underlying geometry and it's far above or below the navmesh,
+		// we'll follow that geometry's plane out of range of valid navigation.
+		if (bSameNavLocation && bProjectNavMeshWalking)
+		{
+			UMovementUtils::SetGravityVerticalComponent(AdjustedDest, CachedNavLocation.Location.Dot(UpDirection), UpDirection);
+		}
+
+		// Find the point on the NavMesh
+		const bool bHasNavigationData = FindNavFloor(AdjustedDest, DestNavLocation);
+		if (!bHasNavigationData)
+		{
+			OutputState.MovementEndState.NextModeName = DefaultModeNames::Walking;
+			OutputState.MovementEndState.RemainingMs = Params.TimeStep.StepMs;
+			MoveRecord.SetDeltaSeconds((Params.TimeStep.StepMs - OutputState.MovementEndState.RemainingMs) * 0.001f);
+			
+			CaptureFinalState(UpdatedComponent, true, CurrentFloor, MoveRecord, OutputSyncState);
+			return;
+		}
+
+		CachedNavLocation = DestNavLocation;
+	}
+	
+	if (DestNavLocation.NodeRef != INVALID_NAVNODEREF)
+	{
+		FVector NewLocation = UMovementUtils::ProjectToGravityFloor(AdjustedDest, UpDirection) + UMovementUtils::GetGravityVerticalComponent(DestNavLocation.Location, UpDirection);
+		if (bProjectNavMeshWalking)
+		{
+			const float TotalCapsuleHeight = SimpleHalfHeight * 2.0f;
+			const float UpOffset = TotalCapsuleHeight * FMath::Max(0.f, NavMeshProjectionHeightScaleUp);
+			const float DownOffset = TotalCapsuleHeight * FMath::Max(0.f, NavMeshProjectionHeightScaleDown);
+			NewLocation = ProjectLocationFromNavMesh(DeltaSeconds, StartingLocation, NewLocation, UpOffset, DownOffset);
+		}
+
+		FVector AdjustedDelta = NewLocation - StartingLocation;
+		
+		if (!AdjustedDelta.IsNearlyZero())
+		{
+			FHitResult MoveHitResult;
+			UMovementUtils::TrySafeMoveUpdatedComponent(Params.MovingComps, AdjustedDelta, TargetOrientQuat, bSweepWhileNavWalking, MoveHitResult, ETeleportType::None, MoveRecord);
+		}
+	}
+	else
+	{
+		OutputState.MovementEndState.NextModeName = CommonLegacySettings->AirMovementModeName;
+		OutputState.MovementEndState.RemainingMs = Params.TimeStep.StepMs;
+		MoveRecord.SetDeltaSeconds((Params.TimeStep.StepMs - OutputState.MovementEndState.RemainingMs) * 0.001f);
+	}
+
+	CaptureFinalState(UpdatedComponent, true, CurrentFloor, MoveRecord, OutputSyncState);
+}
+
+bool UNavWalkingMode::FindNavFloor(const FVector& TestLocation, FNavLocation& NavFloorLocation) const
+{
+	const INavigationDataInterface* NavData = GetNavData();
+	if (NavData == nullptr || NavMoverComponent == nullptr)
+	{
+		return false;
+	}
+
+	const FNavAgentProperties& AgentProps = NavMoverComponent->GetNavAgentPropertiesRef();
+	const float SearchRadius = AgentProps.AgentRadius * 2.0f;
+	const float SearchHeight = AgentProps.AgentHeight * AgentProps.NavWalkingSearchHeightScale;
+
+	return NavData->ProjectPoint(TestLocation, NavFloorLocation, FVector(SearchRadius, SearchRadius, SearchHeight));
+}
+
+UObject* UNavWalkingMode::GetTurnGenerator()
+{
+	return TurnGenerator;
+}
+
+void UNavWalkingMode::SetTurnGeneratorClass(TSubclassOf<UObject> TurnGeneratorClass)
+{
+	if (TurnGeneratorClass)
+	{
+		TurnGenerator = NewObject<UObject>(this, TurnGeneratorClass);
+	}
+	else
+	{
+		TurnGenerator = nullptr; // Clearing the turn generator is valid - will go back to the default turn generation
+	}
+}
+
+void UNavWalkingMode::SetCollisionForNavWalking(bool bEnable)
+{
+	if (const UMoverComponent* MoverComponent = GetMoverComponent())
+	{
+		if (UPrimitiveComponent* UpdatedCompAsPrimitive = Cast<UPrimitiveComponent>(MoverComponent->GetUpdatedComponent()))
+		{
+			if (bEnable)
+			{
+				// TODO NS: Eventually might want make the new collision response an option so they trigger certain things while still not colliding 
+				UpdatedCompAsPrimitive->SetCollisionResponseToChannel(ECC_WorldStatic, ECR_Ignore);
+				UpdatedCompAsPrimitive->SetCollisionResponseToChannel(ECC_WorldDynamic, ECR_Ignore);
+			
+				if (UNavWalkingMode* NavWalkingMode = Cast<UNavWalkingMode>(MoverComponent->FindMovementMode(UNavWalkingMode::StaticClass())))
+				{
+					if (UMoverBlackboard* SimBlackboard = MoverComponent->GetSimBlackboard_Mutable())
+					{
+						SimBlackboard->Invalidate(NavWalkingBlackBoard::ProjectedNavMeshHitResult);
+					}
+					
+					// Stagger timed updates so many different characters spawned at the same time don't update on the same frame.
+					// Initially we want an immediate update though, so set time to a negative randomized range.
+					NavWalkingMode->NavMeshProjectionTimer = (NavWalkingMode->NavMeshProjectionInterval > 0.f) ? FMath::FRandRange(-NavWalkingMode->NavMeshProjectionInterval, 0.f) : 0.f;
+				}
+			}
+			else
+			{
+				// Grabbing the original shape settings and reverting to our old collision responses
+				if (const UShapeComponent* OriginalShapeComp = UMovementUtils::GetOriginalComponentType<UShapeComponent>(MoverComponent->GetOwner()))
+				{
+					UpdatedCompAsPrimitive->SetCollisionResponseToChannel(ECC_WorldStatic, OriginalShapeComp->GetCollisionResponseToChannel(ECC_WorldStatic));
+					UpdatedCompAsPrimitive->SetCollisionResponseToChannel(ECC_WorldDynamic, OriginalShapeComp->GetCollisionResponseToChannel(ECC_WorldDynamic));
+				}
+			}
+		}
+	}
+}
+
+void UNavWalkingMode::OnActivate()
+{
+	SetCollisionForNavWalking(true);
+}
+
+void UNavWalkingMode::OnDeactivate()
+{
+	SetCollisionForNavWalking(false);
+}
+
+const INavigationDataInterface* UNavWalkingMode::GetNavData() const
+{
+	const UWorld* World = GetWorld();
+	
+	const AActor* MoverComponentOwner;
+	if (const UMoverComponent* MoverComponent = GetMoverComponent())
+	{
+		MoverComponentOwner = MoverComponent->GetOwner();
+		if (World == nullptr || World->GetNavigationSystem() == nullptr || MoverComponentOwner == nullptr)
+		{
+			return nullptr;
+		}
+	}
+	else
+	{
+		return nullptr;
+	}
+	
+
+	const INavigationDataInterface* NavData = FNavigationSystem::GetNavDataForActor(*MoverComponentOwner);
+
+	return NavData;
+}
+
+void UNavWalkingMode::FindBestNavMeshLocation(const FVector& TraceStart, const FVector& TraceEnd, const FVector& CurrentFeetLocation, const FVector& TargetNavLocation, FHitResult& OutHitResult) const
+{
+	// raycast to underlying mesh to allow us to more closely follow geometry
+	// we use static objects here as a best approximation to accept only objects that
+	// influence navmesh generation
+	FCollisionQueryParams Params(SCENE_QUERY_STAT(ProjectLocation), false);
+
+	// blocked by world static and optionally world dynamic
+	FCollisionResponseParams ResponseParams(ECR_Ignore);
+	ResponseParams.CollisionResponse.SetResponse(ECC_WorldStatic, ECR_Overlap);
+	ResponseParams.CollisionResponse.SetResponse(ECC_WorldDynamic, bProjectNavMeshOnBothWorldChannels ? ECR_Overlap : ECR_Ignore);
+
+	TArray<FHitResult> MultiTraceHits;
+	GetWorld()->LineTraceMultiByChannel(MultiTraceHits, TraceStart, TraceEnd, ECC_WorldStatic, Params, ResponseParams);
+
+	struct FCompareFHitResultNavMeshTrace
+	{
+		explicit FCompareFHitResultNavMeshTrace(const FVector& inSourceLocation) : SourceLocation(inSourceLocation)
+		{
+		}
+
+		FORCEINLINE bool operator()(const FHitResult& A, const FHitResult& B) const
+		{
+			const float ADistSqr = (SourceLocation - A.ImpactPoint).SizeSquared();
+			const float BDistSqr = (SourceLocation - B.ImpactPoint).SizeSquared();
+
+			return (ADistSqr < BDistSqr);
+		}
+
+		const FVector& SourceLocation;
+	};
+
+	struct FRemoveNotBlockingResponseNavMeshTrace
+	{
+		FRemoveNotBlockingResponseNavMeshTrace(bool bInCheckOnlyWorldStatic) : bCheckOnlyWorldStatic(bInCheckOnlyWorldStatic) {}
+
+		FORCEINLINE bool operator()(const FHitResult& TestHit) const
+		{
+			UPrimitiveComponent* PrimComp = TestHit.GetComponent();
+			const bool bBlockOnWorldStatic = PrimComp && (PrimComp->GetCollisionResponseToChannel(ECC_WorldStatic) == ECR_Block);
+			const bool bBlockOnWorldDynamic = PrimComp && (PrimComp->GetCollisionResponseToChannel(ECC_WorldDynamic) == ECR_Block);
+
+			return !bBlockOnWorldStatic && (!bBlockOnWorldDynamic || bCheckOnlyWorldStatic);
+		}
+
+		bool bCheckOnlyWorldStatic;
+	};
+
+	MultiTraceHits.RemoveAllSwap(FRemoveNotBlockingResponseNavMeshTrace(!bProjectNavMeshOnBothWorldChannels), EAllowShrinking::No);
+	if (MultiTraceHits.Num() > 0)
+	{
+		// Sort the hits by the closest to our origin.
+		MultiTraceHits.Sort(FCompareFHitResultNavMeshTrace(TargetNavLocation));
+
+		// Cache the closest hit and treat it as a blocking hit (we used an overlap to get all the world static hits so we could sort them ourselves)
+		OutHitResult = MultiTraceHits[0];
+		OutHitResult.bBlockingHit = true;
+	}
+}
+
+FVector UNavWalkingMode::ProjectLocationFromNavMesh(float DeltaSeconds, const FVector& CurrentFeetLocation, const FVector& TargetNavLocation, float UpOffset, float DownOffset)
+{
+	FVector NewLocation = TargetNavLocation;
+
+	const float VerticalOffset = -(DownOffset + UpOffset);
+	if (VerticalOffset > -UE_SMALL_NUMBER)
+	{
+		return NewLocation;
+	}
+
+	const UMoverComponent* MoverComp = GetMoverComponent();
+	const FVector UpDirection = MoverComp->GetUpDirection();
+	
+	const FVector TraceStart = TargetNavLocation + UpOffset * UpDirection;
+	const FVector TraceEnd   = TargetNavLocation + DownOffset * -UpDirection;
+
+	FHitResult CachedProjectedNavMeshHitResult;
+	UMoverBlackboard* SimBlackboard = MoverComp->GetSimBlackboard_Mutable();
+	SimBlackboard->TryGet(NavWalkingBlackBoard::ProjectedNavMeshHitResult, CachedProjectedNavMeshHitResult);
+	
+	// We can skip this trace if we are checking at the same location as the last trace (ie, we haven't moved).
+	const bool bCachedLocationStillValid = (CachedProjectedNavMeshHitResult.bBlockingHit &&
+											CachedProjectedNavMeshHitResult.TraceStart == TraceStart &&
+											CachedProjectedNavMeshHitResult.TraceEnd == TraceEnd);
+
+	NavMeshProjectionTimer -= DeltaSeconds;
+	if (NavMeshProjectionTimer <= 0.0f)
+	{
+		if (!bCachedLocationStillValid)
+		{
+			UE_LOG(LogMover, VeryVerbose, TEXT("ProjectLocationFromNavMesh(): %s interval: %.3f [SKIP TRACE]"), *GetNameSafe(GetMoverComponent()->GetOwner()), NavMeshProjectionInterval);
+
+			FHitResult HitResult;
+			FindBestNavMeshLocation(TraceStart, TraceEnd, CurrentFeetLocation, TargetNavLocation, HitResult);
+
+			// discard result if we were already inside something
+			if (HitResult.bStartPenetrating || !HitResult.bBlockingHit)
+			{
+				CachedProjectedNavMeshHitResult.Reset();
+				SimBlackboard->Invalidate(NavWalkingBlackBoard::ProjectedNavMeshHitResult);
+			}
+			else
+			{
+				CachedProjectedNavMeshHitResult = HitResult;
+				SimBlackboard->Set(NavWalkingBlackBoard::ProjectedNavMeshHitResult, CachedProjectedNavMeshHitResult);
+			}
+		}
+		else
+		{
+			UE_LOG(LogMover, VeryVerbose, TEXT("ProjectLocationFromNavMesh(): %s interval: %.3f [SKIP TRACE]"), *GetNameSafe(GetMoverComponent()->GetOwner()), NavMeshProjectionInterval);
+		}
+
+		// Wrap around to maintain same relative offset to tick time changes.
+		// Prevents large framerate spikes from aligning multiple characters to the same frame (if they start staggered, they will now remain staggered).
+		float ModTime = 0.f;
+		if (NavMeshProjectionInterval > UE_SMALL_NUMBER)
+		{
+			ModTime = FMath::Fmod(-NavMeshProjectionTimer, NavMeshProjectionInterval);
+		}
+
+		NavMeshProjectionTimer = NavMeshProjectionInterval - ModTime;
+	}
+	
+	// Project to last plane we found.
+	if (CachedProjectedNavMeshHitResult.bBlockingHit)
+	{
+		if (bCachedLocationStillValid && FMath::IsNearlyEqual(CurrentFeetLocation.Dot(UpDirection), CachedProjectedNavMeshHitResult.ImpactPoint.Dot(UpDirection), (FVector::FReal)0.01f))
+		{
+			// Already at destination.
+			UMovementUtils::SetGravityVerticalComponent(NewLocation, CurrentFeetLocation.Dot(UpDirection), UpDirection);
+		}
+		else
+		{
+			const FVector ProjectedPoint = FMath::LinePlaneIntersection(TraceStart, TraceEnd, CachedProjectedNavMeshHitResult.ImpactPoint, CachedProjectedNavMeshHitResult.Normal);
+			FVector::FReal ProjectedVertical = ProjectedPoint.Dot(UpDirection);
+			
+			// Limit to not be too far above or below NavMesh location
+			const FVector::FReal VertTraceStart = TraceStart.Dot(UpDirection);
+			const FVector::FReal VertTraceEnd = TraceEnd.Dot(UpDirection);
+			const FVector::FReal TraceMin = FMath::Min(VertTraceStart, VertTraceEnd);
+			const FVector::FReal TraceMax = FMath::Max(VertTraceStart, VertTraceEnd);
+			ProjectedVertical = FMath::Clamp(ProjectedVertical, TraceMin, TraceMax);
+
+			// Interp for smoother updates (less "pop" when trace hits something new). 0 interp speed is instant.
+			const FVector::FReal InterpSpeed = FMath::Max<FVector::FReal>(0.f, NavMeshProjectionInterpSpeed);
+			ProjectedVertical = FMath::FInterpTo(CurrentFeetLocation.Dot(UpDirection), ProjectedVertical, (FVector::FReal)DeltaSeconds, InterpSpeed);
+			ProjectedVertical = FMath::Clamp(ProjectedVertical, TraceMin, TraceMax);
+			
+			// Final result
+			UMovementUtils::SetGravityVerticalComponent(NewLocation, ProjectedVertical, UpDirection);
+		}
+	}
+
+	return NewLocation;
+}
+
+void UNavWalkingMode::OnRegistered(const FName ModeName)
+{
+	Super::OnRegistered(ModeName);
+
+	UMoverComponent* MoverComponent = GetMoverComponent();
+	CommonLegacySettings = MoverComponent->FindSharedSettings<UCommonLegacyMovementSettings>();
+	ensureMsgf(CommonLegacySettings, TEXT("Failed to find instance of CommonLegacyMovementSettings on %s. Movement may not function properly."), *GetPathNameSafe(this));
+	
+	if (const AActor* MoverCompOwner = MoverComponent->GetOwner())
+	{
+		NavMoverComponent = MoverCompOwner->FindComponentByClass<UNavMoverComponent>();
+	}
+	
+	if (!NavMoverComponent)
+	{
+		UE_LOG(LogMover, Warning, TEXT("NavWalkingMode on %s could not find a valid NavMoverComponent and will not function properly."), *GetNameSafe(GetMoverComponent()->GetOwner()));
+	}
+}
+
+void UNavWalkingMode::OnUnregistered()
+{
+	CommonLegacySettings = nullptr;
+
+	Super::OnUnregistered();
+}
+
+void UNavWalkingMode::CaptureFinalState(USceneComponent* UpdatedComponent, bool bDidAttemptMovement, const FFloorCheckResult& FloorResult, const FMovementRecord& Record, FMoverDefaultSyncState& OutputSyncState) const
+{
+	FRelativeBaseInfo PriorBaseInfo;
+
+	UMoverBlackboard* SimBlackboard = GetMoverComponent()->GetSimBlackboard_Mutable();
+
+	const bool bHasPriorBaseInfo = SimBlackboard->TryGet(CommonBlackboard::LastFoundDynamicMovementBase, PriorBaseInfo);
+
+	FRelativeBaseInfo CurrentBaseInfo = UpdateFloorAndBaseInfo(FloorResult);
+
+	// If we're on a dynamic base and we're not trying to move, keep using the same relative actor location. This prevents slow relative 
+	//  drifting that can occur from repeated floor sampling as the base moves through the world.
+	if (CurrentBaseInfo.HasRelativeInfo() 
+		&& bHasPriorBaseInfo && !bDidAttemptMovement 
+		&& PriorBaseInfo.UsesSameBase(CurrentBaseInfo))
+	{
+		CurrentBaseInfo.ContactLocalPosition = PriorBaseInfo.ContactLocalPosition;
+	}
+	
+	if (CurrentBaseInfo.HasRelativeInfo())
+	{
+		SimBlackboard->Set(CommonBlackboard::LastFoundDynamicMovementBase, CurrentBaseInfo);
+
+		OutputSyncState.SetTransforms_WorldSpace( UpdatedComponent->GetComponentLocation(),
+												  UpdatedComponent->GetComponentRotation(),
+												  Record.GetRelevantVelocity(),
+												  CurrentBaseInfo.MovementBase.Get(), CurrentBaseInfo.BoneName);
+	}
+	else
+	{
+		SimBlackboard->Invalidate(CommonBlackboard::LastFoundDynamicMovementBase);
+
+		OutputSyncState.SetTransforms_WorldSpace( UpdatedComponent->GetComponentLocation(),
+												  UpdatedComponent->GetComponentRotation(),
+												  Record.GetRelevantVelocity(),
+												  nullptr);	// no movement base
+	}
+
+	UpdatedComponent->ComponentVelocity = OutputSyncState.GetVelocity_WorldSpace();
+}
+
+FRelativeBaseInfo UNavWalkingMode::UpdateFloorAndBaseInfo(const FFloorCheckResult& FloorResult) const
+{
+	FRelativeBaseInfo ReturnBaseInfo;
+
+	UMoverBlackboard* SimBlackboard = GetMoverComponent()->GetSimBlackboard_Mutable();
+
+	SimBlackboard->Set(CommonBlackboard::LastFloorResult, FloorResult);
+
+	if (FloorResult.IsWalkableFloor() && UBasedMovementUtils::IsADynamicBase(FloorResult.HitResult.GetComponent()))
+	{
+		ReturnBaseInfo.SetFromFloorResult(FloorResult);
+	}
+
+	return ReturnBaseInfo;
+}
+
+
diff --git a/Engine/Plugins/Experimental/Mover/Source/Mover/Private/MoverTypes.cpp b/Engine/Plugins/Experimental/Mover/Source/Mover/Private/MoverTypes.cpp
--- a/Engine/Plugins/Experimental/Mover/Source/Mover/Private/MoverTypes.cpp
+++ b/Engine/Plugins/Experimental/Mover/Source/Mover/Private/MoverTypes.cpp
@@ -13,6 +13,7 @@ UE_DEFINE_GAMEPLAY_TAG_COMMENT(Mover_IsFalling, "Mover.IsFalling", "Default Move
 UE_DEFINE_GAMEPLAY_TAG_COMMENT(Mover_IsFlying, "Mover.IsFlying", "Default Mover state flag indicating character is flying.");
 UE_DEFINE_GAMEPLAY_TAG_COMMENT(Mover_IsSwimming, "Mover.IsSwimming", "Default Mover state flag indicating character is swimming.");
 UE_DEFINE_GAMEPLAY_TAG_COMMENT(Mover_IsCrouching, "Mover.Stance.IsCrouching", "Default Mover state flag indicating character is crouching.");
+UE_DEFINE_GAMEPLAY_TAG_COMMENT(Mover_IsNavWalking, "Mover.IsNavWalking", "Default Mover state flag indicating character is NavWalking.");
 
 FMoverOnImpactParams::FMoverOnImpactParams() 
 	: AttemptedMoveDelta(0) 
diff --git a/Engine/Plugins/Experimental/Mover/Source/Mover/Public/DefaultMovementSet/Modes/NavWalkingMode.h b/Engine/Plugins/Experimental/Mover/Source/Mover/Public/DefaultMovementSet/Modes/NavWalkingMode.h
new file mode 100644
--- /dev/null
+++ b/Engine/Plugins/Experimental/Mover/Source/Mover/Public/DefaultMovementSet/Modes/NavWalkingMode.h
@@ -0,0 +1,135 @@
+﻿// Copyright Epic Games, Inc. All Rights Reserved.
+
+#pragma once
+
+#include "CoreMinimal.h"
+#include "MovementMode.h"
+#include "MoveLibrary/BasedMovementUtils.h"
+#include "AI/Navigation/NavigationTypes.h"
+#include "NavWalkingMode.generated.h"
+
+class UNavMoverComponent;
+class INavigationDataInterface;
+class UCommonLegacyMovementSettings;
+
+namespace NavWalkingBlackBoard
+{
+	const FName ProjectedNavMeshHitResult = TEXT("ProjectedNavMeshHitResult");
+}
+
+/**
+ * NavWalkingMode: a default movement mode for traversing surfaces and movement bases by using an active navmesh when moving the actor rather than collision checks.
+ * Note: This movement mode requires a NavMoverComponent be on the actor to function properly. This mode also contains some randomization to avoid navmesh look ups
+ *	happening at the same time (which is fine for AI characters running on the server) but may cause issues if used on autonomous proxies.
+ */
+UCLASS(Blueprintable, BlueprintType)
+class MOVER_API UNavWalkingMode : public UBaseMovementMode
+{
+	GENERATED_BODY()
+
+public:
+	UNavWalkingMode();
+	
+	UFUNCTION(BlueprintCallable, Category=Mover)
+	virtual void OnGenerateMove(const FMoverTickStartData& StartState, const FMoverTimeStep& TimeStep, FProposedMove& OutProposedMove) const override;
+
+	UFUNCTION(BlueprintCallable, Category=Mover)
+	virtual void OnSimulationTick(const FSimulationTickParams& Params, FMoverTickEndData& OutputState) override;
+
+	/**
+	 * Whether or not the actor should sweep for collision geometry while walking.
+	 */
+	UPROPERTY(Category="NavMesh Movement", EditAnywhere, BlueprintReadWrite)
+	bool bSweepWhileNavWalking;
+	
+	/** Whether to raycast to underlying geometry to better conform navmesh-walking actors */
+	UPROPERTY(Category="NavMesh Movement", EditAnywhere, BlueprintReadOnly)
+	bool bProjectNavMeshWalking;
+
+	/**
+	 * Scale of the total capsule height to use for projection from navmesh to underlying geometry in the upward direction.
+	 * In other words, start the trace at [CapsuleHeight * NavMeshProjectionHeightScaleUp] above nav mesh.
+	 */
+	UPROPERTY(Category="NavMesh Movement", EditAnywhere, BlueprintReadWrite, meta=(editcondition = "bProjectNavMeshWalking", ClampMin="0", UIMin="0"))
+	float NavMeshProjectionHeightScaleUp;
+
+	/**
+	 * Scale of the total capsule height to use for projection from navmesh to underlying geometry in the downward direction.
+	 * In other words, trace down to [CapsuleHeight * NavMeshProjectionHeightScaleDown] below nav mesh.
+	 */
+	UPROPERTY(Category="NavMesh Movement", EditAnywhere, BlueprintReadWrite, meta=(editcondition = "bProjectNavMeshWalking", ClampMin="0", UIMin="0"))
+	float NavMeshProjectionHeightScaleDown;
+
+	/** How often we should raycast to project from navmesh to underlying geometry */
+	UPROPERTY(Category="NavMesh Movement", EditAnywhere, BlueprintReadWrite, meta=(editcondition = "bProjectNavMeshWalking", ForceUnits="s"))
+	float NavMeshProjectionInterval;
+
+	/** Speed at which to interpolate agent navmesh offset between traces. 0: Instant (no interp) > 0: Interp speed") */
+	UPROPERTY(Category="NavMesh Movement", EditAnywhere, BlueprintReadWrite, meta=(editcondition = "bProjectNavMeshWalking", ClampMin="0", UIMin="0"))
+	float NavMeshProjectionInterpSpeed;
+	
+	UPROPERTY(Transient)
+	float NavMeshProjectionTimer;
+	
+	/** last known location projected on navmesh */
+	FNavLocation CachedNavLocation;
+	
+	/**
+	 * Project a location to navmesh to find adjusted height.
+	 * @param TestLocation		Location to project
+	 * @param NavFloorLocation	Location on navmesh
+	 * @return True if projection was performed (successfully or not)
+	 */
+	virtual bool FindNavFloor(const FVector& TestLocation, FNavLocation& NavFloorLocation) const;
+
+	// Returns the active turn generator. Note: you will need to cast the return value to the generator you expect to get, it can also be none
+	UFUNCTION(BlueprintPure, Category=Mover)
+	UObject* GetTurnGenerator();
+
+	// Sets the active turn generator to use the class provided. Note: To set it back to the default implementation pass in none
+	UFUNCTION(BlueprintCallable, Category=Mover)
+	void SetTurnGeneratorClass(UPARAM(meta=(MustImplement="/Script/Mover.TurnGeneratorInterface", AllowAbstract="false")) TSubclassOf<UObject> TurnGeneratorClass);
+	
+protected:
+	/** Associated Movement component that will actually move the actor */ 
+	UPROPERTY(BlueprintReadOnly, Category="Nav Movement")
+	TObjectPtr<UNavMoverComponent> NavMoverComponent;
+	
+	/** Use both WorldStatic and WorldDynamic channels for NavWalking geometry conforming */
+	UPROPERTY(Category = "NavMesh Movement", EditAnywhere, BlueprintReadOnly, AdvancedDisplay)
+	uint8 bProjectNavMeshOnBothWorldChannels : 1;
+
+	/** Optional modular object for generating rotation towards desired orientation. If not specified, linear interpolation will be used. */
+	UPROPERTY(EditAnywhere, Instanced, Category=Mover, meta=(MustImplement="/Script/Mover.TurnGeneratorInterface"))
+	TObjectPtr<UObject> TurnGenerator;
+
+	/** Switch collision settings for NavWalking mode (ignore world collisions) */
+	virtual void SetCollisionForNavWalking(bool bEnable);
+	
+	virtual void OnActivate() override;
+
+	virtual void OnDeactivate() override;
+	
+	/** Get Navigation data for the actor. Returns null if there is no associated nav data. */
+	const INavigationDataInterface* GetNavData() const;
+
+	/** Performs trace for ProjectLocationFromNavMesh */
+	virtual void FindBestNavMeshLocation(const FVector& TraceStart, const FVector& TraceEnd, const FVector& CurrentFeetLocation, const FVector& TargetNavLocation, FHitResult& OutHitResult) const;
+
+	/** 
+	 * Attempts to better align navmesh walking actors with underlying geometry (sometimes 
+	 * navmesh can differ quite significantly from geometry).
+	 * Updates CachedProjectedNavMeshHitResult, access this for more info about hits.
+	 */
+	virtual FVector ProjectLocationFromNavMesh(float DeltaSeconds, const FVector& CurrentFeetLocation, const FVector& TargetNavLocation, float UpOffset, float DownOffset);
+
+	
+	virtual void OnRegistered(const FName ModeName) override;
+	virtual void OnUnregistered() override;
+	
+	void CaptureFinalState(USceneComponent* UpdatedComponent, bool bDidAttemptMovement, const FFloorCheckResult& FloorResult, const FMovementRecord& Record, FMoverDefaultSyncState& OutputSyncState) const;
+
+	FRelativeBaseInfo UpdateFloorAndBaseInfo(const FFloorCheckResult& FloorResult) const;
+	
+	TObjectPtr<const UCommonLegacyMovementSettings> CommonLegacySettings;
+};
diff --git a/Engine/Plugins/Experimental/Mover/Source/Mover/Public/MoveLibrary/MovementUtils.h b/Engine/Plugins/Experimental/Mover/Source/Mover/Public/MoveLibrary/MovementUtils.h
--- a/Engine/Plugins/Experimental/Mover/Source/Mover/Public/MoveLibrary/MovementUtils.h
+++ b/Engine/Plugins/Experimental/Mover/Source/Mover/Public/MoveLibrary/MovementUtils.h
@@ -6,6 +6,9 @@
 #include "MoverSimulationTypes.h"
 #include "MoverDataModelTypes.h"
 #include "Components/PrimitiveComponent.h"
+#include "Engine/Blueprint.h" // For gathering CDO info from a BP
+#include "Engine/SCS_Node.h" // For gathering CDO info from a BP
+#include "Engine/SimpleConstructionScript.h" // For gathering CDO info from a BP
 
 #include "MovementUtils.generated.h"
 
@@ -116,6 +119,10 @@ public:
 	// JAH TODO: Make sure all 'out' parameters are last in the param list and marked as "Out"
 	// JAH TODO: separate out the public-facing ones from the internally-used ones and make all public-facing ones BlueprintCallable
 
+	// Gets CDO component type - useful for getting original values
+	template <class ComponentType>
+	static const ComponentType* GetOriginalComponentType(const AActor* MoverCompOwner);
+	
 	/** Checks whether a given velocity is exceeding a maximum speed, with some leeway to account for numeric imprecision */
 	UFUNCTION(BlueprintCallable, Category = Mover)
 	static bool IsExceedingMaxSpeed(const FVector& Velocity, float InMaxSpeed);
@@ -142,6 +149,17 @@ public:
 	UFUNCTION(BlueprintCallable, Category = Mover)
 	static FVector ConstrainToPlane(const FVector& Vector, const FPlane& MovementPlane, bool bMaintainMagnitude=true);
 
+	/** Project a vector onto the floor defined by the gravity direction. */
+	UFUNCTION(BlueprintCallable, Category = Mover)
+	static FVector ProjectToGravityFloor(const FVector& Vector, const FVector& UpDirection) { return FVector::VectorPlaneProject(Vector, -UpDirection); }
+
+	/** Returns the component of the vector in the gravity-space vertical direction.  */
+	UFUNCTION(BlueprintCallable, Category = Mover)
+	static FVector GetGravityVerticalComponent(const FVector& Vector, const FVector& UpDirection) { return Vector.Dot(-UpDirection) * -UpDirection; }
+
+	/** Set the vertical component of the vector to the given value in the gravity-space vertical direction. */
+	static void SetGravityVerticalComponent(FVector& Vector, const FVector::FReal VerticalValue, const FVector& UpDirection) { Vector = ProjectToGravityFloor(Vector, UpDirection) - VerticalValue * -UpDirection; }
+
 	// Surface sliding
 
 	/** Returns an alternative move delta to slide along a surface, based on parameters describing a blocked attempted move */
@@ -194,3 +212,49 @@ public:
 	static bool TryMoveUpdatedComponent_Internal(const FMovingComponentSet& MovingComps, const FVector& Delta, const FQuat& NewRotation, bool bSweep, EMoveComponentFlags MoveComponentFlags, FHitResult* OutHit, ETeleportType Teleport);
 
 };
+
+template <class ComponentType>
+const ComponentType* UMovementUtils::GetOriginalComponentType(const AActor* MoverCompOwner)
+{
+	const ComponentType* OriginalComponent = nullptr;
+
+	if (MoverCompOwner)
+	{
+		if (const AActor* OwnerCDO = Cast<AActor>(MoverCompOwner->GetClass()->GetDefaultObject()))
+		{
+			// Check if native CDO has Capsule component
+			OriginalComponent = OwnerCDO->FindComponentByClass<ComponentType>();
+
+			// check if it comes from a BP
+			if (!OriginalComponent)
+			{
+				if (const UBlueprintGeneratedClass* OwnerClassAsBP = Cast<UBlueprintGeneratedClass>(OwnerCDO->GetClass()))
+				{
+					TArray<const UBlueprintGeneratedClass*> BlueprintClasses;
+					UBlueprintGeneratedClass::GetGeneratedClassesHierarchy(OwnerClassAsBP, BlueprintClasses);
+					for (const UBlueprintGeneratedClass* BlueprintClass : BlueprintClasses)
+					{
+						if (BlueprintClass->SimpleConstructionScript)
+						{
+							// Check Simple construction script
+							const TArray<USCS_Node*>& SCSNodes = BlueprintClass->SimpleConstructionScript->GetAllNodes();
+							for (USCS_Node* SCSNode : SCSNodes)
+							{
+								if (SCSNode)
+								{
+									if (const ComponentType* BPComponent = Cast<ComponentType>(SCSNode->ComponentTemplate))
+									{
+										OriginalComponent = BPComponent;
+										break;
+									}
+								}
+							}
+						}
+					}
+				}
+			}
+		}
+	}
+
+	return OriginalComponent;
+}
diff --git a/Engine/Plugins/Experimental/Mover/Source/Mover/Public/MoverTypes.h b/Engine/Plugins/Experimental/Mover/Source/Mover/Public/MoverTypes.h
--- a/Engine/Plugins/Experimental/Mover/Source/Mover/Public/MoverTypes.h
+++ b/Engine/Plugins/Experimental/Mover/Source/Mover/Public/MoverTypes.h
@@ -17,6 +17,7 @@ MOVER_API UE_DECLARE_GAMEPLAY_TAG_EXTERN(Mover_IsFalling);
 MOVER_API UE_DECLARE_GAMEPLAY_TAG_EXTERN(Mover_IsFlying);
 MOVER_API UE_DECLARE_GAMEPLAY_TAG_EXTERN(Mover_IsSwimming);
 MOVER_API UE_DECLARE_GAMEPLAY_TAG_EXTERN(Mover_IsCrouching);
+MOVER_API UE_DECLARE_GAMEPLAY_TAG_EXTERN(Mover_IsNavWalking);
 
 
 /** Options for how to handle smoothing frame data from the backend. Typically this is for advancing the simulation at a lower or fixed rate versus the game thread/render rate. */
diff --git a/Engine/Plugins/Experimental/MoverExamples/Source/MoverExamples/Private/MoverExamplesCharacter.cpp b/Engine/Plugins/Experimental/MoverExamples/Source/MoverExamples/Private/MoverExamplesCharacter.cpp
--- a/Engine/Plugins/Experimental/MoverExamples/Source/MoverExamples/Private/MoverExamplesCharacter.cpp
+++ b/Engine/Plugins/Experimental/MoverExamples/Source/MoverExamples/Private/MoverExamplesCharacter.cpp
@@ -246,12 +246,12 @@ void AMoverExamplesCharacter::OnProduceInput(float DeltaMs, FMoverInputCmdContex
 		if (bOrientRotationToMovement)
 		{
 			// set the intent to the actors movement direction
-			CharacterInputs.OrientationIntent = CharacterInputs.GetMoveInput();
+			CharacterInputs.OrientationIntent = CharacterInputs.GetMoveInput().GetSafeNormal();
 		}
 		else
 		{
 			// set intent to the the control rotation - often a player's camera rotation
-			CharacterInputs.OrientationIntent = CharacterInputs.ControlRotation.Vector();
+			CharacterInputs.OrientationIntent = CharacterInputs.ControlRotation.Vector().GetSafeNormal();
 		}
 
 		LastAffirmativeMoveInput = CharacterInputs.GetMoveInput();
