Mover: fixing crash/check cases where a movement simulation backend would be created for an object that won't be part of the game world, such as a thumbnail

diff --git a/Engine/Plugins/Experimental/Mover/Source/Mover/Private/MoverComponent.cpp b/Engine/Plugins/Experimental/Mover/Source/Mover/Private/MoverComponent.cpp
--- a/Engine/Plugins/Experimental/Mover/Source/Mover/Private/MoverComponent.cpp
+++ b/Engine/Plugins/Experimental/Mover/Source/Mover/Private/MoverComponent.cpp
@@ -68,23 +68,29 @@ UMoverComponent::UMoverComponent()
 void UMoverComponent::InitializeComponent()
 {
 	TGuardValue<bool> InInitializeComponentGuard(bInInitializeComponent, true);
-	FindDefaultUpdatedComponent();
-	
-	// Instantiate out sister backend component that will actually talk to the system driving the simulation
-	if (BackendClass)
+
+	const UWorld* MyWorld = GetWorld();
+
+	if (MyWorld && MyWorld->IsGameWorld())
 	{
-		UActorComponent* NewLiaisonComp = NewObject<UActorComponent>(GetOwner(), BackendClass, TEXT("BackendLiaisonComponent"));
-		BackendLiaisonComp = CastChecked<IMoverBackendLiaisonInterface>(NewLiaisonComp);
-		if (BackendLiaisonComp.Get())
+		FindDefaultUpdatedComponent();
+
+		// Instantiate out sister backend component that will actually talk to the system driving the simulation
+		if (BackendClass)
 		{
-			NewLiaisonComp->RegisterComponent();
-			NewLiaisonComp->InitializeComponent();
-			NewLiaisonComp->SetNetAddressable();
+			UActorComponent* NewLiaisonComp = NewObject<UActorComponent>(GetOwner(), BackendClass, TEXT("BackendLiaisonComponent"));
+			BackendLiaisonComp = CastChecked<IMoverBackendLiaisonInterface>(NewLiaisonComp);
+			if (BackendLiaisonComp.Get())
+			{
+				NewLiaisonComp->RegisterComponent();
+				NewLiaisonComp->InitializeComponent();
+				NewLiaisonComp->SetNetAddressable();
+			}
+		}
+		else
+		{
+			UE_LOG(LogMover, Error, TEXT("No backend class set on %s. Mover actor will not function."), *GetNameSafe(GetOwner()));
 		}
-	}
-	else
-	{
-		UE_LOG(LogMover, Error, TEXT("No backend class set on %s. Mover actor will not function."), *GetNameSafe(GetOwner()));
 	}
 
 	Super::InitializeComponent();
