Mover: improving method of auto-finding a primary visual component by favoring any mesh that's a direct child of the root updated component. This can prevent inconsistencies between cooked and non-cooked builds where the first-found mesh could be different.

diff --git a/Engine/Plugins/Experimental/Mover/Source/Mover/Private/MoverComponent.cpp b/Engine/Plugins/Experimental/Mover/Source/Mover/Private/MoverComponent.cpp
--- a/Engine/Plugins/Experimental/Mover/Source/Mover/Private/MoverComponent.cpp
+++ b/Engine/Plugins/Experimental/Mover/Source/Mover/Private/MoverComponent.cpp
@@ -169,19 +169,30 @@ void UMoverComponent::BeginPlay()
 
 	if (const AActor* MyActor = GetOwner())
 	{
-		// If no primary visual component is already set, fall back to searching for any kind of mesh
+		// If no primary visual component is already set, fall back to searching for any kind of mesh,
+		// favoring a direct scene child of the UpdatedComponent.
 		if (!PrimaryVisualComponent)
 		{
-			SetPrimaryVisualComponent(MyActor->FindComponentByClass<UMeshComponent>());
+			if (UpdatedComponent)
+			{
+				for (USceneComponent* ChildComp : UpdatedComponent->GetAttachChildren())
+				{
+					if (ChildComp->IsA<UMeshComponent>())
+					{
+						SetPrimaryVisualComponent(ChildComp);
+						break;
+					}
+				}
+			}
+
+			if (!PrimaryVisualComponent)
+			{
+				SetPrimaryVisualComponent(MyActor->FindComponentByClass<UMeshComponent>());
+			}
 		}
 
 		ensureMsgf(UpdatedComponent && (PrimaryVisualComponent != UpdatedComponent), TEXT("A Mover actor (%s) must have an UpdatedComponent and cannot have a PrimaryVisualComponent that is the same as UpdatedComponent"), *GetNameSafe(MyActor));
 
-		if (PrimaryVisualComponent)
-		{
-			BaseVisualComponentTransform = PrimaryVisualComponent->GetRelativeTransform();
-		}
-
 		// Optional motion warping support
 		if (UMotionWarpingComponent* WarpingComp = MyActor->FindComponentByClass<UMotionWarpingComponent>())
 		{
@@ -1609,9 +1620,16 @@ USceneComponent* UMoverComponent::GetPrimaryVisualComponent() const
 
 void UMoverComponent::SetPrimaryVisualComponent(USceneComponent* SceneComponent)
 {
-	if (ensure(SceneComponent->GetOwner() == GetOwner()))
+	if (SceneComponent && 
+		ensureMsgf(SceneComponent->GetOwner() == GetOwner(), TEXT("Primary visual component must be owned by the same actor. MoverComp owner: %s  VisualComp owner: %s"), *GetNameSafe(GetOwner()), *GetNameSafe(SceneComponent->GetOwner())))
 	{
 		PrimaryVisualComponent = SceneComponent;
+		BaseVisualComponentTransform = SceneComponent->GetRelativeTransform();
+	}
+	else
+	{
+		PrimaryVisualComponent = nullptr;
+		BaseVisualComponentTransform = FTransform::Identity;
 	}
 }
 
