Mover - Add Reset function to FMoverSyncState

diff --git a/Engine/Plugins/Experimental/Mover/Source/Mover/Private/LayeredMove.cpp b/Engine/Plugins/Experimental/Mover/Source/Mover/Private/LayeredMove.cpp
--- a/Engine/Plugins/Experimental/Mover/Source/Mover/Private/LayeredMove.cpp
+++ b/Engine/Plugins/Experimental/Mover/Source/Mover/Private/LayeredMove.cpp
@@ -477,3 +477,10 @@ void FLayeredMoveGroup::ResetResidualVelocity()
 	ResidualClamping = -1.f;
 }
 
+void FLayeredMoveGroup::Reset()
+{
+	ResetResidualVelocity();
+	QueuedLayeredMoves.Empty();
+	ActiveLayeredMoves.Empty();
+}
+
diff --git a/Engine/Plugins/Experimental/Mover/Source/Mover/Private/MovementModifier.cpp b/Engine/Plugins/Experimental/Mover/Source/Mover/Private/MovementModifier.cpp
--- a/Engine/Plugins/Experimental/Mover/Source/Mover/Private/MovementModifier.cpp
+++ b/Engine/Plugins/Experimental/Mover/Source/Mover/Private/MovementModifier.cpp
@@ -400,4 +400,9 @@ struct FMovementModifierDeleter
 	}
 }
 
+void FMovementModifierGroup::Reset()
+{
+	QueuedModifiers.Empty();
+	ActiveModifiers.Empty();
+}
 
diff --git a/Engine/Plugins/Experimental/Mover/Source/Mover/Public/LayeredMove.h b/Engine/Plugins/Experimental/Mover/Source/Mover/Public/LayeredMove.h
--- a/Engine/Plugins/Experimental/Mover/Source/Mover/Public/LayeredMove.h
+++ b/Engine/Plugins/Experimental/Mover/Source/Mover/Public/LayeredMove.h
@@ -215,6 +215,9 @@ struct MOVER_API FLayeredMoveGroup
 	// Resets residual velocity related settings
 	void ResetResidualVelocity();
 
+	// Resets residual velocity and clears active and queued moves
+	void Reset();
+
 protected:
 	// Clears out any finished or invalid active moves and adds any queued moves to the active moves
 	void FlushMoveArrays(const UMoverComponent* MoverComp, UMoverBlackboard* SimBlackboard, float CurrentSimTimeMs);
diff --git a/Engine/Plugins/Experimental/Mover/Source/Mover/Public/MovementModifier.h b/Engine/Plugins/Experimental/Mover/Source/Mover/Public/MovementModifier.h
--- a/Engine/Plugins/Experimental/Mover/Source/Mover/Public/MovementModifier.h
+++ b/Engine/Plugins/Experimental/Mover/Source/Mover/Public/MovementModifier.h
@@ -216,6 +216,9 @@ struct MOVER_API FMovementModifierGroup
 	
 	/** Const access to queued moves */
 	TArray<TSharedPtr<FMovementModifierBase>>::TConstIterator GetQueuedModifiersIterator() const;
+
+	// Clears active and queued modifiers
+	void Reset();
 	
 protected:
 	/** Movement modifiers that are currently active in this group */
diff --git a/Engine/Plugins/Experimental/Mover/Source/Mover/Public/MoverSimulationTypes.h b/Engine/Plugins/Experimental/Mover/Source/Mover/Public/MoverSimulationTypes.h
--- a/Engine/Plugins/Experimental/Mover/Source/Mover/Public/MoverSimulationTypes.h
+++ b/Engine/Plugins/Experimental/Mover/Source/Mover/Public/MoverSimulationTypes.h
@@ -163,6 +163,16 @@ public:
 		SyncStateCollection.Interpolate(From->SyncStateCollection, To->SyncStateCollection, Pct);
 	}
 
+	// Resets the sync state to its default configuration and removes any
+	// active or queued layered modes and modifiers
+	void Reset()
+	{
+		MovementMode = NAME_None;
+		SyncStateCollection.Empty();
+		SyncStateCollection.FindOrAddDataByType<FMoverDefaultSyncState>();
+		LayeredMoves.Reset();
+		MovementModifiers.Reset();
+	}
 };
 
 // Auxiliary state that is input into the simulation (changes rarely)
