Physics Mover
- Produce inputs only on frames where when physics will tick
- Interpolate output sync state every game thread tick
- Fix for editor crash caused by not setting updsted component when not a game world
- Output movement base

diff --git a/Engine/Plugins/Experimental/Mover/Source/Mover/Private/Backends/MoverNetworkPhysicsLiaison.cpp b/Engine/Plugins/Experimental/Mover/Source/Mover/Private/Backends/MoverNetworkPhysicsLiaison.cpp
--- a/Engine/Plugins/Experimental/Mover/Source/Mover/Private/Backends/MoverNetworkPhysicsLiaison.cpp
+++ b/Engine/Plugins/Experimental/Mover/Source/Mover/Private/Backends/MoverNetworkPhysicsLiaison.cpp
@@ -13,6 +13,7 @@
 #include "DefaultMovementSet/Settings/CommonLegacyMovementSettings.h"
 #include "PBDRigidsSolver.h"
 #include "Physics/Experimental/PhysScene_Chaos.h"
+#include "PhysicsEngine/PhysicsObjectExternalInterface.h"
 #include "PhysicsMover/Modes/PhysicsDrivenFallingMode.h"
 #include "PhysicsMover/Modes/PhysicsDrivenWalkingMode.h"
 #include "PhysicsMover/PhysicsMovementUtils.h"
@@ -410,6 +411,8 @@ void UMoverNetworkPhysicsLiaisonComponent::OnComponentPhysicsStateChanged(UPrimi
 		{
 			MoverComp->ModeFSM->SetModeImmediately(MoverComp->StartingMovementMode);
 		}
+
+		InitializeCachedSyncState();
 	}
 }
 
@@ -417,25 +420,32 @@ void UMoverNetworkPhysicsLiaisonComponent::InitializeComponent()
 {
 	Super::InitializeComponent();
 
-	if (ensureAlwaysMsgf(MoverComp, TEXT("UMoverNetworkPhysicsLiaisonComponent on actor %s failed to find associated Mover component. This actor's movement will not be simulated. Verify its setup."), *GetNameSafe(GetOwner())))
+	UWorld* World = GetWorld();
+	if (World && World->IsGameWorld())
 	{
-		MoverComp->InitMoverSimulation();
-		MoverComp->ModeFSM->SetModeImmediately(MoverComp->StartingMovementMode);
-	}
+		if (ensureAlwaysMsgf(MoverComp, TEXT("UMoverNetworkPhysicsLiaisonComponent on actor %s failed to find associated Mover component. This actor's movement will not be simulated. Verify its setup."), *GetNameSafe(GetOwner())))
+		{
+			MoverComp->InitMoverSimulation();
+			MoverComp->ModeFSM->SetModeImmediately(MoverComp->StartingMovementMode);
+		}
 
-	if (Chaos::FPhysicsSolverBase::IsNetworkPhysicsPredictionEnabled())
-	{
-		NetworkPhysicsComponent = NewObject<UNetworkPhysicsComponent>(GetOwner(), TEXT("PhysMover_NetworkPhysicsComponent"));
-		if (NetworkPhysicsComponent)
+		if (Chaos::FPhysicsSolverBase::IsNetworkPhysicsPredictionEnabled())
 		{
-			NetworkPhysicsComponent->SetNetAddressable(); // Make DSO components net addressable
-			NetworkPhysicsComponent->SetIsReplicated(true);
-			NetworkPhysicsComponent->RegisterComponent();
-			NetworkPhysicsComponent->InitializeComponent();
-			NetworkPhysicsComponent->Activate(true);
-
-			// Register network data for recording and rewind/resim
-			NetworkPhysicsComponent->CreateDataHistory<FNetworkPhysicsMoverTraits>(this);
+			NetworkPhysicsComponent = NewObject<UNetworkPhysicsComponent>(GetOwner(), TEXT("PhysMover_NetworkPhysicsComponent"));
+			if (NetworkPhysicsComponent)
+			{
+				NetworkPhysicsComponent->SetNetAddressable(); // Make DSO components net addressable
+				NetworkPhysicsComponent->SetIsReplicated(true);
+				NetworkPhysicsComponent->RegisterComponent();
+				if (!NetworkPhysicsComponent->HasBeenInitialized())
+				{
+					NetworkPhysicsComponent->InitializeComponent();
+				}
+				NetworkPhysicsComponent->Activate(true);
+
+				// Register network data for recording and rewind/resim
+				NetworkPhysicsComponent->CreateDataHistory<FNetworkPhysicsMoverTraits>(this);
+			}
 		}
 	}
 }
@@ -523,13 +533,16 @@ void UMoverNetworkPhysicsLiaisonComponent::BeginPlay()
 	Super::BeginPlay();
 
 	// Register with the physics mover manager
-	if (UWorld* World = GetWorld())
+	UWorld* World = GetWorld();
+	if (World && World->IsGameWorld())
 	{
 		if (UPhysicsMoverManager* Manager = World->GetSubsystem<UPhysicsMoverManager>())
 		{
 			Manager->RegisterPhysicsMoverComponent(this);
 		}
 	}
+
+	InitializeCachedSyncState();
 }
 
 void UMoverNetworkPhysicsLiaisonComponent::EndPlay(const EEndPlayReason::Type EndPlayReason)
@@ -662,7 +675,39 @@ FMoverTimeStep UMoverNetworkPhysicsLiaisonComponent::GetCurrentMoverTimeStep(flo
 	return TimeStep;
 }
 
-void UMoverNetworkPhysicsLiaisonComponent::ProduceInput_External(float DeltaSeconds, OUT FPhysicsMoverAsyncInput& Input)
+void UMoverNetworkPhysicsLiaisonComponent::InitializeCachedSyncState()
+{
+	SyncStateRecord.Clear();
+
+	if (const UMoverComponent* MoverComponentPtr = MoverComp.Get())
+	{
+		FMoverSyncState SyncState;
+		SyncState.MovementMode = MoverComponentPtr->StartingMovementMode;
+
+		double SimTime = 0.0;
+		if (Chaos::FPhysicsSolver* Solver = GetPhysicsSolver())
+		{
+			SimTime = Solver->GetPhysicsResultsTime_External();
+		}
+
+		SyncStateRecord.Add(SimTime, SyncState);
+	}
+}
+
+Chaos::FPhysicsSolver* UMoverNetworkPhysicsLiaisonComponent::GetPhysicsSolver() const
+{
+	if (const UWorld* World = GetWorld())
+	{
+		if (const FPhysScene* Scene = World->GetPhysicsScene())
+		{
+			return Scene->GetSolver();
+		}
+	}
+
+	return nullptr;
+}
+
+void UMoverNetworkPhysicsLiaisonComponent::ProduceInput_External(int32 PhysicsStep, int32 NumSteps, OUT FPhysicsMoverAsyncInput& Input)
 {
 	Chaos::EnsureIsInGameThreadContext();
 
@@ -686,45 +731,38 @@ void UMoverNetworkPhysicsLiaisonComponent::ProduceInput_External(float DeltaSeco
 			NetworkPhysicsComponent->SetIsRelayingLocalInputs(true);
 		}
 
-		if (bProduceInput)
+		float DeltaSeconds = 0.0f;
+		if (const UWorld* World = GetWorld())
 		{
-			if (!bCachedInputIsValid)
+			if (bUsingAsyncPhysics)
 			{
-				const int DeltaTimeMS = FMath::RoundToInt(DeltaSeconds * 1000.0f);
-				MoverComp->ProduceInput(DeltaTimeMS, &Input.InputCmd);
-
-				// We only want to consume one input per physics frame
-				// so if there is already a valid cached input we use that.
-				// Input is set invalid when the async output is consumed
-				bCachedInputIsValid = MoverComp->CachedLastProducedInputCmd.InputCollection.FindDataByType<FCharacterDefaultInputs>() != nullptr;
+				if (Chaos::FPhysicsSolver* Solver = GetPhysicsSolver())
+				{
+					if (Solver->IsUsingAsyncResults())
+					{
+						DeltaSeconds = Solver->GetAsyncDeltaTime();
+					}
+				}
 			}
 			else
 			{
-				Input.InputCmd = MoverComp->CachedLastProducedInputCmd;
+				DeltaSeconds = World->GetDeltaSeconds();
 			}
 		}
 
-		Input.SyncState.SyncStateCollection.Empty();
-		if (MoverComp->bHasValidCachedState)
-		{
-			Input.SyncState = MoverComp->CachedLastSyncState;
-		}
-		else
+		if (bProduceInput)
 		{
-			Input.SyncState.MovementMode = MoverComp->StartingMovementMode;
+			const int DeltaTimeMS = FMath::RoundToInt(DeltaSeconds * 1000.0f);
+			MoverComp->ProduceInput(DeltaTimeMS, &Input.InputCmd);
 		}
 
+		Input.SyncState.SyncStateCollection.Empty();
+		Input.SyncState = SyncStateRecord.GetLatest();
+
 		const FMoverTimeStep MoverTimeStep = bUsingAsyncPhysics ? GetCurrentAsyncMoverTimeStep_External() : GetCurrentMoverTimeStep(DeltaSeconds);
 		MoverComp->CachedLastSimTickTimeStep = MoverTimeStep;
 
-		if (bCachedInputIsValid)
-		{
-			MoverComp->OnPreSimulationTick.Broadcast(MoverTimeStep, Input.InputCmd);
-		}
-		else
-		{
-			MoverComp->OnPreSimulationTick.Broadcast(MoverTimeStep, NetInputCmd);
-		}
+		MoverComp->OnPreSimulationTick.Broadcast(MoverTimeStep, Input.InputCmd);
 
 		// This is required so that the physics thread can have a copy of the data to access
 		NetInputCmd = Input.InputCmd;
@@ -736,81 +774,76 @@ void UMoverNetworkPhysicsLiaisonComponent::ConsumeOutput_External(const FPhysics
 {
 	Chaos::EnsureIsInGameThreadContext();
 
-	if (Output.bIsValid && MoverComp)
+	if (Output.bIsValid)
 	{
-		const float DeltaSeconds = MoverComp->CachedLastSimTickTimeStep.StepMs * 0.001f;
-		const FMoverTimeStep MoverTimeStep = bUsingAsyncPhysics ? GetCurrentAsyncMoverTimeStep_External() : GetCurrentMoverTimeStep(DeltaSeconds);
+		SyncStateRecord.Add(OutputTimeInSeconds, Output.SyncState);
+	}
+}
 
-		if (bUsingAsyncPhysics)
+void UMoverNetworkPhysicsLiaisonComponent::PostPhysicsUpdate_External()
+{
+	Chaos::EnsureIsInGameThreadContext();
+
+	// Write the interpolated result to the mover component
+	if (UMoverComponent* MoverComponentPtr = MoverComp.Get())
+	{
+		if (const Chaos::FPhysicsSolver* Solver = GetPhysicsSolver())
 		{
-			if (bCachedLastPhysicsSyncStateIsValid)
+			const float DeltaSeconds = MoverComp->CachedLastSimTickTimeStep.StepMs * 0.001f;
+			const FMoverTimeStep MoverTimeStep = bUsingAsyncPhysics ? GetCurrentAsyncMoverTimeStep_External() : GetCurrentMoverTimeStep(DeltaSeconds);
+
+			FMoverSyncState OutputSyncState;
+			if (bUsingAsyncPhysics)
 			{
-				const double PhysicsResultsTime = GetWorld()->GetPhysicsScene()->GetSolver()->GetPhysicsResultsTime_External();
+				const float ResultsTime = Solver->GetPhysicsResultsTime_External();
+				SyncStateRecord.GetInterpolated(ResultsTime, OutputSyncState);
+			}
+			else
+			{
+				OutputSyncState = SyncStateRecord.GetLatest();
+			}
 
-				if (PhysicsResultsTime < OutputTimeInSeconds)
+			FMoverDefaultSyncState& DefaultSyncState = OutputSyncState.SyncStateCollection.FindOrAddMutableDataByType<FMoverDefaultSyncState>();
+
+			// The physics mover does not set the sync state transform, it uses the physics transform,
+			// so get the world transform from the physics particle.
+			if (IPhysicsComponent* PhysComp = Cast<IPhysicsComponent>(MoverComponentPtr->GetUpdatedComponent()))
+			{
+				Chaos::FPhysicsObject* PhysicsObject = PhysComp->GetPhysicsObjectByName(NAME_None);
+				Chaos::FReadPhysicsObjectInterface_External Interface = FPhysicsObjectExternalInterface::GetRead_AssumesLocked();
+				if (Chaos::FPBDRigidParticle* Particle = Interface.GetRigidParticle(PhysicsObject))
 				{
+					DefaultSyncState.SetTransforms_WorldSpace(Particle->X(), FRotator(Particle->R()), Particle->V(), DefaultSyncState.GetMovementBase(), DefaultSyncState.GetMovementBaseBoneName());
+				}
+			}
 
-					// Output is in the future. Interpolate from the cached previous state to the next state
-					const float Denom = (OutputTimeInSeconds - CachedLastPhysicsSyncStateOutputTime);
-					if (Denom > 0)
-					{
-						const float Alpha = FMath::Clamp((PhysicsResultsTime - CachedLastPhysicsSyncStateOutputTime) / Denom, 0.0f, 1.0f);
-						FMoverSyncState InterpolatedSyncState;
-						InterpolatedSyncState.Interpolate(&CachedLastPhysicsSyncState, &Output.SyncState, Alpha);
+			const FName CachedMovementMode = MoverComponentPtr->GetMovementModeName();
 
-						// The physics mover does not set the sync state transform, it uses the physics transform,
-						// so get the world transform from the physics particle.
-						if (const Chaos::FSingleParticlePhysicsProxy* CharacterProxy = Constraint->GetCharacterParticleProxy())
-						{
-							const Chaos::FRigidBodyHandle_External& Body = CharacterProxy->GetGameThreadAPI();
-							if (FMoverDefaultSyncState* SyncState = InterpolatedSyncState.SyncStateCollection.FindMutableDataByType<FMoverDefaultSyncState>())
-							{
-								SyncState->SetTransforms_WorldSpace(Body.X(), FRotator(Body.R()), Body.V());
-							}
-						}
+			const bool bIsCachedModeFalling = MoverComponentPtr->HasGameplayTag(Mover_IsFalling, true);
 
-						// If landed, broadcast OnLanded
-						// TODO: Generalize handling of events generated on physics thread
-						if (MoverComp->HasValidCachedState())
-						{
-							if ((MoverComp->GetSyncState().MovementMode == DefaultModeNames::Falling) && (InterpolatedSyncState.MovementMode == DefaultModeNames::Walking))
-							{
-								if (UFallingMode* FallingMode = MoverComp->FindMode_Mutable<UFallingMode>())
-								{
-									FHitResult HitResult;
-									MoverComp->TryGetFloorCheckHitResult(HitResult);
-									FallingMode->OnLanded.Broadcast(InterpolatedSyncState.MovementMode, HitResult);
-								}
-							}
-						}
+			MoverComponentPtr->CachedLastSyncState = OutputSyncState;
+			MoverComponentPtr->CachedLastSimTickTimeStep = MoverTimeStep;
+			MoverComponentPtr->bHasValidCachedState = true;
 
-						// TODO: Consider moving to a util function
-						MoverComp->CachedLastSyncState = InterpolatedSyncState;
-						MoverComp->CachedLastSimTickTimeStep = MoverTimeStep;
-						MoverComp->bHasValidCachedState = true;
-					}
+			const FName NextMode = MoverComponentPtr->GetMovementModeName();
+			if (CachedMovementMode != NextMode)
+			{
+				UpdateConstraintSettings();
+				MoverComponentPtr->OnMovementModeChanged.Broadcast(CachedMovementMode, NextMode);
+
+				if (bIsCachedModeFalling && MoverComponentPtr->HasGameplayTag(Mover_IsOnGround, true))
+				{
+					if (UFallingMode* FallingMode = MoverComponentPtr->FindMode_Mutable<UFallingMode>())
+					{
+						FHitResult HitResult;
+						MoverComponentPtr->TryGetFloorCheckHitResult(HitResult);
 
-					// Call on the last output for the frame
-					UpdateConstraintSettings();
-					MoverComp->OnPostSimulationTick.Broadcast(MoverTimeStep);
-					bCachedInputIsValid = false;
+						FallingMode->OnLanded.Broadcast(OutputSyncState.MovementMode, HitResult);
+					}
 				}
 			}
 
-			CachedLastPhysicsSyncState = Output.SyncState;
-			CachedLastPhysicsSyncStateOutputTime = OutputTimeInSeconds;
-			bCachedLastPhysicsSyncStateIsValid = true;
-		}
-		else
-		{
-			MoverComp->CachedLastSyncState = Output.SyncState;
-			MoverComp->CachedLastSimTickTimeStep = MoverTimeStep;
-			MoverComp->bHasValidCachedState = true;
-
-			// Call on the last output for the frame
-			UpdateConstraintSettings();
-			MoverComp->OnPostSimulationTick.Broadcast(MoverTimeStep);
-			bCachedInputIsValid = false;
+			MoverComponentPtr->OnPostSimulationTick.Broadcast(MoverTimeStep);
 		}
 	}
 }
@@ -1026,6 +1059,10 @@ void UMoverNetworkPhysicsLiaisonComponent::OnPreSimulate_Internal(const FPhysics
 	// Update the constraint data based on the floor result
 	if (Output.FloorResult.bBlockingHit)
 	{
+		// Note: We want to record the movement base but we don't record the transform
+		// so don't use this to get a relative transform for the sync state
+		OutputSyncState.SetMovementBase(Output.FloorResult.HitResult.GetComponent());
+
 		// Set the ground particle on the constraint
 		Chaos::FGeometryParticleHandle* GroundParticle = nullptr;
 
@@ -1065,6 +1102,8 @@ void UMoverNetworkPhysicsLiaisonComponent::OnPreSimulate_Internal(const FPhysics
 	}
 	else
 	{
+		OutputSyncState.SetMovementBase(nullptr);
+
 		ConstraintHandle->SetData({
 			Chaos::FVec3::ZAxisVector,
 			Chaos::FVec3::ZeroVector,
@@ -1089,17 +1128,11 @@ void UMoverNetworkPhysicsLiaisonComponent::TeleportParticle(Chaos::FGeometryPart
 		return;
 	}
 
-	if (UWorld* World = GetWorld())
+	if (Chaos::FPhysicsSolver* Solver = GetPhysicsSolver())
 	{
-		if (FPhysScene* Scene = World->GetPhysicsScene())
+		if (Solver->GetEvolution())
 		{
-			if (Chaos::FPhysicsSolver* Solver = Scene->GetSolver())
-			{
-				if (Solver->GetEvolution())
-				{
-					Solver->GetEvolution()->SetParticleTransform(Particle, Position, Rotation, true);
-				}
-			}
+			Solver->GetEvolution()->SetParticleTransform(Particle, Position, Rotation, true);
 		}
 	}
 }
diff --git a/Engine/Plugins/Experimental/Mover/Source/Mover/Private/PhysicsMover/PhysicsMoverManager.cpp b/Engine/Plugins/Experimental/Mover/Source/Mover/Private/PhysicsMover/PhysicsMoverManager.cpp
--- a/Engine/Plugins/Experimental/Mover/Source/Mover/Private/PhysicsMover/PhysicsMoverManager.cpp
+++ b/Engine/Plugins/Experimental/Mover/Source/Mover/Private/PhysicsMover/PhysicsMoverManager.cpp
@@ -19,8 +19,17 @@ void UPhysicsMoverManager::OnWorldBeginPlay(UWorld& InWorld)
 
 	if (FPhysScene* PhysScene = InWorld.GetPhysicsScene())
 	{
-		PhysScenePreTickCallbackHandle = PhysScene->OnPhysScenePreTick.AddUObject(this, &ThisClass::PrePhysicsUpdate);
-		AsyncCallback = PhysScene->GetSolver()->CreateAndRegisterSimCallbackObject_External<FPhysicsMoverManagerAsyncCallback>();
+		PhysScenePostTickCallbackHandle = PhysScene->OnPhysScenePostTick.AddUObject(this, &UPhysicsMoverManager::OnPostPhysicsTick);
+
+		if (Chaos::FPhysicsSolver* Solver = PhysScene->GetSolver())
+		{
+			AsyncCallback = Solver->CreateAndRegisterSimCallbackObject_External<FPhysicsMoverManagerAsyncCallback>();
+
+			if (FNetworkPhysicsCallback* SolverCallback = static_cast<FNetworkPhysicsCallback*>(Solver->GetRewindCallback()))
+			{
+				InjectInputsExternalCallbackHandle = SolverCallback->InjectInputsExternal.AddUObject(this, &UPhysicsMoverManager::InjectInputs_External);
+			}
+		}
 	}
 }
 
@@ -30,16 +39,31 @@ void UPhysicsMoverManager::Deinitialize()
 	{
 		if (FPhysScene* PhysScene = World->GetPhysicsScene())
 		{
-			PhysScene->OnPhysScenePreTick.Remove(PhysScenePreTickCallbackHandle);
-			if (AsyncCallback)
+			if (Chaos::FPhysicsSolver* Solver = PhysScene->GetSolver())
 			{
-				PhysScene->GetSolver()->UnregisterAndFreeSimCallbackObject_External(AsyncCallback);
-			}
+				if (InjectInputsExternalCallbackHandle.IsValid())
+				{
+					if (FNetworkPhysicsCallback* SolverCallback = static_cast<FNetworkPhysicsCallback*>(Solver->GetRewindCallback()))
+					{
+						SolverCallback->InjectInputsExternal.Remove(InjectInputsExternalCallbackHandle);
+					}
+				}
 
-			PhysicsMoverComponents.Empty();
+				if (PhysScenePostTickCallbackHandle.IsValid())
+				{
+					PhysScene->OnPhysScenePostTick.Remove(PhysScenePostTickCallbackHandle);
+				}
+
+				if (AsyncCallback)
+				{
+					Solver->UnregisterAndFreeSimCallbackObject_External(AsyncCallback);
+				}
+			}
 		}
 	}
 
+	PhysicsMoverComponents.Empty();
+
 	Super::Deinitialize();
 }
 
@@ -62,13 +86,28 @@ void UPhysicsMoverManager::UnregisterPhysicsMoverComponent(TWeakObjectPtr<UMover
 
 //////////////////////////////////////////////////////////////////////////
 
-void UPhysicsMoverManager::PrePhysicsUpdate(FPhysScene* PhysScene, float DeltaTime)
+void UPhysicsMoverManager::InjectInputs_External(int32 PhysicsStep, int32 NumSteps)
 {
 	// Clear invalid data before we attempt to use it.
 	PhysicsMoverComponents.RemoveAllSwap([](const TWeakObjectPtr<UMoverNetworkPhysicsLiaisonComponent> MoverComp) {
 		return (MoverComp.IsValid() == false) || (MoverComp->GetUniqueIdx().IsValid() == false);
 		});
 
+
+	FPhysicsMoverManagerAsyncInput* ManagerAsyncInput = AsyncCallback->GetProducerInputData_External();
+	ManagerAsyncInput->Reset();
+	ManagerAsyncInput->AsyncInput.Reserve(PhysicsMoverComponents.Num());
+
+	for (TWeakObjectPtr<UMoverNetworkPhysicsLiaisonComponent> PhysicsMoverComp : PhysicsMoverComponents)
+	{
+		TUniquePtr<FPhysicsMoverAsyncInput> InputData = MakeUnique<FPhysicsMoverAsyncInput>();
+		PhysicsMoverComp->ProduceInput_External(PhysicsStep, NumSteps, *InputData);
+		ManagerAsyncInput->AsyncInput.Add(MoveTemp(InputData));
+	}
+}
+
+void UPhysicsMoverManager::OnPostPhysicsTick(FChaosScene*)
+{
 	while (Chaos::TSimCallbackOutputHandle<FPhysicsMoverManagerAsyncOutput> ManagerAsyncOutput = AsyncCallback->PopFutureOutputData_External())
 	{
 		const double OutputTime = ManagerAsyncOutput->InternalTime;
@@ -85,17 +124,8 @@ void UPhysicsMoverManager::PrePhysicsUpdate(FPhysScene* PhysScene, float DeltaTi
 		}
 	}
 
-	FPhysicsMoverManagerAsyncInput* ManagerAsyncInput = AsyncCallback->GetProducerInputData_External();
-	ManagerAsyncInput->Reset();
-	ManagerAsyncInput->AsyncInput.Reserve(PhysicsMoverComponents.Num());
-
-	if (PhysScene && PhysScene->GetOwningWorld())
+	for (TWeakObjectPtr<UMoverNetworkPhysicsLiaisonComponent> PhysicsMoverComp : PhysicsMoverComponents)
 	{
-		for (TWeakObjectPtr<UMoverNetworkPhysicsLiaisonComponent> PhysicsMoverComp : PhysicsMoverComponents)
-		{
-			TUniquePtr<FPhysicsMoverAsyncInput> InputData = MakeUnique<FPhysicsMoverAsyncInput>();
-			PhysicsMoverComp->ProduceInput_External(DeltaTime, *InputData);
-			ManagerAsyncInput->AsyncInput.Add(MoveTemp(InputData));
-		}
+		PhysicsMoverComp->PostPhysicsUpdate_External();
 	}
 }
\ No newline at end of file
diff --git a/Engine/Plugins/Experimental/Mover/Source/Mover/Public/Backends/MoverNetworkPhysicsLiaison.h b/Engine/Plugins/Experimental/Mover/Source/Mover/Public/Backends/MoverNetworkPhysicsLiaison.h
--- a/Engine/Plugins/Experimental/Mover/Source/Mover/Public/Backends/MoverNetworkPhysicsLiaison.h
+++ b/Engine/Plugins/Experimental/Mover/Source/Mover/Public/Backends/MoverNetworkPhysicsLiaison.h
@@ -94,6 +94,64 @@ struct MOVER_API FNetworkPhysicsMoverTraits
 	using StatesType = FNetworkPhysicsMoverState;
 };
 
+//////////////////////////////////////////////////////////////////////////
+// FSyncStateRecord
+
+class FSyncStateRecord
+{
+public:
+	struct FData
+	{
+		FMoverSyncState SyncState;
+		double SimTime;
+	};
+
+	void Add(double InSimTime, const FMoverSyncState& InSyncState)
+	{
+		CurrentIndex = (CurrentIndex + 1) % 2;
+		Data[CurrentIndex].SimTime = InSimTime;
+		Data[CurrentIndex].SyncState = InSyncState;
+	}
+
+	const FMoverSyncState& GetLatest() const
+	{
+		return Data[CurrentIndex].SyncState;
+	}
+
+	void GetInterpolated(double SimTime, FMoverSyncState& OutSyncState) const
+	{
+		const uint8 PrevIndex = (CurrentIndex + 1) % 2;
+		const double PrevTime = Data[PrevIndex].SimTime;
+		const double CurrTime = Data[CurrentIndex].SimTime;
+
+		if (FMath::IsNearlyEqual(PrevTime, CurrTime) || (SimTime >= CurrTime))
+		{
+			OutSyncState = Data[CurrentIndex].SyncState;
+		}
+		else if (SimTime <= PrevTime)
+		{
+			OutSyncState = Data[PrevIndex].SyncState;
+		}
+		else
+		{
+			const float Alpha = FMath::Clamp((SimTime - PrevTime) / (CurrTime - PrevTime), 0.0f, 1.0f);
+			OutSyncState.Interpolate(&Data[PrevIndex].SyncState, &Data[CurrentIndex].SyncState, Alpha);
+		}
+	}
+
+	void Clear()
+	{
+		CurrentIndex = 1;
+		Data[0].SyncState.Reset();
+		Data[0].SyncState.SyncStateCollection.Empty();
+		Data[1].SyncState.Reset();
+		Data[1].SyncState.SyncStateCollection.Empty();
+	}
+
+	FData Data[2];
+	uint8 CurrentIndex = 1;
+};
+
 //////////////////////////////////////////////////////////////////////////
 // UMoverNetworkPhysicsLiaisonComponent
 
@@ -139,10 +197,11 @@ public:
 	// Used by the manager to uniquely identify the component
 	Chaos::FUniqueIdx GetUniqueIdx() const;
 
-	void ProduceInput_External(float DeltaSeconds, OUT FPhysicsMoverAsyncInput& Input);
+	void ProduceInput_External(int32 PhysicsStep, int32 NumSteps, OUT FPhysicsMoverAsyncInput& Input);
 	void ConsumeOutput_External(const FPhysicsMoverAsyncOutput& Output, const double OutputTimeInSeconds);
+	void PostPhysicsUpdate_External();
 
-	void ProcessInputs_Internal(int32 PhysicsStep, float DeltaTime, const FPhysicsMoverAsyncInput& Input) const;
+	void ProcessInputs_Internal(int32 PhysicsStep, float DeltaTimeInSeconds, const FPhysicsMoverAsyncInput& Input) const;
 	void OnPreSimulate_Internal(const FPhysicsMoverSimulationTickParams& TickParams, const FPhysicsMoverAsyncInput& Input, OUT FPhysicsMoverAsyncOutput& SimOutput) const;
 	void OnContactModification_Internal(const FPhysicsMoverAsyncInput& Input, Chaos::FCollisionContactModifier& Modifier) const;
 
@@ -170,6 +229,10 @@ protected:
 	// Time step on the game thread when using non-async physics
 	FMoverTimeStep GetCurrentMoverTimeStep(float DeltaSeconds) const;
 
+	void InitializeCachedSyncState();
+
+	Chaos::FPhysicsSolver* GetPhysicsSolver() const;
+
 	// These are written to by the network input and state data
 	FMoverInputCmdContext NetInputCmd;
 	FMoverSyncState NetSyncState;
@@ -180,15 +243,10 @@ protected:
 
 	TObjectPtr<UNetworkPhysicsComponent> NetworkPhysicsComponent;
 
-	// The cached physics state is the latest output from the physics thread
-	// This can be different from the cached sync state on the mover component
-	// which is interpolated to match the interpolated physics particle
-	// Note: Time is physics solver time
-	FMoverSyncState CachedLastPhysicsSyncState;
-	double CachedLastPhysicsSyncStateOutputTime;
-	bool bCachedLastPhysicsSyncStateIsValid = false;
-
-	bool bCachedInputIsValid = false;
+	// This stores the previous two outputs from the physics thread
+	// which can be interpolated to get the sync state at a time
+	// that matches the interpolated physics particle.
+	FSyncStateRecord SyncStateRecord;
 
 	bool bUsingAsyncPhysics = false;
 };
\ No newline at end of file
diff --git a/Engine/Plugins/Experimental/Mover/Source/Mover/Public/PhysicsMover/PhysicsMoverManager.h b/Engine/Plugins/Experimental/Mover/Source/Mover/Public/PhysicsMover/PhysicsMoverManager.h
--- a/Engine/Plugins/Experimental/Mover/Source/Mover/Public/PhysicsMover/PhysicsMoverManager.h
+++ b/Engine/Plugins/Experimental/Mover/Source/Mover/Public/PhysicsMover/PhysicsMoverManager.h
@@ -8,6 +8,7 @@
 
 //////////////////////////////////////////////////////////////////////////
 
+class FChaosScene;
 class UMoverNetworkPhysicsLiaisonComponent;
 
 UCLASS()
@@ -20,13 +21,15 @@ public:
 	virtual void Deinitialize() override;
 	virtual bool DoesSupportWorldType(const EWorldType::Type WorldType) const override;
 
-	void PrePhysicsUpdate(FPhysScene* PhysScene, float DeltaTime);
+	void InjectInputs_External(int32 PhysicsStep, int32 NumSteps);
+	void OnPostPhysicsTick(FChaosScene* Scene);
 
 	void RegisterPhysicsMoverComponent(TWeakObjectPtr<UMoverNetworkPhysicsLiaisonComponent> InPhysicsMoverComp);
 	void UnregisterPhysicsMoverComponent(TWeakObjectPtr<UMoverNetworkPhysicsLiaisonComponent> InPhysicsMoverComp);
 
 private:
 	TArray<TWeakObjectPtr<UMoverNetworkPhysicsLiaisonComponent>> PhysicsMoverComponents;
-	FDelegateHandle PhysScenePreTickCallbackHandle;
+	FDelegateHandle InjectInputsExternalCallbackHandle;
+	FDelegateHandle PhysScenePostTickCallbackHandle;
 	class FPhysicsMoverManagerAsyncCallback* AsyncCallback = nullptr;
 };
\ No newline at end of file
