Mover: Fixing FindMovementModifier to search through queued modifiers in addition to active ones

diff --git a/Engine/Plugins/Experimental/Mover/Source/Mover/Private/MoverComponent.cpp b/Engine/Plugins/Experimental/Mover/Source/Mover/Private/MoverComponent.cpp
--- a/Engine/Plugins/Experimental/Mover/Source/Mover/Private/MoverComponent.cpp
+++ b/Engine/Plugins/Experimental/Mover/Source/Mover/Private/MoverComponent.cpp
@@ -608,7 +608,7 @@ const FMovementModifierBase* UMoverComponent::FindMovementModifier(const FMoveme
 		}
 
 		// Check queued modifiers for modifier handle
-		for (auto QueuedModifierFromSyncStateIt = CachedLastSyncState.MovementModifiers.GetActiveModifiersIterator(); QueuedModifierFromSyncStateIt; ++QueuedModifierFromSyncStateIt)
+		for (auto QueuedModifierFromSyncStateIt = CachedLastSyncState.MovementModifiers.GetQueuedModifiersIterator(); QueuedModifierFromSyncStateIt; ++QueuedModifierFromSyncStateIt)
 		{
 			const TSharedPtr<FMovementModifierBase> QueuedModifierFromSyncState = *QueuedModifierFromSyncStateIt;
 
@@ -638,7 +638,7 @@ const FMovementModifierBase* UMoverComponent::FindMovementModifierByType(const U
 		}
 
 		// Check queued modifiers for modifier handle
-		for (auto QueuedModifierFromSyncStateIt = CachedLastSyncState.MovementModifiers.GetActiveModifiersIterator(); QueuedModifierFromSyncStateIt; ++QueuedModifierFromSyncStateIt)
+		for (auto QueuedModifierFromSyncStateIt = CachedLastSyncState.MovementModifiers.GetQueuedModifiersIterator(); QueuedModifierFromSyncStateIt; ++QueuedModifierFromSyncStateIt)
 		{
 			const TSharedPtr<FMovementModifierBase> QueuedModifierFromSyncState = *QueuedModifierFromSyncStateIt;
 
