Mover: Adding BlueprintInternalUseOnly to LayeredMove, InstantEffect and MovmentModifier base structs to disallow making and queueing the base structs in BP.
Also making error message for not implementing Clone function clearer and more helpful.

diff --git a/Engine/Plugins/Experimental/Mover/Source/Mover/Private/InstantMovementEffect.cpp b/Engine/Plugins/Experimental/Mover/Source/Mover/Private/InstantMovementEffect.cpp
--- a/Engine/Plugins/Experimental/Mover/Source/Mover/Private/InstantMovementEffect.cpp
+++ b/Engine/Plugins/Experimental/Mover/Source/Mover/Private/InstantMovementEffect.cpp
@@ -5,7 +5,7 @@
 FInstantMovementEffect* FInstantMovementEffect::Clone() const
 {
 	// If child classes don't override this, saved moves will not work
-	checkf(false, TEXT("FInstantMovementEffect::Clone() being called erroneously. This should always be overridden in child classes!"));
+	checkf(false, TEXT("FInstantMovementEffect::Clone() being called erroneously from %s. A FInstantMovementEffect should never be queued directly and Clone should always be overridden in child structs!"), *GetNameSafe(GetScriptStruct()));
 	return nullptr;
 }
 
diff --git a/Engine/Plugins/Experimental/Mover/Source/Mover/Private/LayeredMove.cpp b/Engine/Plugins/Experimental/Mover/Source/Mover/Private/LayeredMove.cpp
--- a/Engine/Plugins/Experimental/Mover/Source/Mover/Private/LayeredMove.cpp
+++ b/Engine/Plugins/Experimental/Mover/Source/Mover/Private/LayeredMove.cpp
@@ -61,7 +61,7 @@ void FLayeredMoveBase::EndMove(const UMoverComponent* MoverComp, UMoverBlackboar
 FLayeredMoveBase* FLayeredMoveBase::Clone() const
 {
 	// If child classes don't override this, saved moves will not work
-	checkf(false, TEXT("FLayeredMoveBase::Clone() being called erroneously. This should always be overridden in child classes!"));
+	checkf(false, TEXT("FLayeredMoveBase::Clone() being called erroneously from %s. A FLayeredMoveBase should never be queued directly and Clone should always be overridden in child structs!"), *GetNameSafe(GetScriptStruct()));
 	return nullptr;
 }
 
diff --git a/Engine/Plugins/Experimental/Mover/Source/Mover/Private/MovementModifier.cpp b/Engine/Plugins/Experimental/Mover/Source/Mover/Private/MovementModifier.cpp
--- a/Engine/Plugins/Experimental/Mover/Source/Mover/Private/MovementModifier.cpp
+++ b/Engine/Plugins/Experimental/Mover/Source/Mover/Private/MovementModifier.cpp
@@ -48,7 +48,7 @@ bool FMovementModifierBase::IsFinished(float CurrentSimTimeMs) const
 FMovementModifierBase* FMovementModifierBase::Clone() const
 {
 	// If child classes don't override this, saved modifiers will not work
-	checkf(false, TEXT("FMovementModifierBase::Clone() being called erroneously. This should always be overridden in child classes!"));
+	checkf(false, TEXT("FMovementModifierBase::Clone() being called erroneously from %s. A FMovementModifierBase should never be queued directly and Clone should always be overridden in child structs!"), *GetNameSafe(GetScriptStruct()));
 	return nullptr;
 }
 
diff --git a/Engine/Plugins/Experimental/Mover/Source/Mover/Public/InstantMovementEffect.h b/Engine/Plugins/Experimental/Mover/Source/Mover/Public/InstantMovementEffect.h
--- a/Engine/Plugins/Experimental/Mover/Source/Mover/Public/InstantMovementEffect.h
+++ b/Engine/Plugins/Experimental/Mover/Source/Mover/Public/InstantMovementEffect.h
@@ -29,7 +29,7 @@ struct FApplyMovementEffectParams
  * Common uses would be for Teleporting, Changing Movement Modes directly, one time force application, etc.
  * Multiple Instant Movement Effects can be active at the time
  */
-USTRUCT(BlueprintType)
+USTRUCT(BlueprintInternalUseOnly)
 struct MOVER_API FInstantMovementEffect
 {
 	GENERATED_USTRUCT_BODY()
diff --git a/Engine/Plugins/Experimental/Mover/Source/Mover/Public/LayeredMove.h b/Engine/Plugins/Experimental/Mover/Source/Mover/Public/LayeredMove.h
--- a/Engine/Plugins/Experimental/Mover/Source/Mover/Public/LayeredMove.h
+++ b/Engine/Plugins/Experimental/Mover/Source/Mover/Public/LayeredMove.h
@@ -68,7 +68,7 @@ struct FLayeredMoveFinishVelocitySettings
 */
 
 // Base class for all layered moves
-USTRUCT(BlueprintType)
+USTRUCT(BlueprintInternalUseOnly)
 struct MOVER_API FLayeredMoveBase
 {
 	GENERATED_USTRUCT_BODY()
diff --git a/Engine/Plugins/Experimental/Mover/Source/Mover/Public/MovementModifier.h b/Engine/Plugins/Experimental/Mover/Source/Mover/Public/MovementModifier.h
--- a/Engine/Plugins/Experimental/Mover/Source/Mover/Public/MovementModifier.h
+++ b/Engine/Plugins/Experimental/Mover/Source/Mover/Public/MovementModifier.h
@@ -75,7 +75,7 @@ private:
  *		 This can be fixed by extending the Matches function to check more than just type,
  *		 but make sure anything used to compare is synced through the NetSerialize function.
  */
-USTRUCT(BlueprintType)
+USTRUCT(BlueprintInternalUseOnly)
 struct MOVER_API FMovementModifierBase
 {
 	GENERATED_BODY()
