Physics Mover - Base on root of compound objects

diff --git a/Engine/Plugins/Experimental/Mover/Source/Mover/Private/PhysicsMover/PhysicsMovementUtils.cpp b/Engine/Plugins/Experimental/Mover/Source/Mover/Private/PhysicsMover/PhysicsMovementUtils.cpp
--- a/Engine/Plugins/Experimental/Mover/Source/Mover/Private/PhysicsMover/PhysicsMovementUtils.cpp
+++ b/Engine/Plugins/Experimental/Mover/Source/Mover/Private/PhysicsMover/PhysicsMovementUtils.cpp
@@ -90,6 +90,14 @@ void UPhysicsMovementUtils::FloorSweep_Internal(
 			OutFloorResult.bWalkableFloor = bWalkable;
 			OutFloorResult.FloorDist = UpDir.Dot(Location - OutHit.ImpactPoint);
 			OutFloorResult.HitResult = OutHit;
+
+			if (UPrimitiveComponent* GroundComp = OutHit.GetComponent())
+			{
+				if (UPrimitiveComponent* RootComp = Cast<UPrimitiveComponent>(GroundComp->GetAttachmentRoot()))
+				{
+					OutFloorResult.HitResult.Component = RootComp;
+				}
+			}
 		}
 		else
 		{
