Physics Mover - Jumping applies reaction to ground

diff --git a/Engine/Plugins/Experimental/Mover/Source/Mover/Private/PhysicsMover/Transitions/PhysicsJumpCheck.cpp b/Engine/Plugins/Experimental/Mover/Source/Mover/Private/PhysicsMover/Transitions/PhysicsJumpCheck.cpp
--- a/Engine/Plugins/Experimental/Mover/Source/Mover/Private/PhysicsMover/Transitions/PhysicsJumpCheck.cpp
+++ b/Engine/Plugins/Experimental/Mover/Source/Mover/Private/PhysicsMover/Transitions/PhysicsJumpCheck.cpp
@@ -1,13 +1,39 @@
 // Copyright Epic Games, Inc. All Rights Reserved.
 
 #include "PhysicsMover/Transitions/PhysicsJumpCheck.h"
+
+#include "Chaos/PhysicsObject.h"
+#include "Chaos/PhysicsObjectInternalInterface.h"
 #include "DefaultMovementSet/InstantMovementEffects/BasicInstantMovementEffects.h"
 #include "MoverComponent.h"
+#include "MoveLibrary/FloorQueryUtils.h"
 
 #if WITH_EDITOR
 #include "Misc/DataValidation.h"
 #endif
 
+namespace Private
+{
+	static Chaos::FPBDRigidParticleHandle* GetParticleHandleFromComponent(UPrimitiveComponent* PrimComp)
+	{
+		if (IPhysicsComponent* PhysComp = Cast<IPhysicsComponent>(PrimComp))
+		{
+			Chaos::FPhysicsObject* PhysicsObject = PhysComp->GetPhysicsObjectByName(NAME_None);
+			Chaos::FReadPhysicsObjectInterface_Internal Interface = Chaos::FPhysicsObjectInternalInterface::GetRead();
+			return Interface.GetRigidParticle(PhysicsObject);
+		}
+
+		return nullptr;
+	}
+
+	static void ApplyImpulse(Chaos::FPBDRigidParticleHandle* Particle, const Chaos::FVec3& Impulse, const Chaos::FVec3& Location)
+	{
+		const Chaos::FVec3 Offset = Location - Particle->GetX();
+		Particle->SetW(Particle->GetW() + Particle->InvI() * Offset.Cross(Impulse));
+		Particle->SetV(Particle->GetV() + Particle->InvM() * Impulse);
+	}
+}
+
 UPhysicsJumpCheck::UPhysicsJumpCheck(const FObjectInitializer& ObjectInitializer)
 	: Super(ObjectInitializer)
 {
@@ -33,9 +59,29 @@ FTransitionEvalResult UPhysicsJumpCheck::OnEvaluate(const FSimulationTickParams&
 void UPhysicsJumpCheck::OnTrigger(const FSimulationTickParams& Params)
 {
 	TSharedPtr<FJumpImpulseEffect> JumpMove = MakeShared<FJumpImpulseEffect>();
-	JumpMove->UpwardsSpeed = JumpUpwardsSpeed;
 
-	Params.MovingComps.MoverComponent->QueueInstantMovementEffect(JumpMove);
+	if (UMoverComponent* MoverComp = Params.MovingComps.MoverComponent.Get())
+	{
+		JumpMove->UpwardsSpeed = JumpUpwardsSpeed;
+		Params.MovingComps.MoverComponent->QueueInstantMovementEffect(JumpMove);
+
+		if (FractionalGroundReactionImpulse > 0.0f)
+		{
+			Chaos::FPBDRigidParticleHandle* CharacterParticle = Private::GetParticleHandleFromComponent(Params.MovingComps.UpdatedPrimitive.Get());
+			Chaos::FPBDRigidParticleHandle* GroundParticle = nullptr;
+			FHitResult HitResult;
+			if (MoverComp->TryGetFloorCheckHitResult(HitResult))
+			{
+				GroundParticle = Private::GetParticleHandleFromComponent(HitResult.GetComponent());
+			}
+
+			if (CharacterParticle && GroundParticle)
+			{
+				const Chaos::FVec3 Impulse = -FractionalGroundReactionImpulse * CharacterParticle->M() * JumpUpwardsSpeed * MoverComp->GetUpDirection();
+				Private::ApplyImpulse(GroundParticle, Impulse, HitResult.ImpactPoint);
+			}
+		}
+	}
 }
 
 #if WITH_EDITOR
diff --git a/Engine/Plugins/Experimental/Mover/Source/Mover/Public/PhysicsMover/Transitions/PhysicsJumpCheck.h b/Engine/Plugins/Experimental/Mover/Source/Mover/Public/PhysicsMover/Transitions/PhysicsJumpCheck.h
--- a/Engine/Plugins/Experimental/Mover/Source/Mover/Public/PhysicsMover/Transitions/PhysicsJumpCheck.h
+++ b/Engine/Plugins/Experimental/Mover/Source/Mover/Public/PhysicsMover/Transitions/PhysicsJumpCheck.h
@@ -22,6 +22,14 @@ public:
 	UPROPERTY(BlueprintReadWrite, EditAnywhere, Category = Parameters, meta = (ClampMin = "0", UIMin = "0", ForceUnits = "cm/s"))
 	float JumpUpwardsSpeed = 500.0f;
 
+	/**
+	* Controls how much of the jump impulse the character will apply to the ground.
+	* A value of 0 means no impulse will be applied to the ground.
+	* A value of 1 means that the full equal and opposite jump impulse will be applied.
+	*/
+	UPROPERTY(BlueprintReadWrite, EditAnywhere, Category = Parameters, meta = (ClampMin = "0", UIMin = "0", ClampMax = "1", UIMax = "1"))
+	float FractionalGroundReactionImpulse = 0.0f;
+
 	/** Name of movement mode to transition to when jumping is activated */
 	UPROPERTY(BlueprintReadWrite, EditAnywhere, Category = Parameters)
 	FName TransitionToMode;
