Physics Mover - Fix for desyncs on jumping

diff --git a/Engine/Plugins/Experimental/Mover/Source/Mover/Private/Backends/MoverNetworkPhysicsLiaison.cpp b/Engine/Plugins/Experimental/Mover/Source/Mover/Private/Backends/MoverNetworkPhysicsLiaison.cpp
--- a/Engine/Plugins/Experimental/Mover/Source/Mover/Private/Backends/MoverNetworkPhysicsLiaison.cpp
+++ b/Engine/Plugins/Experimental/Mover/Source/Mover/Private/Backends/MoverNetworkPhysicsLiaison.cpp
@@ -207,22 +207,22 @@ const FString FNetworkPhysicsMoverState::DebugData()
 
 void UMoverNetworkPhysicsLiaisonComponent::GetCurrentInputData(OUT FMoverInputCmdContext& InputCmd) const
 {
-	InputCmd = NetInputCmd;
+	InputCmd = LatestInputCmd;
 }
 
 void UMoverNetworkPhysicsLiaisonComponent::GetCurrentStateData(OUT FMoverSyncState& SyncState) const
 {
-	SyncState = NetSyncState;
+	SyncState = LatestSyncState;
 }
 
 void UMoverNetworkPhysicsLiaisonComponent::SetCurrentInputData(const FMoverInputCmdContext& InputCmd)
 {
-	NetInputCmd = InputCmd;
+	LatestInputCmd = InputCmd;
 }
 
 void UMoverNetworkPhysicsLiaisonComponent::SetCurrentStateData(const FMoverSyncState& SyncState)
 {
-	NetSyncState = SyncState;
+	LatestSyncState = SyncState;
 }
 
 bool UMoverNetworkPhysicsLiaisonComponent::ValidateInputData(FMoverInputCmdContext& InputCmd) const
@@ -691,6 +691,7 @@ void UMoverNetworkPhysicsLiaisonComponent::InitializeCachedSyncState()
 		}
 
 		SyncStateRecord.Add(SimTime, SyncState);
+		LatestSyncState = SyncState;
 	}
 }
 
@@ -757,7 +758,7 @@ void UMoverNetworkPhysicsLiaisonComponent::ProduceInput_External(int32 PhysicsSt
 		}
 
 		Input.SyncState.SyncStateCollection.Empty();
-		Input.SyncState = SyncStateRecord.GetLatest();
+		Input.SyncState = LatestSyncState;
 
 		const FMoverTimeStep MoverTimeStep = bUsingAsyncPhysics ? GetCurrentAsyncMoverTimeStep_External() : GetCurrentMoverTimeStep(DeltaSeconds);
 		MoverComp->CachedLastSimTickTimeStep = MoverTimeStep;
@@ -765,8 +766,7 @@ void UMoverNetworkPhysicsLiaisonComponent::ProduceInput_External(int32 PhysicsSt
 		MoverComp->OnPreSimulationTick.Broadcast(MoverTimeStep, Input.InputCmd);
 
 		// This is required so that the physics thread can have a copy of the data to access
-		NetInputCmd = Input.InputCmd;
-		NetSyncState = Input.SyncState;
+		LatestInputCmd = Input.InputCmd;
 	}
 }
 
@@ -892,7 +892,7 @@ void UMoverNetworkPhysicsLiaisonComponent::ProcessInputs_Internal(int32 PhysicsS
 				{
 					FMoverAuxStateContext UnusedInvalidAuxState;
 					FMoverAuxStateContext UnusedAuxState;
-					MoverComp->OnSimulationPreRollback(&NetSyncState, &Input.SyncState, &UnusedInvalidAuxState, &UnusedAuxState);
+					MoverComp->OnSimulationPreRollback(&LatestSyncState, &Input.SyncState, &UnusedInvalidAuxState, &UnusedAuxState);
 					MoverComp->OnSimulationRollback(&Input.SyncState, &UnusedAuxState);
 				}
 			}
@@ -907,7 +907,7 @@ void UMoverNetworkPhysicsLiaisonComponent::ProcessInputs_Internal(int32 PhysicsS
 	}
 }
 
-void UMoverNetworkPhysicsLiaisonComponent::OnPreSimulate_Internal(const FPhysicsMoverSimulationTickParams& TickParams, const FPhysicsMoverAsyncInput& Input, OUT FPhysicsMoverAsyncOutput& Output) const
+void UMoverNetworkPhysicsLiaisonComponent::OnPreSimulate_Internal(const FPhysicsMoverSimulationTickParams& TickParams, const FPhysicsMoverAsyncInput& Input, OUT FPhysicsMoverAsyncOutput& Output)
 {
 	Chaos::EnsureIsInPhysicsThreadContext();
 
@@ -1118,6 +1118,11 @@ void UMoverNetworkPhysicsLiaisonComponent::OnPreSimulate_Internal(const FPhysics
 	Input.SyncState = Output.SyncState;
 	InputCmd->SuggestedMovementMode = NAME_None; // Should have already been set this frame
 
+	// This is required for cases where we run a second physics update (including generating
+	// input) before the output is saved on the game thread. We want to make sure the next
+	// physics tick starts with the sync state from the previous tick
+	LatestSyncState = Output.SyncState;
+
 	Output.bIsValid = true;
 }
 
diff --git a/Engine/Plugins/Experimental/Mover/Source/Mover/Public/Backends/MoverNetworkPhysicsLiaison.h b/Engine/Plugins/Experimental/Mover/Source/Mover/Public/Backends/MoverNetworkPhysicsLiaison.h
--- a/Engine/Plugins/Experimental/Mover/Source/Mover/Public/Backends/MoverNetworkPhysicsLiaison.h
+++ b/Engine/Plugins/Experimental/Mover/Source/Mover/Public/Backends/MoverNetworkPhysicsLiaison.h
@@ -202,7 +202,7 @@ public:
 	void PostPhysicsUpdate_External();
 
 	void ProcessInputs_Internal(int32 PhysicsStep, float DeltaTimeInSeconds, const FPhysicsMoverAsyncInput& Input) const;
-	void OnPreSimulate_Internal(const FPhysicsMoverSimulationTickParams& TickParams, const FPhysicsMoverAsyncInput& Input, OUT FPhysicsMoverAsyncOutput& SimOutput) const;
+	void OnPreSimulate_Internal(const FPhysicsMoverSimulationTickParams& TickParams, const FPhysicsMoverAsyncInput& Input, OUT FPhysicsMoverAsyncOutput& SimOutput);
 	void OnContactModification_Internal(const FPhysicsMoverAsyncInput& Input, Chaos::FCollisionContactModifier& Modifier) const;
 
 protected:
@@ -234,8 +234,8 @@ protected:
 	Chaos::FPhysicsSolver* GetPhysicsSolver() const;
 
 	// These are written to by the network input and state data
-	FMoverInputCmdContext NetInputCmd;
-	FMoverSyncState NetSyncState;
+	FMoverInputCmdContext LatestInputCmd;
+	FMoverSyncState LatestSyncState;
 
 	TUniquePtr<Chaos::FCharacterGroundConstraint> Constraint;
 	TObjectPtr<UMoverComponent> MoverComp;
