Mover: Adding support for arbitrary gravity directions to default movement modes and util libraries.

#virtualized

diff --git a/Engine/Plugins/Experimental/Mover/Source/Mover/Private/DefaultMovementSet/Modes/FallingMode.cpp b/Engine/Plugins/Experimental/Mover/Source/Mover/Private/DefaultMovementSet/Modes/FallingMode.cpp
--- a/Engine/Plugins/Experimental/Mover/Source/Mover/Private/DefaultMovementSet/Modes/FallingMode.cpp
+++ b/Engine/Plugins/Experimental/Mover/Source/Mover/Private/DefaultMovementSet/Modes/FallingMode.cpp
@@ -29,7 +29,7 @@ UFallingMode::UFallingMode(const FObjectInitializer& ObjectInitializer)
 	GameplayTags.AddTag(Mover_IsFalling);
 }
 
-constexpr float VERTICAL_SLOPE_NORMAL_Z = 0.001f; // Slope is vertical if Abs(Normal.Z) <= this threshold. Accounts for precision problems that sometimes angle normals slightly off horizontal for vertical surface.
+constexpr float VERTICAL_SLOPE_NORMAL_MAX_DOT = 0.001f; // Slope is vertical if Abs(UpDirection) <= this threshold. Accounts for precision problems that sometimes angle normals slightly off horizontal for vertical surface.
 
 void UFallingMode::OnGenerateMove(const FMoverTickStartData& StartState, const FMoverTimeStep& TimeStep, FProposedMove& OutProposedMove) const
 {
@@ -41,10 +41,12 @@ void UFallingMode::OnGenerateMove(const FMoverTickStartData& StartState, const F
 
 	const float DeltaSeconds = TimeStep.StepMs * 0.001f;
 
+	FVector UpDirection = MoverComp->GetUpDirection();
+	
 	// We don't want velocity limits to take the falling velocity component into account, since it is handled 
 	//   separately by the terminal velocity of the environment.
 	const FVector StartVelocity = StartingSyncState->GetVelocity_WorldSpace();
-	const FVector StartHorizontalVelocity = FVector(StartVelocity.X, StartVelocity.Y, 0.f);
+	const FVector StartHorizontalVelocity =  FVector::VectorPlaneProject(StartVelocity, UpDirection);
 
 	FFreeMoveParams Params;
 	if (CharacterInputs)
@@ -59,6 +61,10 @@ void UFallingMode::OnGenerateMove(const FMoverTickStartData& StartState, const F
 		Params.MoveInput = FVector::ZeroVector;
 	}
 
+	Params.MoveInput *= AirControlPercentage;
+	// Don't care about up axis input since falling - if up input matters that should probably be a different movement mode
+	Params.MoveInput = FVector::VectorPlaneProject(Params.MoveInput, UpDirection);
+	
 	FRotator IntendedOrientation_WorldSpace;
 	// If there's no intent from input to change orientation, use the current orientation
 	if (!CharacterInputs || CharacterInputs->OrientationIntent.IsNearlyZero())
@@ -70,6 +76,8 @@ void UFallingMode::OnGenerateMove(const FMoverTickStartData& StartState, const F
 		IntendedOrientation_WorldSpace = CharacterInputs->GetOrientationIntentDir_WorldSpace().ToOrientationRotator();
 	}
 
+	IntendedOrientation_WorldSpace = UMovementUtils::ApplyGravityToOrientationIntent(IntendedOrientation_WorldSpace, MoverComp->GetWorldToGravityTransform(), CommonLegacySettings->bShouldRemainVertical);
+	
 	Params.OrientationIntent = IntendedOrientation_WorldSpace;
 	Params.PriorVelocity = StartHorizontalVelocity;
 	Params.PriorOrientation = StartingSyncState->GetOrientation_WorldSpace();
@@ -79,17 +87,11 @@ void UFallingMode::OnGenerateMove(const FMoverTickStartData& StartState, const F
 	Params.MaxSpeed = CommonLegacySettings->MaxSpeed;
 	Params.Acceleration = CommonLegacySettings->Acceleration;
 	Params.Deceleration = FallingDeceleration;
-
-	Params.MoveInput *= AirControlPercentage;
-	// Don't care about Z axis input since falling - if Z input matters that should probably be a different movement mode
-	Params.MoveInput.Z = 0;
+	Params.WorldToGravityQuat = MoverComp->GetWorldToGravityTransform();
 
 	// Check if any current velocity values are over our terminal velocity - if so limit the move input in that direction and apply OverTerminalVelocityFallingDeceleration
-	if (Params.MoveInput.Dot(StartVelocity) > 0 && StartVelocity.Size2D() >= TerminalMovementPlaneSpeed)
+	if (Params.MoveInput.Dot(StartVelocity) > 0 && StartHorizontalVelocity.Size() >= TerminalMovementPlaneSpeed)
 	{
-		const FPlane MovementNormalPlane(StartVelocity, StartVelocity.GetSafeNormal());
-		Params.MoveInput = Params.MoveInput.ProjectOnTo(MovementNormalPlane);
-
 		Params.Deceleration = OverTerminalSpeedFallingDeceleration;
 	}
 	
@@ -98,14 +100,14 @@ void UFallingMode::OnGenerateMove(const FMoverTickStartData& StartState, const F
 	// limit our moveinput based on the floor we're on
 	if (SimBlackboard && SimBlackboard->TryGet(CommonBlackboard::LastFloorResult, LastFloorResult))
 	{
-		if (LastFloorResult.HitResult.IsValidBlockingHit() && LastFloorResult.HitResult.Normal.Z > VERTICAL_SLOPE_NORMAL_Z && !LastFloorResult.IsWalkableFloor())
+		if (LastFloorResult.HitResult.IsValidBlockingHit() && LastFloorResult.HitResult.Normal.Dot(UpDirection) > VERTICAL_SLOPE_NORMAL_MAX_DOT && !LastFloorResult.IsWalkableFloor())
 		{
 			// If acceleration is into the wall, limit contribution.
 			if (FVector::DotProduct(Params.MoveInput, LastFloorResult.HitResult.Normal) < 0.f)
 			{
 				// Allow movement parallel to the wall, but not into it because that may push us up.
-				const FVector Normal2D = LastFloorResult.HitResult.Normal.GetSafeNormal2D();
-				Params.MoveInput = FVector::VectorPlaneProject(Params.MoveInput, Normal2D);
+				const FVector FallingHitNormal = FVector::VectorPlaneProject( LastFloorResult.HitResult.Normal, -UpDirection).GetSafeNormal();
+				Params.MoveInput = FVector::VectorPlaneProject(Params.MoveInput, FallingHitNormal);
 			}
 		}
 	}
@@ -114,23 +116,25 @@ void UFallingMode::OnGenerateMove(const FMoverTickStartData& StartState, const F
 	const FVector VelocityWithGravity = StartVelocity + UMovementUtils::ComputeVelocityFromGravity(MoverComponent->GetGravityAcceleration(), DeltaSeconds);
 
 	//  If we are going faster than TerminalVerticalVelocity apply VerticalFallingDeceleration otherwise reset Z velocity to before we applied deceleration 
-	if (VelocityWithGravity.GetAbs().Z > TerminalVerticalSpeed)
+	if (VelocityWithGravity.GetAbs().Dot(UpDirection) > TerminalVerticalSpeed)
 	{
 		if (bShouldClampTerminalVerticalSpeed)
 		{
-			OutProposedMove.LinearVelocity.Z = FMath::Sign(VelocityWithGravity.Z) * TerminalVerticalSpeed;
+			UMovementUtils::SetGravityVerticalComponent(OutProposedMove.LinearVelocity, -TerminalVerticalSpeed, UpDirection);
 		}
 		else
 		{
-			float DesiredDeceleration = FMath::Abs(TerminalVerticalSpeed - FMath::Abs(VelocityWithGravity.Z)) / DeltaSeconds;
+			float DesiredDeceleration = FMath::Abs(TerminalVerticalSpeed - VelocityWithGravity.GetAbs().Dot(UpDirection)) / DeltaSeconds;
 			float DecelerationToApply = FMath::Min(DesiredDeceleration, VerticalFallingDeceleration);
-			DecelerationToApply = FMath::Sign(VelocityWithGravity.Z) * DecelerationToApply * DeltaSeconds;
-			OutProposedMove.LinearVelocity.Z = VelocityWithGravity.Z - DecelerationToApply;
+			DecelerationToApply = FMath::Sign(VelocityWithGravity.Dot(UpDirection)) * DecelerationToApply * DeltaSeconds;
+			FVector MaxUpDirVelocity = VelocityWithGravity * UpDirection - (UpDirection * DecelerationToApply);
+			
+			UMovementUtils::SetGravityVerticalComponent(OutProposedMove.LinearVelocity, MaxUpDirVelocity.Dot(UpDirection), UpDirection);
 		}
 	}
 	else
 	{
-		OutProposedMove.LinearVelocity.Z = VelocityWithGravity.Z;
+		UMovementUtils::SetGravityVerticalComponent(OutProposedMove.LinearVelocity, VelocityWithGravity.Dot(UpDirection), UpDirection);
 	}
 }
 
@@ -147,7 +151,6 @@ void UFallingMode::OnSimulationTick(const FSimulationTickParams& Params, FMoverT
 
 	FMoverDefaultSyncState& OutputSyncState = OutputState.SyncState.SyncStateCollection.FindOrAddMutableDataByType<FMoverDefaultSyncState>();
 
-	//const FVector GravityAccel = FVector(0.0f, 0.0f, -9800.0f);	// -9.8 m / sec^2
 	const float DeltaSeconds = Params.TimeStep.StepMs * 0.001f;
 	float PctTimeApplied = 0.f;
 
@@ -161,7 +164,6 @@ void UFallingMode::OnSimulationTick(const FSimulationTickParams& Params, FMoverT
 
 	OutputSyncState.MoveDirectionIntent = (ProposedMove.bHasDirIntent ? ProposedMove.DirectionIntent : FVector::ZeroVector);
 
-
 	// Use the orientation intent directly. If no intent is provided, use last frame's orientation. Note that we are assuming rotation changes can't fail. 
 	FRotator TargetOrient = StartingSyncState->GetOrientation_WorldSpace();
 
@@ -171,7 +173,7 @@ void UFallingMode::OnSimulationTick(const FSimulationTickParams& Params, FMoverT
 		TargetOrient += (ProposedMove.AngularVelocity * DeltaSeconds);
 	}
 	
-	const FVector PriorFallingVelocity = StartingSyncState->GetVelocity_WorldSpace();
+	const FVector StartingFallingVelocity = StartingSyncState->GetVelocity_WorldSpace();
 
 	//FVector MoveDelta = 0.5f * (PriorFallingVelocity + ProposedMove.LinearVelocity) * DeltaSeconds; 	// TODO: revive midpoint integration
 	FVector MoveDelta = ProposedMove.LinearVelocity * DeltaSeconds;
@@ -182,9 +184,7 @@ void UFallingMode::OnSimulationTick(const FSimulationTickParams& Params, FMoverT
 	UMovementUtils::TrySafeMoveUpdatedComponent(Params.MovingComps, MoveDelta, TargetOrientQuat, true, Hit, ETeleportType::None, MoveRecord);
 
 	// Compute final velocity based on how long we actually go until we get a hit.
-	FVector NewFallingVelocity = StartingSyncState->GetVelocity_WorldSpace();
-	//NewFallingVelocity.Z = 0.5f * (PriorFallingVelocity.Z + (ProposedMove.LinearVelocity.Z * Hit.Time));	// TODO: revive midpoint integration
-	NewFallingVelocity.Z = ProposedMove.LinearVelocity.Z * Hit.Time;
+	FVector NewFallingVelocity = StartingFallingVelocity + (ProposedMove.LinearVelocity * Hit.Time * MoveDelta);
 
 	FFloorCheckResult LandingFloor;
 
diff --git a/Engine/Plugins/Experimental/Mover/Source/Mover/Private/DefaultMovementSet/Modes/FlyingMode.cpp b/Engine/Plugins/Experimental/Mover/Source/Mover/Private/DefaultMovementSet/Modes/FlyingMode.cpp
--- a/Engine/Plugins/Experimental/Mover/Source/Mover/Private/DefaultMovementSet/Modes/FlyingMode.cpp
+++ b/Engine/Plugins/Experimental/Mover/Source/Mover/Private/DefaultMovementSet/Modes/FlyingMode.cpp
@@ -48,6 +48,8 @@ void UFlyingMode::OnGenerateMove(const FMoverTickStartData& StartState, const FM
 	{
 		IntendedOrientation_WorldSpace = CharacterInputs->GetOrientationIntentDir_WorldSpace().ToOrientationRotator();
 	}
+
+	IntendedOrientation_WorldSpace = UMovementUtils::ApplyGravityToOrientationIntent(IntendedOrientation_WorldSpace, MoverComp->GetWorldToGravityTransform(), CommonLegacySettings->bShouldRemainVertical);
 	
 	Params.OrientationIntent = IntendedOrientation_WorldSpace;
 	Params.PriorVelocity = StartingSyncState->GetVelocity_WorldSpace();
@@ -58,6 +60,7 @@ void UFlyingMode::OnGenerateMove(const FMoverTickStartData& StartState, const FM
 	Params.Acceleration = CommonLegacySettings->Acceleration;
 	Params.Deceleration = CommonLegacySettings->Deceleration;
 	Params.DeltaSeconds = DeltaSeconds;
+	Params.WorldToGravityQuat = MoverComp->GetWorldToGravityTransform();
 	
 	OutProposedMove = UAirMovementUtils::ComputeControlledFreeMove(Params);
 }
diff --git a/Engine/Plugins/Experimental/Mover/Source/Mover/Private/DefaultMovementSet/Modes/NavWalkingMode.cpp b/Engine/Plugins/Experimental/Mover/Source/Mover/Private/DefaultMovementSet/Modes/NavWalkingMode.cpp
--- a/Engine/Plugins/Experimental/Mover/Source/Mover/Private/DefaultMovementSet/Modes/NavWalkingMode.cpp
+++ b/Engine/Plugins/Experimental/Mover/Source/Mover/Private/DefaultMovementSet/Modes/NavWalkingMode.cpp
@@ -39,6 +39,7 @@ void UNavWalkingMode::OnGenerateMove(const FMoverTickStartData& StartState, cons
 	const float DeltaSeconds = TimeStep.StepMs * 0.001f;
 	FFloorCheckResult LastFloorResult;
 	FVector MovementNormal;
+	FVector UpDirection = MoverComp->GetUpDirection();
 
 	UMoverBlackboard* SimBlackboard = MoverComp->GetSimBlackboard_Mutable();
 
@@ -62,6 +63,8 @@ void UNavWalkingMode::OnGenerateMove(const FMoverTickStartData& StartState, cons
 	{
 		IntendedOrientation_WorldSpace = CharacterInputs->GetOrientationIntentDir_WorldSpace().ToOrientationRotator();
 	}
+
+	IntendedOrientation_WorldSpace = UMovementUtils::ApplyGravityToOrientationIntent(IntendedOrientation_WorldSpace, MoverComp->GetWorldToGravityTransform(), CommonLegacySettings->bShouldRemainVertical);
 	
 	FGroundMoveParams Params;
 
@@ -86,6 +89,8 @@ void UNavWalkingMode::OnGenerateMove(const FMoverTickStartData& StartState, cons
 	Params.Acceleration = CommonLegacySettings->Acceleration;
 	Params.Deceleration = CommonLegacySettings->Deceleration;
 	Params.DeltaSeconds = DeltaSeconds;
+	Params.WorldToGravityQuat = MoverComp->GetWorldToGravityTransform();
+	Params.UpDirection = UpDirection;
 	
 	if (Params.MoveInput.SizeSquared() > 0.f && !UMovementUtils::IsExceedingMaxSpeed(Params.PriorVelocity, CommonLegacySettings->MaxSpeed))
 	{
diff --git a/Engine/Plugins/Experimental/Mover/Source/Mover/Private/DefaultMovementSet/Modes/SwimmingMode.cpp b/Engine/Plugins/Experimental/Mover/Source/Mover/Private/DefaultMovementSet/Modes/SwimmingMode.cpp
--- a/Engine/Plugins/Experimental/Mover/Source/Mover/Private/DefaultMovementSet/Modes/SwimmingMode.cpp
+++ b/Engine/Plugins/Experimental/Mover/Source/Mover/Private/DefaultMovementSet/Modes/SwimmingMode.cpp
@@ -8,6 +8,7 @@
 #include "MoverComponent.h"
 #include "DefaultMovementSet/Settings/CommonLegacyMovementSettings.h"
 #include "MoverLog.h"
+#include "MoveLibrary/MovementUtils.h"
 
 #include UE_INLINE_GENERATED_CPP_BY_NAME(SwimmingMode)
 
@@ -131,6 +132,8 @@ void USwimmingMode::OnGenerateMove(const FMoverTickStartData& StartState, const
 		{
 			IntendedOrientation_WorldSpace = CharacterInputs->GetOrientationIntentDir_WorldSpace().ToOrientationRotator();
 		}
+
+		IntendedOrientation_WorldSpace = UMovementUtils::ApplyGravityToOrientationIntent(IntendedOrientation_WorldSpace, MoverComp->GetWorldToGravityTransform(), CommonLegacySettings->bShouldRemainVertical);
 		
 		FWaterMoveParams Params;
 		if (CharacterInputs)
@@ -156,6 +159,7 @@ void USwimmingMode::OnGenerateMove(const FMoverTickStartData& StartState, const
 		Params.DeltaSeconds = DeltaSeconds;
 		Params.MoveSpeed = WaterSpeed;
 		Params.MoveAcceleration = WaterAcceleration;
+		Params.WorldToGravityQuat = MoverComp->GetWorldToGravityTransform();
 
 		// Calculate the move
 		OutProposedMove = UWaterMovementUtils::ComputeControlledWaterMove(Params);
diff --git a/Engine/Plugins/Experimental/Mover/Source/Mover/Private/DefaultMovementSet/Modes/WalkingMode.cpp b/Engine/Plugins/Experimental/Mover/Source/Mover/Private/DefaultMovementSet/Modes/WalkingMode.cpp
--- a/Engine/Plugins/Experimental/Mover/Source/Mover/Private/DefaultMovementSet/Modes/WalkingMode.cpp
+++ b/Engine/Plugins/Experimental/Mover/Source/Mover/Private/DefaultMovementSet/Modes/WalkingMode.cpp
@@ -37,6 +37,7 @@ void UWalkingMode::OnGenerateMove(const FMoverTickStartData& StartState, const F
 	FVector MovementNormal;
 
 	UMoverBlackboard* SimBlackboard = MoverComp->GetSimBlackboard_Mutable();
+	FVector UpDirection = MoverComp->GetUpDirection();
 
 	// Try to use the floor as the basis for the intended move direction (i.e. try to walk along slopes, rather than into them)
 	if (SimBlackboard && SimBlackboard->TryGet(CommonBlackboard::LastFloorResult, LastFloorResult) && LastFloorResult.IsWalkableFloor())
@@ -45,11 +46,9 @@ void UWalkingMode::OnGenerateMove(const FMoverTickStartData& StartState, const F
 	}
 	else
 	{
-		MovementNormal = MoverComp->GetUpDirection();
+		MovementNormal = UpDirection;
 	}
 
-	
-
 	FRotator IntendedOrientation_WorldSpace;
 	// If there's no intent from input to change orientation, use the current orientation
 	if (!CharacterInputs || CharacterInputs->OrientationIntent.IsNearlyZero())
@@ -61,6 +60,8 @@ void UWalkingMode::OnGenerateMove(const FMoverTickStartData& StartState, const F
 		IntendedOrientation_WorldSpace = CharacterInputs->GetOrientationIntentDir_WorldSpace().ToOrientationRotator();
 	}
 
+	IntendedOrientation_WorldSpace = UMovementUtils::ApplyGravityToOrientationIntent(IntendedOrientation_WorldSpace, MoverComp->GetWorldToGravityTransform(), CommonLegacySettings->bShouldRemainVertical);
+	
 	FGroundMoveParams Params;
 
 	if (CharacterInputs)
@@ -86,6 +87,8 @@ void UWalkingMode::OnGenerateMove(const FMoverTickStartData& StartState, const F
 	Params.Acceleration = CommonLegacySettings->Acceleration;
 	Params.Deceleration = CommonLegacySettings->Deceleration;
 	Params.DeltaSeconds = DeltaSeconds;
+	Params.WorldToGravityQuat = MoverComp->GetWorldToGravityTransform();
+	Params.UpDirection = UpDirection;
 	
 	if (Params.MoveInput.SizeSquared() > 0.f && !UMovementUtils::IsExceedingMaxSpeed(Params.PriorVelocity, CommonLegacySettings->MaxSpeed))
 	{
@@ -139,12 +142,13 @@ void UWalkingMode::OnSimulationTick(const FSimulationTickParams& Params, FMoverT
 	FFloorCheckResult CurrentFloor;
 	UMoverBlackboard* SimBlackboard = MoverComp->GetSimBlackboard_Mutable();
 
+	FVector UpDirection = MoverComp->GetUpDirection();
+	FMovingComponentSet MovingComponents(MoverComp);
+	
 	// If we don't have cached floor information, we need to search for it again
 	if (!SimBlackboard->TryGet(CommonBlackboard::LastFloorResult, CurrentFloor))
 	{
-		UFloorQueryUtils::FindFloor(UpdatedComponent, UpdatedPrimitive,
-			CommonLegacySettings->FloorSweepDistance, CommonLegacySettings->MaxWalkSlopeCosine,
-			UpdatedPrimitive->GetComponentLocation(), CurrentFloor);
+		UFloorQueryUtils::FindFloor(MovingComponents, CommonLegacySettings->FloorSweepDistance, CommonLegacySettings->MaxWalkSlopeCosine, UpdatedPrimitive->GetComponentLocation(), CurrentFloor);
 	}
 
 	OutputSyncState.MoveDirectionIntent = (ProposedMove.bHasDirIntent ? ProposedMove.DirectionIntent : FVector::ZeroVector);
@@ -193,13 +197,12 @@ void UWalkingMode::OnSimulationTick(const FSimulationTickParams& Params, FMoverT
 			PercentTimeAppliedSoFar = MoveHitResult.Time;		
 
 			// Check if the blockage is a walkable ramp rising in front of us
-			// JAH TODO: change this normal.Z test to consider vs the movement plane (not just Z=0 plane)
-			if ((MoveHitResult.Time > 0.f) && (MoveHitResult.Normal.Z > UE_KINDA_SMALL_NUMBER) && 
-			    UFloorQueryUtils::IsHitSurfaceWalkable(MoveHitResult, CommonLegacySettings->MaxWalkSlopeCosine))
+			if ((MoveHitResult.Time > 0.f) && (MoveHitResult.Normal.Dot(UpDirection) > UE_KINDA_SMALL_NUMBER) && 
+			    UFloorQueryUtils::IsHitSurfaceWalkable(MoveHitResult, UpDirection, CommonLegacySettings->MaxWalkSlopeCosine))
 			{
 				// It's a walkable ramp, so cut up the move and attempt to move the remainder of it along the ramp's surface, possibly generating another hit
 				const float PercentTimeRemaining = 1.f - PercentTimeAppliedSoFar;
-				CurMoveDelta = UGroundMovementUtils::ComputeDeflectedMoveOntoRamp(CurMoveDelta * PercentTimeRemaining, MoveHitResult, CommonLegacySettings->MaxWalkSlopeCosine, CurrentFloor.bLineTrace);
+				CurMoveDelta = UGroundMovementUtils::ComputeDeflectedMoveOntoRamp(CurMoveDelta * PercentTimeRemaining, MoverComp->GetUpDirection(), MoveHitResult, CommonLegacySettings->MaxWalkSlopeCosine, CurrentFloor.bLineTrace);
 				UMovementUtils::TrySafeMoveUpdatedComponent(Params.MovingComps, CurMoveDelta, TargetOrientQuat, true, MoveHitResult, ETeleportType::None, MoveRecord);
 				LastMoveSeconds = PercentTimeRemaining * LastMoveSeconds;
 
@@ -239,9 +242,7 @@ void UWalkingMode::OnSimulationTick(const FSimulationTickParams& Params, FMoverT
 		}
 
 		// Search for the floor we've ended up on
-		UFloorQueryUtils::FindFloor(UpdatedComponent, UpdatedPrimitive,
-			CommonLegacySettings->FloorSweepDistance, CommonLegacySettings->MaxWalkSlopeCosine,
-			UpdatedPrimitive->GetComponentLocation(), CurrentFloor);
+		UFloorQueryUtils::FindFloor(MovingComponents, CommonLegacySettings->FloorSweepDistance, CommonLegacySettings->MaxWalkSlopeCosine, UpdatedPrimitive->GetComponentLocation(), CurrentFloor);
 
 		if (CurrentFloor.IsWalkableFloor())
 		{
@@ -261,16 +262,14 @@ void UWalkingMode::OnSimulationTick(const FSimulationTickParams& Params, FMoverT
 	else
 	{
 		// If the actor isn't moving we still need to check if they have a valid floor
-		UFloorQueryUtils::FindFloor(UpdatedComponent, UpdatedPrimitive,
-			CommonLegacySettings->FloorSweepDistance, CommonLegacySettings->MaxWalkSlopeCosine,
-			UpdatedPrimitive->GetComponentLocation(), CurrentFloor);
+		UFloorQueryUtils::FindFloor(MovingComponents, CommonLegacySettings->FloorSweepDistance, CommonLegacySettings->MaxWalkSlopeCosine, UpdatedPrimitive->GetComponentLocation(), CurrentFloor);
 		
 		FHitResult Hit(CurrentFloor.HitResult);
 		if (Hit.bStartPenetrating)
 		{
 			// The floor check failed because it started in penetration
 			// We do not want to try to move downward because the downward sweep failed, rather we'd like to try to pop out of the floor.
-			Hit.TraceEnd = Hit.TraceStart + FVector(0.f, 0.f, 2.4f);
+			Hit.TraceEnd = Hit.TraceStart + UpDirection * 2.4;
 			FVector RequestedAdjustment = UMovementUtils::ComputePenetrationAdjustment(Hit);
 			
 			const EMoveComponentFlags IncludeBlockingOverlapsWithoutEvents = (MOVECOMP_NeverIgnoreBlockingOverlaps | MOVECOMP_DisableBlockingOverlapDispatch);
diff --git a/Engine/Plugins/Experimental/Mover/Source/Mover/Private/DefaultMovementSet/MovementModifiers/StanceModifier.cpp b/Engine/Plugins/Experimental/Mover/Source/Mover/Private/DefaultMovementSet/MovementModifiers/StanceModifier.cpp
--- a/Engine/Plugins/Experimental/Mover/Source/Mover/Private/DefaultMovementSet/MovementModifiers/StanceModifier.cpp
+++ b/Engine/Plugins/Experimental/Mover/Source/Mover/Private/DefaultMovementSet/MovementModifiers/StanceModifier.cpp
@@ -254,7 +254,8 @@ void FStanceModifier::AdjustCapsule(UMoverComponent* MoverComp, float OldHalfHei
 	}
 
 	const FVector CapsuleOffset = MoverComp->GetUpDirection() * (bExpanding ? HalfHeightDifference : -HalfHeightDifference);
-	const FVector VisualOffset = MoverComp->GetUpDirection() * (bExpanding ? -HalfHeightDifference : HalfHeightDifference);
+	// This is only getting used to add relative offset - so assuming z is up is fine here
+	const FVector VisualOffset = FVector::UpVector * (bExpanding ? -HalfHeightDifference : HalfHeightDifference);
 	
 	// Adjust location of capsule as setting it's size left it floating
 	if (!bExpanding || MoverComp->GetVelocity().Length() <= 0)
diff --git a/Engine/Plugins/Experimental/Mover/Source/Mover/Private/MoveLibrary/AirMovementUtils.cpp b/Engine/Plugins/Experimental/Mover/Source/Mover/Private/MoveLibrary/AirMovementUtils.cpp
--- a/Engine/Plugins/Experimental/Mover/Source/Mover/Private/MoveLibrary/AirMovementUtils.cpp
+++ b/Engine/Plugins/Experimental/Mover/Source/Mover/Private/MoveLibrary/AirMovementUtils.cpp
@@ -25,15 +25,12 @@ FProposedMove UAirMovementUtils::ComputeControlledFreeMove(const FFreeMoveParams
 	ComputeVelocityParams.Acceleration = InParams.Acceleration;
 	
 	OutMove.LinearVelocity = UMovementUtils::ComputeVelocity(ComputeVelocityParams);
-
-	// JAH TODO: this is where we can perform turning, based on aux settings. For now, just snap to the intended final orientation.
-	FVector IntendedFacingDir = InParams.OrientationIntent.RotateVector(FVector::ForwardVector).GetSafeNormal();
-	OutMove.AngularVelocity = UMovementUtils::ComputeAngularVelocity(InParams.PriorOrientation, IntendedFacingDir.ToOrientationRotator(), InParams.DeltaSeconds, InParams.TurningRate);
+	OutMove.AngularVelocity = UMovementUtils::ComputeAngularVelocity(InParams.PriorOrientation, InParams.OrientationIntent, InParams.WorldToGravityQuat, InParams.DeltaSeconds, InParams.TurningRate);
 
 	return OutMove;
 }
 
-bool UAirMovementUtils::IsValidLandingSpot(const FMovingComponentSet& MovingComps, const FVector& Location, const FHitResult& Hit, float FloorSweepDistance, float WalkableFloorZ, FFloorCheckResult& OutFloorResult)
+bool UAirMovementUtils::IsValidLandingSpot(const FMovingComponentSet& MovingComps, const FVector& Location, const FHitResult& Hit, float FloorSweepDistance, float MaxWalkSlopeCosine, FFloorCheckResult& OutFloorResult)
 {
 	OutFloorResult.Clear();
 
@@ -48,15 +45,13 @@ bool UAirMovementUtils::IsValidLandingSpot(const FMovingComponentSet& MovingComp
 	}
 
 	// Reject unwalkable floor normals.
-	if (!UFloorQueryUtils::IsHitSurfaceWalkable(Hit, WalkableFloorZ))
+	if (!UFloorQueryUtils::IsHitSurfaceWalkable(Hit, MovingComps.MoverComponent->GetUpDirection(), MaxWalkSlopeCosine))
 	{
 		return false;
 	}
 
 	// Make sure floor test passes here.
-	UFloorQueryUtils::FindFloor(MovingComps.UpdatedComponent.Get(), MovingComps.UpdatedPrimitive.Get(),
-		FloorSweepDistance, WalkableFloorZ,
-		Location, OutFloorResult);
+	UFloorQueryUtils::FindFloor(MovingComps, FloorSweepDistance, MaxWalkSlopeCosine, Location, OutFloorResult);
 
 	if (!OutFloorResult.IsWalkableFloor())
 	{
diff --git a/Engine/Plugins/Experimental/Mover/Source/Mover/Private/MoveLibrary/BasedMovementUtils.cpp b/Engine/Plugins/Experimental/Mover/Source/Mover/Private/MoveLibrary/BasedMovementUtils.cpp
--- a/Engine/Plugins/Experimental/Mover/Source/Mover/Private/MoveLibrary/BasedMovementUtils.cpp
+++ b/Engine/Plugins/Experimental/Mover/Source/Mover/Private/MoveLibrary/BasedMovementUtils.cpp
@@ -383,12 +383,15 @@ void UBasedMovementUtils::UpdateSimpleBasedMovement(UMoverComponent* TargetMover
 					DeltaQuat = CurrentBaseQuat * LastAppliedBaseInfo.Rotation.Inverse();
 					WorldTargetQuat = DeltaQuat * WorldTargetQuat;
 
-					// TODO: make this respect the "up" direction, rather than assuming +Z = up
 					FVector TargetForwVector = WorldTargetQuat.GetForwardVector();
-					TargetForwVector.Z = 0.f;
+					TargetForwVector = FVector::VectorPlaneProject(TargetForwVector, -TargetMoverComp->GetUpDirection());
 					TargetForwVector.Normalize();
 
-					WorldTargetQuat = UKismetMathLibrary::MakeRotFromX(TargetForwVector).Quaternion();
+					FVector TargetRightVector = WorldTargetQuat.GetRightVector();
+					TargetRightVector = FVector::VectorPlaneProject(TargetRightVector, -TargetMoverComp->GetUpDirection());
+					TargetRightVector.Normalize();
+					
+					WorldTargetQuat = UKismetMathLibrary::MakeRotFromXY(TargetForwVector, TargetRightVector).Quaternion();
 				}
 
 				if (bDidBaseLocationChange || bDidBaseRotationChange)
diff --git a/Engine/Plugins/Experimental/Mover/Source/Mover/Private/MoveLibrary/FloorQueryUtils.cpp b/Engine/Plugins/Experimental/Mover/Source/Mover/Private/MoveLibrary/FloorQueryUtils.cpp
--- a/Engine/Plugins/Experimental/Mover/Source/Mover/Private/MoveLibrary/FloorQueryUtils.cpp
+++ b/Engine/Plugins/Experimental/Mover/Source/Mover/Private/MoveLibrary/FloorQueryUtils.cpp
@@ -1,6 +1,9 @@
 ﻿// Copyright Epic Games, Inc. All Rights Reserved.
 
 #include "MoveLibrary/FloorQueryUtils.h"
+
+#include "MoverComponent.h"
+#include "Components/CapsuleComponent.h"
 #include "MoveLibrary/MovementUtils.h"
 #include "Engine/World.h"
 
@@ -11,9 +14,9 @@ namespace UE::FloorQueryUtility
 	const float SWEEP_EDGE_REJECT_DISTANCE = 0.15f;
 }
 
-void UFloorQueryUtils::FindFloor(const USceneComponent* UpdatedComponent, const UPrimitiveComponent* UpdatedPrimitive, float FloorSweepDistance, float MaxWalkSlopeCosine, const FVector& Location, FFloorCheckResult& OutFloorResult)
+void UFloorQueryUtils::FindFloor(const FMovingComponentSet& MovingComps, float FloorSweepDistance, float MaxWalkSlopeCosine, const FVector& Location, FFloorCheckResult& OutFloorResult)
 {
-	if (!UpdatedComponent || !UpdatedComponent->IsQueryCollisionEnabled())
+	if (!MovingComps.UpdatedComponent->IsQueryCollisionEnabled())
 	{
 		OutFloorResult.Clear();
 		return;
@@ -21,22 +24,25 @@ void UFloorQueryUtils::FindFloor(const USceneComponent* UpdatedComponent, const
 
 	// Sweep for the floor
 	// TODO: Might need to plug in a different value for LineTraceDistance - using the same value as FloorSweepDistance for now - function takes both so we can plug in different values if needed
-	ComputeFloorDist(UpdatedComponent, UpdatedPrimitive,FloorSweepDistance, FloorSweepDistance, MaxWalkSlopeCosine, Location, OutFloorResult);
+	ComputeFloorDist(MovingComps,FloorSweepDistance, FloorSweepDistance, MaxWalkSlopeCosine, Location, OutFloorResult);
 }
 
-void UFloorQueryUtils::ComputeFloorDist(const USceneComponent* UpdatedComponent, const UPrimitiveComponent* UpdatedPrimitive, float LineTraceDistance, float FloorSweepDistance, float MaxWalkSlopeCosine, const FVector& Location, FFloorCheckResult& OutFloorResult)
+void UFloorQueryUtils::ComputeFloorDist(const FMovingComponentSet& MovingComps, float LineTraceDistance, float FloorSweepDistance, float MaxWalkSlopeCosine, const FVector& Location, FFloorCheckResult& OutFloorResult)
 {
 	OutFloorResult.Clear();
 
-	FCollisionQueryParams QueryParams(SCENE_QUERY_STAT(ComputeFloorDist), false, UpdatedPrimitive->GetOwner());
+	FCollisionQueryParams QueryParams(SCENE_QUERY_STAT(ComputeFloorDist), false, MovingComps.UpdatedPrimitive->GetOwner());
 	FCollisionResponseParams ResponseParam;
-	UMovementUtils::InitCollisionParams(UpdatedPrimitive, QueryParams, ResponseParam);
-	const ECollisionChannel CollisionChannel = UpdatedComponent->GetCollisionObjectType();
+	UMovementUtils::InitCollisionParams(MovingComps.UpdatedPrimitive.Get(), QueryParams, ResponseParam);
+	const ECollisionChannel CollisionChannel = MovingComps.UpdatedComponent->GetCollisionObjectType();
 
 	// TODO: pluggable shapes
 	float PawnRadius = 0.0f;
 	float PawnHalfHeight = 0.0f;
-	UpdatedPrimitive->CalcBoundingCylinder(PawnRadius, PawnHalfHeight);
+	UCapsuleComponent* CapsuleComponent = Cast<UCapsuleComponent>(MovingComps.UpdatedComponent);
+	CapsuleComponent->GetScaledCapsuleSize(PawnRadius, PawnHalfHeight);
+	//MovingComps.UpdatedPrimitive->CalcBoundingCylinder(PawnRadius, PawnHalfHeight);
+	FVector UpDirection = MovingComps.MoverComponent->GetUpDirection();
 
 	bool bBlockingHit = false;
 	
@@ -52,15 +58,14 @@ void UFloorQueryUtils::ComputeFloorDist(const USceneComponent* UpdatedComponent,
 		FCollisionShape CapsuleShape = FCollisionShape::MakeCapsule(PawnRadius, PawnHalfHeight - ShrinkHeight);
 
 		FHitResult Hit(1.f);
-		// TODO: arbitrary direction
-		FVector SweepDirection = FVector(0.f, 0.f, -TraceDist);
-		bBlockingHit = FloorSweepTest(UpdatedPrimitive, Hit, Location, Location + SweepDirection, CollisionChannel, CapsuleShape, QueryParams, ResponseParam);
+		FVector SweepDirection = UpDirection * -TraceDist;
+		bBlockingHit = FloorSweepTest(MovingComps, Hit, Location, Location + SweepDirection, CollisionChannel, CapsuleShape, QueryParams, ResponseParam);
 
 		if (bBlockingHit)
 		{
 			// Reject hits adjacent to us, we only care about hits on the bottom portion of our capsule.
 			// Check 2D distance to impact point, reject if within a tolerance from radius.
-			if (Hit.bStartPenetrating || !IsWithinEdgeTolerance(Location, Hit.ImpactPoint, CapsuleShape.Capsule.Radius))
+			if (Hit.bStartPenetrating || !IsWithinEdgeTolerance(Location, Hit.ImpactPoint, CapsuleShape.Capsule.Radius, UpDirection))
 			{
 				// Use a capsule with a slightly smaller radius and shorter height to avoid the adjacent object.
 				// Capsule must not be nearly zero or the trace will fall back to a line trace from the start point and have the wrong length.
@@ -69,11 +74,11 @@ void UFloorQueryUtils::ComputeFloorDist(const USceneComponent* UpdatedComponent,
 				{
 					ShrinkHeight = (PawnHalfHeight - PawnRadius) * (1.f - ShrinkScaleOverlap);
 					TraceDist = FloorSweepDistance + ShrinkHeight;
-					SweepDirection = FVector(0.f, 0.f, -TraceDist);
+					SweepDirection = UpDirection * -TraceDist;
 					CapsuleShape.Capsule.HalfHeight = FMath::Max(PawnHalfHeight - ShrinkHeight, CapsuleShape.Capsule.Radius);
 					Hit.Reset(1.f, false);
 
-					bBlockingHit = FloorSweepTest(UpdatedPrimitive, Hit, Location, Location + SweepDirection, CollisionChannel, CapsuleShape, QueryParams, ResponseParam);
+					bBlockingHit = FloorSweepTest(MovingComps, Hit, Location, Location + SweepDirection, CollisionChannel, CapsuleShape, QueryParams, ResponseParam);
 				}
 			}
 
@@ -84,7 +89,7 @@ void UFloorQueryUtils::ComputeFloorDist(const USceneComponent* UpdatedComponent,
 			const float SweepResult = FMath::Max(-MaxPenetrationAdjust, Hit.Time * TraceDist - ShrinkHeight);
 
 			OutFloorResult.SetFromSweep(Hit, SweepResult, false);
-			if (Hit.IsValidBlockingHit() && IsHitSurfaceWalkable(Hit, MaxWalkSlopeCosine))
+			if (Hit.IsValidBlockingHit() && IsHitSurfaceWalkable(Hit, UpDirection, MaxWalkSlopeCosine))
 			{
 				if (SweepResult <= FloorSweepDistance)
 				{
@@ -110,11 +115,11 @@ void UFloorQueryUtils::ComputeFloorDist(const USceneComponent* UpdatedComponent,
 		const float ShrinkHeight = PawnHalfHeight;
 		const FVector LineTraceStart = Location;	
 		const float TraceDist = LineTraceDistance + ShrinkHeight;
-		const FVector Down = FVector(0.f, 0.f, -TraceDist);
+		const FVector Down = UpDirection * -TraceDist;
 		QueryParams.TraceTag = SCENE_QUERY_STAT_NAME_ONLY(FloorLineTrace);
 
 		FHitResult Hit(1.f);
-		bBlockingHit = UpdatedComponent->GetWorld()->LineTraceSingleByChannel(Hit, LineTraceStart, LineTraceStart + Down, CollisionChannel, QueryParams, ResponseParam);
+		bBlockingHit = MovingComps.UpdatedComponent->GetWorld()->LineTraceSingleByChannel(Hit, LineTraceStart, LineTraceStart + Down, CollisionChannel, QueryParams, ResponseParam);
 		
 		if (bBlockingHit && Hit.Time > 0.f)
 		{
@@ -124,7 +129,7 @@ void UFloorQueryUtils::ComputeFloorDist(const USceneComponent* UpdatedComponent,
 			const float LineResult = FMath::Max(-MaxPenetrationAdjust, Hit.Time * TraceDist - ShrinkHeight);
 			
 			OutFloorResult.bBlockingHit = true;
-			if (LineResult <= LineTraceDistance && IsHitSurfaceWalkable(Hit, MaxWalkSlopeCosine))
+			if (LineResult <= LineTraceDistance && IsHitSurfaceWalkable(Hit, UpDirection, MaxWalkSlopeCosine))
 			{
 				OutFloorResult.SetFromLineTrace(Hit, OutFloorResult.FloorDist, LineResult, true);
 				return;
@@ -136,22 +141,20 @@ void UFloorQueryUtils::ComputeFloorDist(const USceneComponent* UpdatedComponent,
 	OutFloorResult.bWalkableFloor = false;
 }
 
-bool UFloorQueryUtils::FloorSweepTest(const UPrimitiveComponent* UpdatedPrimitive, FHitResult& OutHit, const FVector& Start, const FVector& End, ECollisionChannel TraceChannel, const struct FCollisionShape& CollisionShape, const struct FCollisionQueryParams& Params, const struct FCollisionResponseParams& ResponseParam)
+bool UFloorQueryUtils::FloorSweepTest(const FMovingComponentSet& MovingComps, FHitResult& OutHit, const FVector& Start, const FVector& End, ECollisionChannel TraceChannel, const struct FCollisionShape& CollisionShape, const struct FCollisionQueryParams& Params, const struct FCollisionResponseParams& ResponseParam)
 {
 	bool bBlockingHit = false;
 
-	if (UpdatedPrimitive)
+	if (MovingComps.UpdatedPrimitive.IsValid())
 	{
-		bBlockingHit = UpdatedPrimitive->GetWorld()->SweepSingleByChannel(OutHit, Start, End, FQuat::Identity, TraceChannel, CollisionShape, Params, ResponseParam);
+		bBlockingHit = MovingComps.UpdatedPrimitive->GetWorld()->SweepSingleByChannel(OutHit, Start, End, MovingComps.UpdatedComponent->GetComponentRotation().Quaternion(), TraceChannel, CollisionShape, Params, ResponseParam);
 	}
 
 	return bBlockingHit;
 }
 
-bool UFloorQueryUtils::IsHitSurfaceWalkable(const FHitResult& Hit, float MaxWalkSlopeCosine)
+bool UFloorQueryUtils::IsHitSurfaceWalkable(const FHitResult& Hit, FVector UpDirection, float MaxWalkSlopeCosine)
 {
-	// JAH TODO: refactor this to support arbitrary movement planes (currently assumes it's the Z=0 plane)
-
 	if (!Hit.IsValidBlockingHit())
 	{
 		// No hit, or starting in penetration
@@ -159,23 +162,23 @@ bool UFloorQueryUtils::IsHitSurfaceWalkable(const FHitResult& Hit, float MaxWalk
 	}
 
 	// Never walk up vertical surfaces.
-	if (Hit.ImpactNormal.Z < KINDA_SMALL_NUMBER)
+	if (Hit.ImpactNormal.Dot(UpDirection) < KINDA_SMALL_NUMBER)
 	{
 		return false;
 	}
 
-	float TestWalkableZ = MaxWalkSlopeCosine;
+	float TestWalkableSlopeCosine = MaxWalkSlopeCosine;
 
-	// See if this component overrides the walkable floor z.
+	// See if this component overrides the walkable floor slope cosine.
 	const UPrimitiveComponent* HitComponent = Hit.Component.Get();
 	if (HitComponent)
 	{
 		const FWalkableSlopeOverride& SlopeOverride = HitComponent->GetWalkableSlopeOverride();
-		TestWalkableZ = SlopeOverride.ModifyWalkableFloorZ(TestWalkableZ);
+		TestWalkableSlopeCosine = SlopeOverride.ModifyWalkableFloorZ(TestWalkableSlopeCosine);
 	}
 
 	// Can't walk on this surface if it is too steep.
-	if (Hit.ImpactNormal.Z < TestWalkableZ)
+	if (Hit.ImpactNormal.Dot(UpDirection) < TestWalkableSlopeCosine)
 	{
 		return false;
 	}
@@ -183,9 +186,9 @@ bool UFloorQueryUtils::IsHitSurfaceWalkable(const FHitResult& Hit, float MaxWalk
 	return true;
 }
 
-bool UFloorQueryUtils::IsWithinEdgeTolerance(const FVector& CapsuleLocation, const FVector& TestImpactPoint, float CapsuleRadius)
+bool UFloorQueryUtils::IsWithinEdgeTolerance(const FVector& CapsuleLocation, const FVector& TestImpactPoint, float CapsuleRadius, const FVector& UpDirection)
 {
-	const float DistFromCenterSq = (TestImpactPoint - CapsuleLocation).SizeSquared2D();
+	const float DistFromCenterSq = FVector::VectorPlaneProject((TestImpactPoint - CapsuleLocation), -UpDirection).SizeSquared();
 	const float ReducedRadiusSq = FMath::Square(FMath::Max(UE::FloorQueryUtility::SWEEP_EDGE_REJECT_DISTANCE + UE_KINDA_SMALL_NUMBER, CapsuleRadius - UE::FloorQueryUtility::SWEEP_EDGE_REJECT_DISTANCE));
 	return DistFromCenterSq < ReducedRadiusSq;
 }
diff --git a/Engine/Plugins/Experimental/Mover/Source/Mover/Private/MoveLibrary/GroundMovementUtils.cpp b/Engine/Plugins/Experimental/Mover/Source/Mover/Private/MoveLibrary/GroundMovementUtils.cpp
--- a/Engine/Plugins/Experimental/Mover/Source/Mover/Private/MoveLibrary/GroundMovementUtils.cpp
+++ b/Engine/Plugins/Experimental/Mover/Source/Mover/Private/MoveLibrary/GroundMovementUtils.cpp
@@ -16,7 +16,7 @@ FProposedMove UGroundMovementUtils::ComputeControlledGroundMove(const FGroundMov
 
 	const FVector MoveDirIntent = UMovementUtils::ComputeDirectionIntent(InParams.MoveInput, InParams.MoveInputType);
 
-	const FPlane MovementPlane(FVector::ZeroVector, FVector::UpVector);
+	const FPlane MovementPlane(FVector::ZeroVector, InParams.UpDirection);
 	FVector MoveDirIntentInMovementPlane = UMovementUtils::ConstrainToPlane(MoveDirIntent, MovementPlane, true);
 
 	const FPlane GroundSurfacePlane(FVector::ZeroVector, InParams.GroundNormal);
@@ -39,7 +39,7 @@ FProposedMove UGroundMovementUtils::ComputeControlledGroundMove(const FGroundMov
 	OutMove.LinearVelocity = UMovementUtils::ConstrainToPlane(Velocity, GroundSurfacePlane, true);
 
 	// Linearly rotate in place
-	OutMove.AngularVelocity = UMovementUtils::ComputeAngularVelocity(InParams.PriorOrientation, InParams.OrientationIntent, InParams.DeltaSeconds, InParams.TurningRate);
+	OutMove.AngularVelocity = UMovementUtils::ComputeAngularVelocity(InParams.PriorOrientation, InParams.OrientationIntent, InParams.WorldToGravityQuat, InParams.DeltaSeconds, InParams.TurningRate);
 
 	return OutMove;
 }
@@ -60,6 +60,7 @@ bool UGroundMovementUtils::TryMoveToStepUp(const FMovingComponentSet& MovingComp
 
 	TArray<FMovementSubstep> QueuedSubsteps;	// keeping track of substeps before committing, because some moves can be backed out
 
+	FVector UpDirection = MovingComps.MoverComponent->GetUpDirection();
 
 	const FVector OldLocation = CapsuleComponent->GetComponentLocation();
 	FVector LastComponentLocation = OldLocation;
@@ -69,8 +70,9 @@ bool UGroundMovementUtils::TryMoveToStepUp(const FMovingComponentSet& MovingComp
 	CapsuleComponent->GetScaledCapsuleSize(PawnRadius, PawnHalfHeight);
 
 	// Don't bother stepping up if top of capsule is hitting something.
-	const float InitialImpactZ = MoveHitResult.ImpactPoint.Z;
-	if (InitialImpactZ > OldLocation.Z + (PawnHalfHeight - PawnRadius))
+	const float InitialImpactDot = MoveHitResult.ImpactPoint.Dot(UpDirection);
+	const float OldLocationDot = OldLocation.Dot(UpDirection);
+	if (InitialImpactDot > OldLocationDot + (PawnHalfHeight - PawnRadius))
 	{
 		UE_LOG(LogMover, VeryVerbose, TEXT("Not stepping up due to top of capsule hitting something"));
 		return false;
@@ -88,33 +90,33 @@ bool UGroundMovementUtils::TryMoveToStepUp(const FMovingComponentSet& MovingComp
 
 	float StepTravelUpHeight = MaxStepHeight;
 	float StepTravelDownHeight = StepTravelUpHeight;
-	const float StepSideZ = -1.f * FVector::DotProduct(MoveHitResult.ImpactNormal, GravDir);
-	float PawnInitialFloorBaseZ = OldLocation.Z - PawnHalfHeight;
-	float PawnFloorPointZ = PawnInitialFloorBaseZ;
+	const float StepSideDot = -1.f * FVector::DotProduct(MoveHitResult.ImpactNormal, GravDir);
+	float PawnInitialFloorBaseDot = OldLocationDot - PawnHalfHeight;
+	float PawnFloorPointDot = PawnInitialFloorBaseDot;
 
 	//if (IsMovingOnGround() && CurrentFloor.IsWalkableFloor())
 	if (CurrentFloor.IsWalkableFloor())
 	{
 		// Since we float a variable amount off the floor, we need to enforce max step height off the actual point of impact with the floor.
 		const float FloorDist = FMath::Max(0.f, CurrentFloor.GetDistanceToFloor());
-		PawnInitialFloorBaseZ -= FloorDist;
+		PawnInitialFloorBaseDot -= FloorDist;
 		StepTravelUpHeight = FMath::Max(StepTravelUpHeight - FloorDist, 0.f);
 		StepTravelDownHeight = (MaxStepHeight + UE::FloorQueryUtility::MAX_FLOOR_DIST * 2.f);
 
-		const bool bHitVerticalFace = !UFloorQueryUtils::IsWithinEdgeTolerance(MoveHitResult.Location, MoveHitResult.ImpactPoint, PawnRadius);
+		const bool bHitVerticalFace = !UFloorQueryUtils::IsWithinEdgeTolerance(MoveHitResult.Location, MoveHitResult.ImpactPoint, PawnRadius, UpDirection);
 		if (!CurrentFloor.bLineTrace && !bHitVerticalFace)
 		{
-			PawnFloorPointZ = CurrentFloor.HitResult.ImpactPoint.Z;
+			PawnFloorPointDot = CurrentFloor.HitResult.ImpactPoint.Dot(UpDirection);
 		}
 		else
 		{
 			// Base floor point is the base of the capsule moved down by how far we are hovering over the surface we are hitting.
-			PawnFloorPointZ -= CurrentFloor.FloorDist;
+			PawnFloorPointDot -= CurrentFloor.FloorDist;
 		}
 	}
 
 	// Don't step up if the impact is below us, accounting for distance from floor.
-	if (InitialImpactZ <= PawnInitialFloorBaseZ)
+	if (InitialImpactDot <= PawnInitialFloorBaseDot)
 	{
 		UE_LOG(LogMover, VeryVerbose, TEXT("Not stepping up because the impact is below us"));
 		return false;
@@ -250,16 +252,16 @@ bool UGroundMovementUtils::TryMoveToStepUp(const FMovingComponentSet& MovingComp
 	if (StepFwdHit.IsValidBlockingHit())
 	{
 		// See if this step sequence would have allowed us to travel higher than our max step height allows.
-		const float DeltaZ = StepFwdHit.ImpactPoint.Z - PawnFloorPointZ;
-		if (DeltaZ > MaxStepHeight)
+		const float DeltaDot = (StepFwdHit.ImpactPoint.Dot(UpDirection)) - PawnFloorPointDot;
+		if (DeltaDot > MaxStepHeight)
 		{
-			UE_LOG(LogMover, VeryVerbose, TEXT("Reject step-down attempt during step-up/step-fwd, because it made us travel too high (too high Height %.3f) up from floor base %f to %f"), DeltaZ, PawnInitialFloorBaseZ, StepFwdHit.ImpactPoint.Z);
+			UE_LOG(LogMover, VeryVerbose, TEXT("Reject step-down attempt during step-up/step-fwd, because it made us travel too high (too high Height %.3f) up from floor base %f to %f"), DeltaDot, PawnInitialFloorBaseDot, StepFwdHit.ImpactPoint.Z);
 			ScopedStepUpMovement.RevertMove();
 			return false;
 		}
 
 		// Reject unwalkable surface normals here.
-		if (!UFloorQueryUtils::IsHitSurfaceWalkable(StepFwdHit, MaxWalkSlopeCosine))
+		if (!UFloorQueryUtils::IsHitSurfaceWalkable(StepFwdHit, UpDirection, MaxWalkSlopeCosine))
 		{
 			// Reject if normal opposes movement direction
 			const bool bNormalTowardsMe = (MoveDelta | StepFwdHit.ImpactNormal) < 0.f;
@@ -272,7 +274,7 @@ bool UGroundMovementUtils::TryMoveToStepUp(const FMovingComponentSet& MovingComp
 
 			// Also reject if we would end up being higher than our starting location by stepping down.
 			// It's fine to step down onto an unwalkable normal below us, we will just slide off. Rejecting those moves would prevent us from being able to walk off the edge.
-			if (StepFwdHit.Location.Z > OldLocation.Z)
+			if (StepFwdHit.Location.Dot(UpDirection) > OldLocationDot)
 			{
 				UE_LOG(LogMover, VeryVerbose, TEXT("Reject step-down attempt during step-up/step-fwd, due to unwalkable normal %s above old position)"), *StepFwdHit.ImpactNormal.ToString());
 				ScopedStepUpMovement.RevertMove();
@@ -281,7 +283,7 @@ bool UGroundMovementUtils::TryMoveToStepUp(const FMovingComponentSet& MovingComp
 		}
 
 		// Reject moves where the downward sweep hit something very close to the edge of the capsule. This maintains consistency with FindFloor as well.
-		if (!UFloorQueryUtils::IsWithinEdgeTolerance(StepFwdHit.Location, StepFwdHit.ImpactPoint, PawnRadius))
+		if (!UFloorQueryUtils::IsWithinEdgeTolerance(StepFwdHit.Location, StepFwdHit.ImpactPoint, PawnRadius, UpDirection))
 		{
 			UE_LOG(LogMover, VeryVerbose, TEXT("Reject step-down attempt during step-up/step-fwd, due to being outside edge tolerance"));
 			ScopedStepUpMovement.RevertMove();
@@ -289,7 +291,7 @@ bool UGroundMovementUtils::TryMoveToStepUp(const FMovingComponentSet& MovingComp
 		}
 
 		// Don't step up onto invalid surfaces if traveling higher.
-		if (DeltaZ > 0.f && !CanStepUpOnHitSurface(StepFwdHit))
+		if (DeltaDot > 0.f && !CanStepUpOnHitSurface(StepFwdHit))
 		{
 			UE_LOG(LogMover, VeryVerbose, TEXT("Reject step-down attempt during step-up/step-fwd, due to being up onto surface with !CanStepUpOnHitSurface")); 
 			ScopedStepUpMovement.RevertMove();
@@ -300,18 +302,16 @@ bool UGroundMovementUtils::TryMoveToStepUp(const FMovingComponentSet& MovingComp
 		if (OutFloorTestResult != NULL)
 		{
 
-			UFloorQueryUtils::FindFloor(CapsuleComponent, CapsuleComponent,
-				FloorSweepDistance, MaxWalkSlopeCosine,
-				CapsuleComponent->GetComponentLocation(), StepDownResult.FloorTestResult);
+			UFloorQueryUtils::FindFloor(MovingComps, FloorSweepDistance, MaxWalkSlopeCosine, CapsuleComponent->GetComponentLocation(), StepDownResult.FloorTestResult);
 
 			// Reject unwalkable normals if we end up higher than our initial height.
 			// It's fine to walk down onto an unwalkable surface, don't reject those moves.
-			if (StepFwdHit.Location.Z > OldLocation.Z)
+			if (StepFwdHit.Location.Dot(UpDirection) > OldLocationDot)
 			{
 				// We should reject the floor result if we are trying to step up an actual step where we are not able to perch (this is rare).
 				// In those cases we should instead abort the step up and try to slide along the stair.
-				const float MAX_STEP_SIDE_Z = 0.08f; // TODO: Move magic numbers elsewhere
-				if (!StepDownResult.FloorTestResult.bBlockingHit && StepSideZ < MAX_STEP_SIDE_Z)
+				const float MAX_STEP_SIDE_DOT = 0.08f; // TODO: Move magic numbers elsewhere
+				if (!StepDownResult.FloorTestResult.bBlockingHit && StepSideDot < MAX_STEP_SIDE_DOT)
 				{
 					UE_LOG(LogMover, VeryVerbose, TEXT("Reject step-down attempt during step-up/step-fwd, due to it being an unperchable step")); 
 					ScopedStepUpMovement.RevertMove();
@@ -354,6 +354,8 @@ bool UGroundMovementUtils::TryMoveToAdjustHeightAboveFloor(const FMovingComponen
 		return false;
 	}
 
+	const FVector UpDirection = MovingComps.MoverComponent->GetUpDirection();
+	
 	float OldFloorDist = CurrentFloor.FloorDist;
 	if (CurrentFloor.bLineTrace)
 	{
@@ -372,14 +374,14 @@ bool UGroundMovementUtils::TryMoveToAdjustHeightAboveFloor(const FMovingComponen
 	if (OldFloorDist < UE::FloorQueryUtility::MIN_FLOOR_DIST || OldFloorDist > UE::FloorQueryUtility::MAX_FLOOR_DIST)
 	{
 		FHitResult AdjustHit(1.f);
-		const float InitialZ = UpdatedComponent->GetComponentLocation().Z;
+		const float InitialUp = UpdatedComponent->GetComponentLocation().Dot(UpDirection);
 		const float AvgFloorDist = (UE::FloorQueryUtility::MIN_FLOOR_DIST + UE::FloorQueryUtility::MAX_FLOOR_DIST) * 0.5f;
 		const float MoveDist = AvgFloorDist - OldFloorDist;
 
 		MoveRecord.LockRelevancy(false);
 
 		UMovementUtils::TrySafeMoveUpdatedComponent(MovingComps,
-			FVector(0.f, 0.f, MoveDist), UpdatedComponent->GetComponentQuat(), 
+			UpDirection * MoveDist, UpdatedComponent->GetComponentQuat(), 
 			true, AdjustHit, ETeleportType::None, MoveRecord);
 
 		MoveRecord.UnlockRelevancy();
@@ -390,15 +392,15 @@ bool UGroundMovementUtils::TryMoveToAdjustHeightAboveFloor(const FMovingComponen
 		}
 		else if (MoveDist > 0.f)
 		{
-			const float CurrentZ = UpdatedComponent->GetComponentLocation().Z;
-			CurrentFloor.FloorDist += CurrentZ - InitialZ;
+			const float CurrentUp = UpdatedComponent->GetComponentLocation().Dot(UpDirection);
+			CurrentFloor.FloorDist += CurrentUp - InitialUp;
 		}
 		else
 		{
 			checkSlow(MoveDist < 0.f);
-			const float CurrentZ = UpdatedComponent->GetComponentLocation().Z;
-			CurrentFloor.FloorDist = CurrentZ - AdjustHit.Location.Z;
-			if (UFloorQueryUtils::IsHitSurfaceWalkable(AdjustHit, MaxWalkSlopeCosine))
+			const float CurrentUp = UpdatedComponent->GetComponentLocation().Dot(UpDirection);
+			CurrentFloor.FloorDist = CurrentUp - AdjustHit.Location.Dot(UpDirection);
+			if (UFloorQueryUtils::IsHitSurfaceWalkable(AdjustHit, UpDirection, MaxWalkSlopeCosine))
 			{
 				CurrentFloor.SetFromSweep(AdjustHit, CurrentFloor.FloorDist, true);
 			}
@@ -417,20 +419,23 @@ float UGroundMovementUtils::TryWalkToSlideAlongSurface(const FMovingComponentSet
 		return 0.f;
 	}
 	
-	FVector SafeWalkNormal(Normal);
 	UMoverComponent* MoverComponent = MovingComps.MoverComponent.Get();
+	const FVector UpDirection = MoverComponent->GetUpDirection();
+	FVector SafeWalkNormal(Normal);
+	const FVector::FReal NormalDot = Normal.Dot(UpDirection);
 
 	// We don't want to be pushed up an unwalkable surface.
-	if (SafeWalkNormal.Z > 0.f && !UFloorQueryUtils::IsHitSurfaceWalkable(Hit, MaxWalkSlopeCosine))
+	if (NormalDot > 0.f && !UFloorQueryUtils::IsHitSurfaceWalkable(Hit, UpDirection, MaxWalkSlopeCosine))
 	{
-		SafeWalkNormal = SafeWalkNormal.GetSafeNormal2D();
+		SafeWalkNormal = FVector::VectorPlaneProject(SafeWalkNormal, -UpDirection).GetSafeNormal();
 	}
 
 	float PctOfTimeUsed = 0.f;
 	const FVector OldSafeHitNormal = SafeWalkNormal;
 
 	FVector SlideDelta = UMovementUtils::ComputeSlideDelta(MovingComps, Delta, PctOfDeltaToMove, SafeWalkNormal, Hit);
-
+	FVector OriginalSlideDelta = SlideDelta;
+	
 	if (SlideDelta.Dot(Delta) > 0.f)
 	{
 		UMovementUtils::TrySafeMoveUpdatedComponent(MovingComps, SlideDelta, Rotation, true, Hit, ETeleportType::None, MoveRecord);
@@ -448,23 +453,27 @@ float UGroundMovementUtils::TryWalkToSlideAlongSurface(const FMovingComponentSet
 
 			// Compute new slide normal when hitting multiple surfaces.
 			SlideDelta = UMovementUtils::ComputeTwoWallAdjustedDelta(MovingComps, SlideDelta, Hit, OldSafeHitNormal);
-			if (SlideDelta.Z > 0.f && UFloorQueryUtils::IsHitSurfaceWalkable(Hit, MaxWalkSlopeCosine) && Hit.Normal.Z > UE_KINDA_SMALL_NUMBER)
+			const float SlideDeltaDot = SlideDelta.Dot(UpDirection);
+			const float HitNormalDot = Hit.Normal.Dot(UpDirection);
+			if (SlideDeltaDot > 0.f && UFloorQueryUtils::IsHitSurfaceWalkable(Hit, UpDirection, MaxWalkSlopeCosine) && HitNormalDot > UE_KINDA_SMALL_NUMBER)
 			{
 				// Maintain horizontal velocity
 				const float Time = (1.f - Hit.Time);
 				const FVector ScaledDelta = SlideDelta.GetSafeNormal() * SlideDelta.Size();
-				SlideDelta = FVector(SlideDelta.X, SlideDelta.Y, ScaledDelta.Z / Hit.Normal.Z) * Time;
+				const FVector::FReal NewDeltaDot = (ScaledDelta.Dot(UpDirection) / HitNormalDot);
+				FVector SlideProjectedToGrav = FVector::VectorPlaneProject(OriginalSlideDelta, -UpDirection);
+				SlideDelta = (SlideProjectedToGrav + NewDeltaDot * UpDirection) * Time;
 				// Should never exceed MaxStepHeight in vertical component, so rescale if necessary.
-				// This should be rare (Hit.Normal.Z above would have been very small) but we'd rather lose horizontal velocity than go too high.
-				if (SlideDelta.Z > MaxStepHeight)
+				// This should be rare (HitNormalDot above would have been very small) but we'd rather lose horizontal velocity than go too high.
+				if (SlideDeltaDot > MaxStepHeight)
 				{
-					const float Rescale = MaxStepHeight / SlideDelta.Z;
+					const float Rescale = MaxStepHeight / SlideDeltaDot;
 					SlideDelta *= Rescale;
 				}
 			}
 			else
 			{
-				SlideDelta.Z = 0.f;
+				SlideDelta = FVector::VectorPlaneProject(SlideDelta, -UpDirection);
 			}
 			
 			// Only proceed if the new direction is of significant length and not in reverse of original attempted move.
@@ -489,15 +498,15 @@ float UGroundMovementUtils::TryWalkToSlideAlongSurface(const FMovingComponentSet
 	return 0.f;
 }
 
-FVector UGroundMovementUtils::ComputeDeflectedMoveOntoRamp(const FVector& OrigMoveDelta, const FHitResult& RampHitResult, float MaxWalkSlopeCosine, const bool bHitFromLineTrace)
+FVector UGroundMovementUtils::ComputeDeflectedMoveOntoRamp(const FVector& OrigMoveDelta, const FVector& UpDirection, const FHitResult& RampHitResult, float MaxWalkSlopeCosine, const bool bHitFromLineTrace)
 {
 	const FVector FloorNormal = RampHitResult.ImpactNormal;
-	const FVector ContactNormal = RampHitResult.Normal;
+	const FVector::FReal FloorNormalDot = FloorNormal.Dot(UpDirection);
+	const FVector::FReal ContactNormalDot = RampHitResult.Normal.Dot(UpDirection);
 
-	// JAH TODO: Change these Z tests to take movement plane into account, rather than assuming we're using the Z=0 plane
-	if (FloorNormal.Z < (1.f - UE_KINDA_SMALL_NUMBER) && FloorNormal.Z > UE_KINDA_SMALL_NUMBER && 
-		ContactNormal.Z > UE_KINDA_SMALL_NUMBER && 
-		UFloorQueryUtils::IsHitSurfaceWalkable(RampHitResult, MaxWalkSlopeCosine) && !bHitFromLineTrace)
+	if (FloorNormalDot < (1.f - UE_KINDA_SMALL_NUMBER) && FloorNormalDot > UE_KINDA_SMALL_NUMBER && 
+		ContactNormalDot > UE_KINDA_SMALL_NUMBER && 
+		UFloorQueryUtils::IsHitSurfaceWalkable(RampHitResult, UpDirection, MaxWalkSlopeCosine) && !bHitFromLineTrace)
 	{
 		// Compute a vector that moves parallel to the surface, by projecting the horizontal movement direction onto the ramp.
 		const FPlane RampSurfacePlane(FVector::ZeroVector, FloorNormal);
diff --git a/Engine/Plugins/Experimental/Mover/Source/Mover/Private/MoveLibrary/MovementUtils.cpp b/Engine/Plugins/Experimental/Mover/Source/Mover/Private/MoveLibrary/MovementUtils.cpp
--- a/Engine/Plugins/Experimental/Mover/Source/Mover/Private/MoveLibrary/MovementUtils.cpp
+++ b/Engine/Plugins/Experimental/Mover/Source/Mover/Private/MoveLibrary/MovementUtils.cpp
@@ -154,6 +154,24 @@ FVector UMovementUtils::ConstrainToPlane(const FVector& Vector, const FPlane& Mo
 	return ConstrainedResult;
 }
 
+FRotator UMovementUtils::ApplyGravityToOrientationIntent(const FRotator& IntendedOrientation, const FQuat& WorldToGravity, bool bStayVertical)
+{
+	FRotator NewOrientationIntent = IntendedOrientation;
+
+	FRotator GravityRelativeDesiredRotation = (WorldToGravity.Inverse() * NewOrientationIntent.Quaternion()).Rotator();
+	
+	if (bStayVertical)
+	{
+		GravityRelativeDesiredRotation.Pitch = 0.f;
+		GravityRelativeDesiredRotation.Yaw = FRotator::NormalizeAxis(GravityRelativeDesiredRotation.Yaw);
+		GravityRelativeDesiredRotation.Roll = 0.f;
+	}
+	
+	NewOrientationIntent = (WorldToGravity * GravityRelativeDesiredRotation.Quaternion()).Rotator();
+
+	return NewOrientationIntent;
+}
+
 FVector UMovementUtils::ComputeSlideDelta(const FMovingComponentSet& MovingComps, const FVector& Delta, const float PctOfDeltaToMove, const FVector& Normal, const FHitResult& Hit)
 {
 	FVector ConstrainedNormal = Normal;
@@ -447,41 +465,42 @@ FVector UMovementUtils::ComputeVelocityFromPositions(const FVector& FromPos, con
 	return FVector::ZeroVector;
 }
 
-FRotator UMovementUtils::ComputeAngularVelocity(const FRotator& FromOrientation, const FRotator& ToOrientation, float DeltaSeconds, float TurningRateLimit)
+FRotator UMovementUtils::ComputeAngularVelocity(const FRotator& FromOrientation, const FRotator& ToOrientation, const FQuat& WorldToGravity, float DeltaSeconds, float TurningRateLimit)
 {
 	FRotator AngularVelocityDpS(FRotator::ZeroRotator);
 
-	if (DeltaSeconds > 0.f)
+	const float AngleTolerance = 1e-3f;
+	if (!FromOrientation.Equals(ToOrientation, AngleTolerance))
 	{
-		FRotator AngularDelta = (ToOrientation-FromOrientation);
-		FRotator Winding, Remainder;
-
-		AngularDelta.GetWindingAndRemainder(Winding, Remainder);	// to find the fastest turn, just get the (-180,180] remainder
-
-		// Ensure consistent turn direction when rotating by 180 degrees
-		// This can be important in avoiding client/server desyncs for small
-		// angular differences
-		if (Remainder.Yaw < (-180.0 + UE_KINDA_SMALL_NUMBER))
-		{
-			Remainder.Yaw = 180.0f;
-		}
-		if (Remainder.Pitch < (-180.0 + UE_KINDA_SMALL_NUMBER))
+		FRotator GravityRelativeCurrentRotation = (WorldToGravity.Inverse() * FromOrientation.Quaternion()).Rotator();
+		FRotator GravityRelativeDesiredRotation = (WorldToGravity.Inverse() * ToOrientation.Quaternion()).Rotator();
+
+		const FRotator RotationDelta = GravityRelativeDesiredRotation - GravityRelativeCurrentRotation;
+		const float YawTurnRateLimit   = FMath::Clamp(RotationDelta.Yaw,   -TurningRateLimit, TurningRateLimit);
+		const float PitchTurnRateLimit = FMath::Clamp(RotationDelta.Pitch, -TurningRateLimit, TurningRateLimit);
+		const float RollTurnRateLimit  = FMath::Clamp(RotationDelta.Roll,  -TurningRateLimit, TurningRateLimit);
+		
+		// PITCH
+		if (!FMath::IsNearlyEqual(GravityRelativeCurrentRotation.Pitch, GravityRelativeDesiredRotation.Pitch, AngleTolerance))
 		{
-			Remainder.Pitch = 180.0f;
+			GravityRelativeDesiredRotation.Pitch = FMath::FixedTurn(GravityRelativeCurrentRotation.Pitch, GravityRelativeDesiredRotation.Pitch, TurningRateLimit * DeltaSeconds);
 		}
-		if (Remainder.Roll < (-180.0 + UE_KINDA_SMALL_NUMBER))
+
+		// YAW
+		if (!FMath::IsNearlyEqual(GravityRelativeCurrentRotation.Yaw, GravityRelativeDesiredRotation.Yaw, AngleTolerance))
 		{
-			Remainder.Roll = 180.0f;
+			GravityRelativeDesiredRotation.Yaw = FMath::FixedTurn(GravityRelativeCurrentRotation.Yaw, GravityRelativeDesiredRotation.Yaw, TurningRateLimit * DeltaSeconds);
 		}
 
-		AngularVelocityDpS = Remainder * (1.f / DeltaSeconds);
-
-		if (TurningRateLimit >= 0.f)
+		// ROLL
+		if (!FMath::IsNearlyEqual(GravityRelativeCurrentRotation.Roll, GravityRelativeDesiredRotation.Roll, AngleTolerance))
 		{
-			AngularVelocityDpS.Yaw   = FMath::Clamp(AngularVelocityDpS.Yaw,   -TurningRateLimit, TurningRateLimit);
-			AngularVelocityDpS.Pitch = FMath::Clamp(AngularVelocityDpS.Pitch, -TurningRateLimit, TurningRateLimit);
-			AngularVelocityDpS.Roll  = FMath::Clamp(AngularVelocityDpS.Roll,  -TurningRateLimit, TurningRateLimit);
+			GravityRelativeDesiredRotation.Roll = FMath::FixedTurn(GravityRelativeCurrentRotation.Roll, GravityRelativeDesiredRotation.Roll, TurningRateLimit * DeltaSeconds);
 		}
+
+		const FRotator DesiredRotation = (WorldToGravity * GravityRelativeDesiredRotation.Quaternion()).Rotator();
+		const FRotator AngularVelocity = DesiredRotation - FromOrientation;
+		AngularVelocityDpS = AngularVelocity * (1.f / DeltaSeconds);
 	}
 
 	return AngularVelocityDpS;
diff --git a/Engine/Plugins/Experimental/Mover/Source/Mover/Private/MoveLibrary/WaterMovementUtils.cpp b/Engine/Plugins/Experimental/Mover/Source/Mover/Private/MoveLibrary/WaterMovementUtils.cpp
--- a/Engine/Plugins/Experimental/Mover/Source/Mover/Private/MoveLibrary/WaterMovementUtils.cpp
+++ b/Engine/Plugins/Experimental/Mover/Source/Mover/Private/MoveLibrary/WaterMovementUtils.cpp
@@ -25,7 +25,7 @@ FProposedMove UWaterMovementUtils::ComputeControlledWaterMove(const FWaterMovePa
 
 	// JAH TODO: this is where we can perform turning, based on aux settings. For now, just snap to the intended final orientation.
 	FVector IntendedFacingDir = InParams.OrientationIntent.RotateVector(FVector::ForwardVector).GetSafeNormal();
-	OutMove.AngularVelocity = UMovementUtils::ComputeAngularVelocity(InParams.PriorOrientation, IntendedFacingDir.ToOrientationRotator(), InParams.DeltaSeconds, InParams.TurningRate);
+	OutMove.AngularVelocity = UMovementUtils::ComputeAngularVelocity(InParams.PriorOrientation, IntendedFacingDir.ToOrientationRotator(), InParams.WorldToGravityQuat, InParams.DeltaSeconds, InParams.TurningRate);
 
 	return OutMove;
 }
diff --git a/Engine/Plugins/Experimental/Mover/Source/Mover/Private/MovementModeStateMachine.cpp b/Engine/Plugins/Experimental/Mover/Source/Mover/Private/MovementModeStateMachine.cpp
--- a/Engine/Plugins/Experimental/Mover/Source/Mover/Private/MovementModeStateMachine.cpp
+++ b/Engine/Plugins/Experimental/Mover/Source/Mover/Private/MovementModeStateMachine.cpp
@@ -280,8 +280,7 @@ void UMovementModeStateMachine::OnSimulationTick(USceneComponent* UpdatedCompone
 			}
 			if (CurrentLayeredMoves.ResidualClamping >= 0.0f)
 			{
-				CombinedMove.LinearVelocity = CombinedMove.LinearVelocity.GetClampedToMaxSize2D(CurrentLayeredMoves.ResidualClamping);
-				CombinedMove.LinearVelocity.Z = FMath::Min<FVector::FReal>(CombinedMove.LinearVelocity.Z, CurrentLayeredMoves.ResidualClamping);
+				CombinedMove.LinearVelocity = CombinedMove.LinearVelocity.GetClampedToMaxSize(CurrentLayeredMoves.ResidualClamping);
 			}
 			CurrentLayeredMoves.ResetResidualVelocity();
 
diff --git a/Engine/Plugins/Experimental/Mover/Source/Mover/Private/MoverComponent.cpp b/Engine/Plugins/Experimental/Mover/Source/Mover/Private/MoverComponent.cpp
--- a/Engine/Plugins/Experimental/Mover/Source/Mover/Private/MoverComponent.cpp
+++ b/Engine/Plugins/Experimental/Mover/Source/Mover/Private/MoverComponent.cpp
@@ -167,6 +167,9 @@ void UMoverComponent::BeginPlay()
 	FindDefaultUpdatedComponent();
 	ensureMsgf(UpdatedComponent != nullptr, TEXT("No root component found on %s. Simulation initialization will most likely fail."), *GetPathNameSafe(GetOwner()));
 
+	WorldToGravityTransform = FQuat::FindBetweenNormals(FVector::UpVector, GetUpDirection());
+	GravityToWorldTransform = WorldToGravityTransform.Inverse();
+	
 	if (const AActor* MyActor = GetOwner())
 	{
 		// If no primary visual component is already set, fall back to searching for any kind of mesh,
@@ -483,7 +486,7 @@ void UMoverComponent::SimulationTick(const FMoverTimeStep& InTimeStep, const FMo
 		SimOutput.SyncState.MovementMode = MovementModeAfterTick;
 
 		if (UpdatedComponent->GetComponentLocation().Equals(OutputSyncState->GetLocation_WorldSpace()) == false ||
-			UpdatedComponent->GetComponentQuat().Rotator().Equals(OutputSyncState->GetOrientation_WorldSpace(), ROTATOR_TOLERANCE) == false)
+			UpdatedComponent->GetComponentQuat().Equals(OutputSyncState->GetOrientation_WorldSpace().Quaternion(), UE_KINDA_SMALL_NUMBER) == false)
 		{
 			UE_LOG(LogMover, Warning, TEXT("Detected pos/rot difference between Mover actor (%s) sync state and scene component after sim ticking. This indicates a movement mode may not be authoring the final state correctly."), *GetNameSafe(UpdatedComponent->GetOwner()));
 		}
@@ -1694,6 +1697,9 @@ void UMoverComponent::SetGravityOverride(bool bOverrideGravity, FVector NewGravi
 {
 	bHasGravityOverride = bOverrideGravity;
 	GravityAccelOverride = NewGravityAcceleration;
+	
+	WorldToGravityTransform = FQuat::FindBetweenNormals(FVector::UpVector, -GravityAccelOverride.GetSafeNormal());
+	GravityToWorldTransform = WorldToGravityTransform.Inverse();
 }
 
 
diff --git a/Engine/Plugins/Experimental/Mover/Source/Mover/Private/PhysicsMover/PhysicsMovementUtils.cpp b/Engine/Plugins/Experimental/Mover/Source/Mover/Private/PhysicsMover/PhysicsMovementUtils.cpp
--- a/Engine/Plugins/Experimental/Mover/Source/Mover/Private/PhysicsMover/PhysicsMovementUtils.cpp
+++ b/Engine/Plugins/Experimental/Mover/Source/Mover/Private/PhysicsMover/PhysicsMovementUtils.cpp
@@ -76,7 +76,7 @@ void UPhysicsMovementUtils::FloorSweep_Internal(
 
 		if (OutHit.bBlockingHit)
 		{
-			bool bWalkable = UFloorQueryUtils::IsHitSurfaceWalkable(OutHit, MaxWalkSlopeCosine);
+			bool bWalkable = UFloorQueryUtils::IsHitSurfaceWalkable(OutHit, UpDir, MaxWalkSlopeCosine);
 
 #if PHYSICSDRIVENMOTION_DEBUG_DRAW
 			if (GPhysicsDrivenMotionDebugParams.DebugDrawGroundQueries)
diff --git a/Engine/Plugins/Experimental/Mover/Source/Mover/Public/DefaultMovementSet/Settings/CommonLegacyMovementSettings.h b/Engine/Plugins/Experimental/Mover/Source/Mover/Public/DefaultMovementSet/Settings/CommonLegacyMovementSettings.h
--- a/Engine/Plugins/Experimental/Mover/Source/Mover/Public/DefaultMovementSet/Settings/CommonLegacyMovementSettings.h
+++ b/Engine/Plugins/Experimental/Mover/Source/Mover/Public/DefaultMovementSet/Settings/CommonLegacyMovementSettings.h
@@ -17,6 +17,12 @@ class MOVER_API UCommonLegacyMovementSettings : public UObject, public IMovement
 	virtual FString GetDisplayName() const override { return GetName(); }
 
 public:
+	/**
+	 * If true, the actor will remain upright with gravity despite any rotation applied to the actor
+	 */
+	UPROPERTY(Category="General", EditAnywhere, BlueprintReadWrite)
+	bool bShouldRemainVertical = true;
+	
 	// What movement mode to use when on the ground.
 	UPROPERTY(Category="General", EditAnywhere, BlueprintReadWrite)
 	FName GroundMovementModeName = DefaultModeNames::Walking;
diff --git a/Engine/Plugins/Experimental/Mover/Source/Mover/Public/MoveLibrary/AirMovementUtils.h b/Engine/Plugins/Experimental/Mover/Source/Mover/Public/MoveLibrary/AirMovementUtils.h
--- a/Engine/Plugins/Experimental/Mover/Source/Mover/Public/MoveLibrary/AirMovementUtils.h
+++ b/Engine/Plugins/Experimental/Mover/Source/Mover/Public/MoveLibrary/AirMovementUtils.h
@@ -47,6 +47,9 @@ struct MOVER_API FFreeMoveParams
 
 	UPROPERTY(BlueprintReadOnly, EditAnywhere, Category = Mover)
 	float DeltaSeconds = 0.f;
+
+	UPROPERTY(BlueprintReadOnly, EditAnywhere, Category = Mover)
+	FQuat WorldToGravityQuat;
 };
 
 /**
diff --git a/Engine/Plugins/Experimental/Mover/Source/Mover/Public/MoveLibrary/FloorQueryUtils.h b/Engine/Plugins/Experimental/Mover/Source/Mover/Public/MoveLibrary/FloorQueryUtils.h
--- a/Engine/Plugins/Experimental/Mover/Source/Mover/Public/MoveLibrary/FloorQueryUtils.h
+++ b/Engine/Plugins/Experimental/Mover/Source/Mover/Public/MoveLibrary/FloorQueryUtils.h
@@ -3,6 +3,7 @@
 #pragma once
 
 #include "CoreMinimal.h"
+#include "MovementUtilsTypes.h"
 #include "Kismet/BlueprintFunctionLibrary.h"
 #include "Components/PrimitiveComponent.h"
 #include "FloorQueryUtils.generated.h"
@@ -108,18 +109,18 @@ class MOVER_API UFloorQueryUtils : public UBlueprintFunctionLibrary
 	GENERATED_BODY()
 
 public:
-	static void FindFloor(const USceneComponent* UpdatedComponent, const UPrimitiveComponent* UpdatedPrimitive, float FloorSweepDistance, float MaxWalkSlopeCosine, const FVector& Location, FFloorCheckResult& OutFloorResult);
+	static void FindFloor(const FMovingComponentSet& MovingComps, float FloorSweepDistance, float MaxWalkSlopeCosine, const FVector& Location, FFloorCheckResult& OutFloorResult);
 
-	static void ComputeFloorDist(const USceneComponent* UpdatedComponent, const UPrimitiveComponent* UpdatedPrimitive, float LineTraceDistance, float FloorSweepDistance, float MaxWalkSlopeCosine, const FVector& Location, FFloorCheckResult& OutFloorResult);
+	static void ComputeFloorDist(const FMovingComponentSet& MovingComps, float LineTraceDistance, float FloorSweepDistance, float MaxWalkSlopeCosine, const FVector& Location, FFloorCheckResult& OutFloorResult);
 
-	static bool FloorSweepTest(const UPrimitiveComponent* UpdatedPrimitive, FHitResult& OutHit, const FVector& Start, const FVector& End, ECollisionChannel TraceChannel, const struct FCollisionShape& CollisionShape, const struct FCollisionQueryParams& Params, const struct FCollisionResponseParams& ResponseParam);
+	static bool FloorSweepTest(const FMovingComponentSet& MovingComps, FHitResult& OutHit, const FVector& Start, const FVector& End, ECollisionChannel TraceChannel, const struct FCollisionShape& CollisionShape, const struct FCollisionQueryParams& Params, const struct FCollisionResponseParams& ResponseParam);
 
 	UFUNCTION(BlueprintCallable, Category=Mover)
-	static bool IsHitSurfaceWalkable(const FHitResult& Hit, float MaxWalkSlopeCosine);
+	static bool IsHitSurfaceWalkable(const FHitResult& Hit, FVector UpDirection, float MaxWalkSlopeCosine);
 
 	/**
 	 * Return true if the 2D distance to the impact point is inside the edge tolerance (CapsuleRadius minus a small rejection threshold).
 	 * Useful for rejecting adjacent hits when finding a floor or landing spot.
 	 */
-	static bool IsWithinEdgeTolerance(const FVector& CapsuleLocation, const FVector& TestImpactPoint, float CapsuleRadius);
+	static bool IsWithinEdgeTolerance(const FVector& CapsuleLocation, const FVector& TestImpactPoint, float CapsuleRadius, const FVector& UpDirection);
 };
diff --git a/Engine/Plugins/Experimental/Mover/Source/Mover/Public/MoveLibrary/GroundMovementUtils.h b/Engine/Plugins/Experimental/Mover/Source/Mover/Public/MoveLibrary/GroundMovementUtils.h
--- a/Engine/Plugins/Experimental/Mover/Source/Mover/Public/MoveLibrary/GroundMovementUtils.h
+++ b/Engine/Plugins/Experimental/Mover/Source/Mover/Public/MoveLibrary/GroundMovementUtils.h
@@ -58,6 +58,12 @@ struct MOVER_API FGroundMoveParams
 
 	UPROPERTY(BlueprintReadOnly, EditAnywhere, Category = Mover)
 	float DeltaSeconds = 0.f;
+
+	UPROPERTY(BlueprintReadOnly, EditAnywhere, Category = Mover)
+	FVector UpDirection;
+	
+	UPROPERTY(BlueprintReadOnly, EditAnywhere, Category = Mover)
+	FQuat WorldToGravityQuat;
 };
 
 /**
@@ -88,7 +94,7 @@ public:
 
 	/** Used to change a movement to be along a ramp's surface, typically to prevent slow movement when running up/down a ramp */
 	UFUNCTION(BlueprintCallable, Category = Mover)
-	static FVector ComputeDeflectedMoveOntoRamp(const FVector& OrigMoveDelta, const FHitResult& RampHitResult, float MaxWalkSlopeCosine, const bool bHitFromLineTrace);
+	static FVector ComputeDeflectedMoveOntoRamp(const FVector& OrigMoveDelta, const FVector& UpDirection, const FHitResult& RampHitResult, float MaxWalkSlopeCosine, const bool bHitFromLineTrace);
 
 	UFUNCTION(BlueprintCallable, Category = Mover)
 	static bool CanStepUpOnHitSurface(const FHitResult& Hit);
diff --git a/Engine/Plugins/Experimental/Mover/Source/Mover/Public/MoveLibrary/MovementUtils.h b/Engine/Plugins/Experimental/Mover/Source/Mover/Public/MoveLibrary/MovementUtils.h
--- a/Engine/Plugins/Experimental/Mover/Source/Mover/Public/MoveLibrary/MovementUtils.h
+++ b/Engine/Plugins/Experimental/Mover/Source/Mover/Public/MoveLibrary/MovementUtils.h
@@ -149,6 +149,10 @@ public:
 	UFUNCTION(BlueprintCallable, Category = Mover)
 	static FVector ConstrainToPlane(const FVector& Vector, const FPlane& MovementPlane, bool bMaintainMagnitude=true);
 
+	/** Converts intended orientation into orientation rotated by out current gravity */
+	UFUNCTION(BlueprintCallable, Category = Mover)
+	static FRotator ApplyGravityToOrientationIntent(const FRotator& IntendedOrientation, const FQuat& WorldToGravity, bool bStayVertical = true);
+
 	/** Project a vector onto the floor defined by the gravity direction. */
 	UFUNCTION(BlueprintCallable, Category = Mover)
 	static FVector ProjectToGravityFloor(const FVector& Vector, const FVector& UpDirection) { return FVector::VectorPlaneProject(Vector, -UpDirection); }
@@ -159,7 +163,7 @@ public:
 
 	/** Set the vertical component of the vector to the given value in the gravity-space vertical direction. */
 	static void SetGravityVerticalComponent(FVector& Vector, const FVector::FReal VerticalValue, const FVector& UpDirection) { Vector = ProjectToGravityFloor(Vector, UpDirection) - VerticalValue * -UpDirection; }
-
+	
 	// Surface sliding
 
 	/** Returns an alternative move delta to slide along a surface, based on parameters describing a blocked attempted move */
@@ -193,7 +197,7 @@ public:
 
 	/** Computes the angular velocity needed to change from one orientation to another within a time frame. Use the optional TurningRateLimit to clamp to a maximum step (negative=unlimited). */
 	UFUNCTION(BlueprintCallable, Category = Mover)
-	static FRotator ComputeAngularVelocity(const FRotator& From, const FRotator& To, float DeltaSeconds, float TurningRateLimit=-1.f);
+	static FRotator ComputeAngularVelocity(const FRotator& From, const FRotator& To, const FQuat& WorldToGravity, float DeltaSeconds, float TurningRateLimit=-1.f);
 
 	/** Computes the directional movement intent based on input vector and associated type */
 	UFUNCTION(BlueprintCallable, Category = Mover)
diff --git a/Engine/Plugins/Experimental/Mover/Source/Mover/Public/MoveLibrary/WaterMovementUtils.h b/Engine/Plugins/Experimental/Mover/Source/Mover/Public/MoveLibrary/WaterMovementUtils.h
--- a/Engine/Plugins/Experimental/Mover/Source/Mover/Public/MoveLibrary/WaterMovementUtils.h
+++ b/Engine/Plugins/Experimental/Mover/Source/Mover/Public/MoveLibrary/WaterMovementUtils.h
@@ -58,6 +58,9 @@ struct MOVER_API FWaterMoveParams
 
 	UPROPERTY(BlueprintReadOnly, EditAnywhere, Category = Mover)
 	float MoveSpeed = 0.f;
+	
+	UPROPERTY(BlueprintReadOnly, EditAnywhere, Category = Mover)
+	FQuat WorldToGravityQuat;
 };
 
 // Data about the water volume and its interaction with the pawn used in calculating swimming movement
diff --git a/Engine/Plugins/Experimental/Mover/Source/Mover/Public/MoverComponent.h b/Engine/Plugins/Experimental/Mover/Source/Mover/Public/MoverComponent.h
--- a/Engine/Plugins/Experimental/Mover/Source/Mover/Public/MoverComponent.h
+++ b/Engine/Plugins/Experimental/Mover/Source/Mover/Public/MoverComponent.h
@@ -224,8 +224,14 @@ public:
 	UFUNCTION(BlueprintCallable, BlueprintPure = false, Category = Mover)
 	FVector GetGravityAcceleration() const;
 
+	/** Returns a quaternion transforming from world to gravity space. */
+	FQuat GetWorldToGravityTransform() const { return WorldToGravityTransform; }
+
+	/** Returns a quaternion transforming from gravity to world space. */
+	FQuat GetGravityToWorldTransform() const { return GravityToWorldTransform; }
+	
 	// Get the normalized direction considered "up" in worldspace. Typically aligned with gravity, and typically determines the plane an actor tries to move along.
-	UFUNCTION(BlueprintCallable, BlueprintPure = false, Category = Mover)
+	UFUNCTION(BlueprintPure = false, Category = Mover)
 	FVector GetUpDirection() const;
 
 	// Access the planar constraint that may be limiting movement direction
@@ -554,6 +560,18 @@ private:
 	UPROPERTY(EditDefaultsOnly, Category="Mover|Gravity", meta=(ForceUnits = "cm/s^2"))
 	FVector GravityAccelOverride;
 
+	/**
+	 * A cached quaternion representing the rotation from world space to gravity relative space defined by GravityAccelOverride.
+	 */
+	UPROPERTY(Category="Mover|Gravity", VisibleAnywhere, BlueprintReadOnly, meta = (AllowPrivateAccess = true))
+	FQuat WorldToGravityTransform;
+	
+	/**
+	 * A cached quaternion representing the inverse rotation from world space to gravity relative space defined by GravityAccelOverride.
+	 */
+	UPROPERTY(Category="Mover|Gravity", VisibleAnywhere, BlueprintReadOnly, meta = (AllowPrivateAccess = true))
+	FQuat GravityToWorldTransform;
+	
 	/** Settings that can lock movement to a particular plane */
 	UPROPERTY(EditDefaultsOnly, Category = "Mover|Constraints")
 	FPlanarConstraint PlanarConstraint;
diff --git a/Engine/Plugins/Experimental/MoverExamples/Source/MoverExamples/Private/MoverExamplesCharacter.cpp b/Engine/Plugins/Experimental/MoverExamples/Source/MoverExamples/Private/MoverExamplesCharacter.cpp
--- a/Engine/Plugins/Experimental/MoverExamples/Source/MoverExamples/Private/MoverExamplesCharacter.cpp
+++ b/Engine/Plugins/Experimental/MoverExamples/Source/MoverExamples/Private/MoverExamplesCharacter.cpp
@@ -61,9 +61,15 @@ void AMoverExamplesCharacter::Tick(float DeltaTime)
 		// Simple input scaling. A real game will probably map this to an acceleration curve
 		static float LookRateYaw = 100.f;	// degs/sec
 		static float LookRatePitch = 100.f;	// degs/sec
+		static float LookRateRoll = 100.f;	// degs/sec
 
-		PC->AddYawInput(CachedLookInput.Yaw * LookRateYaw * DeltaTime);
-		PC->AddPitchInput(-CachedLookInput.Pitch * LookRatePitch * DeltaTime);
+		const FRotator GravityRotation(CharacterMotionComponent->GetWorldToGravityTransform() * CachedLookInput.Quaternion());
+		const FRotator DefaultRotation(CharacterMotionComponent->GetWorldToGravityTransform() * FRotator::ZeroRotator.Quaternion());
+		const FRotator Difference = GravityRotation - DefaultRotation;
+			
+		PC->AddYawInput(Difference.Yaw * LookRateYaw * DeltaTime);
+		PC->AddPitchInput(-Difference.Pitch * LookRatePitch * DeltaTime);
+		PC->AddRollInput(Difference.Roll * LookRateRoll * DeltaTime);
 	}
 
 	// Clear all camera-related cached input
@@ -206,19 +212,7 @@ void AMoverExamplesCharacter::OnProduceInput(float DeltaMs, FMoverInputCmdContex
 
 	if (bUsingInputIntentForMove)
 	{
-		FRotator Rotator = CharacterInputs.ControlRotation;
-		FVector FinalDirectionalIntent;
-		if (const UCharacterMoverComponent* MoverComp = GetMoverComponent())
-		{
-			if (MoverComp->IsOnGround() || MoverComp->IsFalling())
-			{
-				const FVector RotationProjectedOntoUpDirection = FVector::VectorPlaneProject(Rotator.Vector(), MoverComp->GetUpDirection()).GetSafeNormal();
-				Rotator = RotationProjectedOntoUpDirection.Rotation();
-			}
-
-			FinalDirectionalIntent = Rotator.RotateVector(CachedMoveInputIntent);
-		}
-		
+		const FVector FinalDirectionalIntent = CharacterInputs.ControlRotation.RotateVector(CachedMoveInputIntent);
 		CharacterInputs.SetMoveInput(EMoveInputType::DirectionalIntent, FinalDirectionalIntent);
 	}
 	else
@@ -262,12 +256,6 @@ void AMoverExamplesCharacter::OnProduceInput(float DeltaMs, FMoverInputCmdContex
 		CharacterInputs.OrientationIntent = LastAffirmativeMoveInput;
 	}
 	
-	if (bShouldRemainVertical)
-	{
-		// canceling out any z intent if the actor is supposed to remain vertical
-		CharacterInputs.OrientationIntent = CharacterInputs.OrientationIntent.GetSafeNormal2D();
-	}
-
 	CharacterInputs.bIsJumpPressed = bIsJumpPressed;
 	CharacterInputs.bIsJumpJustPressed = bIsJumpJustPressed;
 
diff --git a/Engine/Plugins/Experimental/MoverExamples/Source/MoverExamples/Public/MoverExamplesCharacter.h b/Engine/Plugins/Experimental/MoverExamples/Source/MoverExamples/Public/MoverExamplesCharacter.h
--- a/Engine/Plugins/Experimental/MoverExamples/Source/MoverExamples/Public/MoverExamplesCharacter.h
+++ b/Engine/Plugins/Experimental/MoverExamples/Source/MoverExamples/Public/MoverExamplesCharacter.h
@@ -92,12 +92,6 @@ public:
 	UPROPERTY(EditAnywhere, BlueprintReadWrite, Category=MoverExamples)
 	bool bOrientRotationToMovement = true;
 	
-	/**
-	 * If true, the actor will remain vertical despite any rotation applied to the actor
-	 */
-	UPROPERTY(EditAnywhere, BlueprintReadWrite, Category=MoverExamples)
-	bool bShouldRemainVertical = true;
-
 	/**
 	 * If true, the actor will continue orienting towards the last intended orientation (from input) even after movement intent input has ceased.
 	 * This makes the character finish orienting after a quick stick flick from the player.  If false, character will not turn without input.
