Mover: Fixing a potential nullptr issue in UGroundMovementUtils::TryWalkToSlideAlongSurface

diff --git a/Engine/Plugins/Experimental/Mover/Source/Mover/Private/MoveLibrary/GroundMovementUtils.cpp b/Engine/Plugins/Experimental/Mover/Source/Mover/Private/MoveLibrary/GroundMovementUtils.cpp
--- a/Engine/Plugins/Experimental/Mover/Source/Mover/Private/MoveLibrary/GroundMovementUtils.cpp
+++ b/Engine/Plugins/Experimental/Mover/Source/Mover/Private/MoveLibrary/GroundMovementUtils.cpp
@@ -419,80 +419,82 @@ float UGroundMovementUtils::TryWalkToSlideAlongSurface(const FMovingComponentSet
 		return 0.f;
 	}
 	
-	UMoverComponent* MoverComponent = MovingComps.MoverComponent.Get();
-	const FVector UpDirection = MoverComponent->GetUpDirection();
-	FVector SafeWalkNormal(Normal);
-	const FVector::FReal NormalDot = Normal.Dot(UpDirection);
-
-	// We don't want to be pushed up an unwalkable surface.
-	if (NormalDot > 0.f && !UFloorQueryUtils::IsHitSurfaceWalkable(Hit, UpDirection, MaxWalkSlopeCosine))
+	if (UMoverComponent* MoverComponent = MovingComps.MoverComponent.Get())
 	{
-		SafeWalkNormal = FVector::VectorPlaneProject(SafeWalkNormal, -UpDirection).GetSafeNormal();
-	}
+		const FVector UpDirection = MoverComponent->GetUpDirection();
+		FVector SafeWalkNormal(Normal);
+		const FVector::FReal NormalDot = Normal.Dot(UpDirection);
 
-	float PctOfTimeUsed = 0.f;
-	const FVector OldSafeHitNormal = SafeWalkNormal;
+		// We don't want to be pushed up an unwalkable surface.
+		if (NormalDot > 0.f && !UFloorQueryUtils::IsHitSurfaceWalkable(Hit, UpDirection, MaxWalkSlopeCosine))
+		{
+			SafeWalkNormal = FVector::VectorPlaneProject(SafeWalkNormal, -UpDirection).GetSafeNormal();
+		}
 
-	FVector SlideDelta = UMovementUtils::ComputeSlideDelta(MovingComps, Delta, PctOfDeltaToMove, SafeWalkNormal, Hit);
-	FVector OriginalSlideDelta = SlideDelta;
+		float PctOfTimeUsed = 0.f;
+		const FVector OldSafeHitNormal = SafeWalkNormal;
+
+		FVector SlideDelta = UMovementUtils::ComputeSlideDelta(MovingComps, Delta, PctOfDeltaToMove, SafeWalkNormal, Hit);
+		FVector OriginalSlideDelta = SlideDelta;
 	
-	if (SlideDelta.Dot(Delta) > 0.f)
-	{
-		UMovementUtils::TrySafeMoveUpdatedComponent(MovingComps, SlideDelta, Rotation, true, Hit, ETeleportType::None, MoveRecord);
+		if (SlideDelta.Dot(Delta) > 0.f)
+		{
+			UMovementUtils::TrySafeMoveUpdatedComponent(MovingComps, SlideDelta, Rotation, true, Hit, ETeleportType::None, MoveRecord);
 
-		PctOfTimeUsed = Hit.Time;
+			PctOfTimeUsed = Hit.Time;
 
-		if (Hit.IsValidBlockingHit())
-		{
-			// Notify first impact
-			if (MoverComponent && bHandleImpact)
+			if (Hit.IsValidBlockingHit())
 			{
-				FMoverOnImpactParams ImpactParams(NAME_None, Hit, SlideDelta);
-				MoverComponent->HandleImpact(ImpactParams);
-			}
+				// Notify first impact
+				if (MoverComponent && bHandleImpact)
+				{
+					FMoverOnImpactParams ImpactParams(NAME_None, Hit, SlideDelta);
+					MoverComponent->HandleImpact(ImpactParams);
+				}
 
-			// Compute new slide normal when hitting multiple surfaces.
-			SlideDelta = UMovementUtils::ComputeTwoWallAdjustedDelta(MovingComps, SlideDelta, Hit, OldSafeHitNormal);
-			const float SlideDeltaDot = SlideDelta.Dot(UpDirection);
-			const float HitNormalDot = Hit.Normal.Dot(UpDirection);
-			if (SlideDeltaDot > 0.f && UFloorQueryUtils::IsHitSurfaceWalkable(Hit, UpDirection, MaxWalkSlopeCosine) && HitNormalDot > UE_KINDA_SMALL_NUMBER)
-			{
-				// Maintain horizontal velocity
-				const float Time = (1.f - Hit.Time);
-				const FVector ScaledDelta = SlideDelta.GetSafeNormal() * SlideDelta.Size();
-				const FVector::FReal NewDeltaDot = (ScaledDelta.Dot(UpDirection) / HitNormalDot);
-				FVector SlideProjectedToGrav = FVector::VectorPlaneProject(OriginalSlideDelta, -UpDirection);
-				SlideDelta = (SlideProjectedToGrav + NewDeltaDot * UpDirection) * Time;
-				// Should never exceed MaxStepHeight in vertical component, so rescale if necessary.
-				// This should be rare (HitNormalDot above would have been very small) but we'd rather lose horizontal velocity than go too high.
-				if (SlideDeltaDot > MaxStepHeight)
+				// Compute new slide normal when hitting multiple surfaces.
+				SlideDelta = UMovementUtils::ComputeTwoWallAdjustedDelta(MovingComps, SlideDelta, Hit, OldSafeHitNormal);
+				const float SlideDeltaDot = SlideDelta.Dot(UpDirection);
+				const float HitNormalDot = Hit.Normal.Dot(UpDirection);
+				if (SlideDeltaDot > 0.f && UFloorQueryUtils::IsHitSurfaceWalkable(Hit, UpDirection, MaxWalkSlopeCosine) && HitNormalDot > UE_KINDA_SMALL_NUMBER)
 				{
-					const float Rescale = MaxStepHeight / SlideDeltaDot;
-					SlideDelta *= Rescale;
+					// Maintain horizontal velocity
+					const float Time = (1.f - Hit.Time);
+					const FVector ScaledDelta = SlideDelta.GetSafeNormal() * SlideDelta.Size();
+					const FVector::FReal NewDeltaDot = (ScaledDelta.Dot(UpDirection) / HitNormalDot);
+					FVector SlideProjectedToGrav = FVector::VectorPlaneProject(OriginalSlideDelta, -UpDirection);
+					SlideDelta = (SlideProjectedToGrav + NewDeltaDot * UpDirection) * Time;
+					// Should never exceed MaxStepHeight in vertical component, so rescale if necessary.
+					// This should be rare (HitNormalDot above would have been very small) but we'd rather lose horizontal velocity than go too high.
+					if (SlideDeltaDot > MaxStepHeight)
+					{
+						const float Rescale = MaxStepHeight / SlideDeltaDot;
+						SlideDelta *= Rescale;
+					}
+				}
+				else
+				{
+					SlideDelta = FVector::VectorPlaneProject(SlideDelta, -UpDirection);
 				}
-			}
-			else
-			{
-				SlideDelta = FVector::VectorPlaneProject(SlideDelta, -UpDirection);
-			}
 			
-			// Only proceed if the new direction is of significant length and not in reverse of original attempted move.
-			if (!SlideDelta.IsNearlyZero(UE::MoverUtils::SMALL_MOVE_DISTANCE) && (SlideDelta | Delta) > 0.f)
-			{
-				// Perform second move
-				UMovementUtils::TrySafeMoveUpdatedComponent(MovingComps, SlideDelta, Rotation, true, Hit, ETeleportType::None, MoveRecord);
-				PctOfTimeUsed += (Hit.Time * (1.f - PctOfTimeUsed));
-
-				// Notify second impact
-				if (MoverComponent && bHandleImpact && Hit.bBlockingHit)
+				// Only proceed if the new direction is of significant length and not in reverse of original attempted move.
+				if (!SlideDelta.IsNearlyZero(UE::MoverUtils::SMALL_MOVE_DISTANCE) && (SlideDelta | Delta) > 0.f)
 				{
-					FMoverOnImpactParams ImpactParams(NAME_None, Hit, SlideDelta);
-					MoverComponent->HandleImpact(ImpactParams);
+					// Perform second move
+					UMovementUtils::TrySafeMoveUpdatedComponent(MovingComps, SlideDelta, Rotation, true, Hit, ETeleportType::None, MoveRecord);
+					PctOfTimeUsed += (Hit.Time * (1.f - PctOfTimeUsed));
+
+					// Notify second impact
+					if (MoverComponent && bHandleImpact && Hit.bBlockingHit)
+					{
+						FMoverOnImpactParams ImpactParams(NAME_None, Hit, SlideDelta);
+						MoverComponent->HandleImpact(ImpactParams);
+					}
 				}
 			}
-		}
 
-		return FMath::Clamp(PctOfTimeUsed, 0.f, 1.f);
+			return FMath::Clamp(PctOfTimeUsed, 0.f, 1.f);
+		}
 	}
 
 	return 0.f;
