Mover properly initializing some variables used for calculating different movement

#rnx

diff --git a/Engine/Plugins/Experimental/Mover/Source/Mover/Public/MoveLibrary/AirMovementUtils.h b/Engine/Plugins/Experimental/Mover/Source/Mover/Public/MoveLibrary/AirMovementUtils.h
--- a/Engine/Plugins/Experimental/Mover/Source/Mover/Public/MoveLibrary/AirMovementUtils.h
+++ b/Engine/Plugins/Experimental/Mover/Source/Mover/Public/MoveLibrary/AirMovementUtils.h
@@ -49,7 +49,7 @@ struct MOVER_API FFreeMoveParams
 	float DeltaSeconds = 0.f;
 
 	UPROPERTY(BlueprintReadOnly, EditAnywhere, Category = Mover)
-	FQuat WorldToGravityQuat;
+	FQuat WorldToGravityQuat = FQuat::Identity;
 };
 
 /**
diff --git a/Engine/Plugins/Experimental/Mover/Source/Mover/Public/MoveLibrary/GroundMovementUtils.h b/Engine/Plugins/Experimental/Mover/Source/Mover/Public/MoveLibrary/GroundMovementUtils.h
--- a/Engine/Plugins/Experimental/Mover/Source/Mover/Public/MoveLibrary/GroundMovementUtils.h
+++ b/Engine/Plugins/Experimental/Mover/Source/Mover/Public/MoveLibrary/GroundMovementUtils.h
@@ -60,10 +60,10 @@ struct MOVER_API FGroundMoveParams
 	float DeltaSeconds = 0.f;
 
 	UPROPERTY(BlueprintReadOnly, EditAnywhere, Category = Mover)
-	FVector UpDirection;
+	FVector UpDirection = FVector::UpVector;
 	
 	UPROPERTY(BlueprintReadOnly, EditAnywhere, Category = Mover)
-	FQuat WorldToGravityQuat;
+	FQuat WorldToGravityQuat = FQuat::Identity;
 };
 
 /**
diff --git a/Engine/Plugins/Experimental/Mover/Source/Mover/Public/MoveLibrary/WaterMovementUtils.h b/Engine/Plugins/Experimental/Mover/Source/Mover/Public/MoveLibrary/WaterMovementUtils.h
--- a/Engine/Plugins/Experimental/Mover/Source/Mover/Public/MoveLibrary/WaterMovementUtils.h
+++ b/Engine/Plugins/Experimental/Mover/Source/Mover/Public/MoveLibrary/WaterMovementUtils.h
@@ -60,7 +60,7 @@ struct MOVER_API FWaterMoveParams
 	float MoveSpeed = 0.f;
 	
 	UPROPERTY(BlueprintReadOnly, EditAnywhere, Category = Mover)
-	FQuat WorldToGravityQuat;
+	FQuat WorldToGravityQuat = FQuat::Identity;
 };
 
 // Data about the water volume and its interaction with the pawn used in calculating swimming movement
