Mover fixing crash when trying to queue a stance modifier on an actor where the MoverComponent was added in BP not natively

#rnx

diff --git a/Engine/Plugins/Experimental/Mover/Source/Mover/Private/DefaultMovementSet/MovementModifiers/StanceModifier.cpp b/Engine/Plugins/Experimental/Mover/Source/Mover/Private/DefaultMovementSet/MovementModifiers/StanceModifier.cpp
--- a/Engine/Plugins/Experimental/Mover/Source/Mover/Private/DefaultMovementSet/MovementModifiers/StanceModifier.cpp
+++ b/Engine/Plugins/Experimental/Mover/Source/Mover/Private/DefaultMovementSet/MovementModifiers/StanceModifier.cpp
@@ -9,8 +9,6 @@
 #include "DefaultMovementSet/InstantMovementEffects/BasicInstantMovementEffects.h"
 #include "DefaultMovementSet/Settings/StanceSettings.h"
 #include "Engine/Blueprint.h"
-#include "Engine/SCS_Node.h"
-#include "Engine/SimpleConstructionScript.h"
 #include "GameFramework/Pawn.h"
 #include "MoveLibrary/MovementUtils.h"
 
@@ -113,7 +111,7 @@ void FStanceModifier::OnEnd(UMoverComponent* MoverComp, const FMoverTimeStep& Ti
 
 	if (UCapsuleComponent* CapsuleComponent = Cast<UCapsuleComponent>(MoverComp->GetUpdatedComponent()))
 	{
-		if (const UCapsuleComponent* OriginalCapsule = GetOriginalComponentType<UCapsuleComponent>(MoverComp))
+		if (const UCapsuleComponent* OriginalCapsule = UMovementUtils::GetOriginalComponentType<UCapsuleComponent>(MoverComp->GetOwner()))
 		{
 			if (const APawn* OwnerCDOAsPawn = Cast<APawn>(OwnerCDO))
 			{
@@ -175,7 +173,7 @@ bool FStanceModifier::CanExpand(const UCharacterMoverComponent* MoverComp) const
 	USceneComponent* UpdatedComponent = MoverComp->GetUpdatedComponent();
 	UPrimitiveComponent* UpdatedCompAsPrimitive = Cast<UPrimitiveComponent>(UpdatedComponent);
 	
-	if (const UCapsuleComponent* OriginalCapsule = GetOriginalComponentType<UCapsuleComponent>(MoverComp))
+	if (const UCapsuleComponent* OriginalCapsule = UMovementUtils::GetOriginalComponentType<UCapsuleComponent>(MoverComp->GetOwner()))
 	{
 		StandingHalfHeight = OriginalCapsule->GetScaledCapsuleHalfHeight();
 	}
@@ -296,13 +294,16 @@ void FStanceModifier::ApplyMovementSettings(UMoverComponent* MoverComp)
 
 void FStanceModifier::RevertMovementSettings(UMoverComponent* MoverComp)
 {
-	const UMoverComponent* CDOMoverComp = GetOriginalComponentType<UMoverComponent>(MoverComp);
-	const UCommonLegacyMovementSettings* OriginalMovementSettings = CDOMoverComp->FindSharedSettings<UCommonLegacyMovementSettings>();
-	
-	// Revert movement settings back to original settings
-	if (UCommonLegacyMovementSettings* MovementSettings = MoverComp->FindSharedSettings_Mutable<UCommonLegacyMovementSettings>())
+	if (const UMoverComponent* CDOMoverComp = UMovementUtils::GetOriginalComponentType<UMoverComponent>(MoverComp->GetOwner()))
 	{
-		MovementSettings->Acceleration = OriginalMovementSettings->Acceleration;
-		MovementSettings->MaxSpeed = OriginalMovementSettings->MaxSpeed;
+		const UCommonLegacyMovementSettings* OriginalMovementSettings = CDOMoverComp->FindSharedSettings<UCommonLegacyMovementSettings>();
+		UCommonLegacyMovementSettings* MovementSettings = MoverComp->FindSharedSettings_Mutable<UCommonLegacyMovementSettings>();
+		
+		// Revert movement settings back to original settings
+		if (MovementSettings && OriginalMovementSettings)
+		{
+			MovementSettings->Acceleration = OriginalMovementSettings->Acceleration;
+			MovementSettings->MaxSpeed = OriginalMovementSettings->MaxSpeed;
+		}
 	}
 }
