Mover | Swimming | Water Current Velocity Bug Fix
* Water spline data update wasn't being propogated onto the SimBlackboard.

diff --git a/Engine/Plugins/Experimental/Mover/Source/Mover/Private/DefaultMovementSet/Modes/SwimmingMode.cpp b/Engine/Plugins/Experimental/Mover/Source/Mover/Private/DefaultMovementSet/Modes/SwimmingMode.cpp
--- a/Engine/Plugins/Experimental/Mover/Source/Mover/Private/DefaultMovementSet/Modes/SwimmingMode.cpp
+++ b/Engine/Plugins/Experimental/Mover/Source/Mover/Private/DefaultMovementSet/Modes/SwimmingMode.cpp
@@ -35,17 +35,9 @@ void USwimmingMode::OnGenerateMove(const FMoverTickStartData& StartState, const
 	
 	const float CapsuleHalfHeight = MoverComp->GetOwner()->GetSimpleCollisionHalfHeight();
 
-	if (SimBlackboard && SimBlackboard->TryGet(CommonBlackboard::LastWaterResult, LastWaterResult) && LastWaterResult.IsSwimmableVolume())
+	if (SimBlackboard)
 	{
-		FUpdateWaterSplineDataParams Params;
-		Params.TargetImmersionDepth = CommonLegacySettings->SwimmingIdealImmersionDepth;
-		Params.WaterVelocityDepthForMax = SurfaceSwimmingWaterControlSettings.WaterVelocityDepthForMax;
-		Params.WaterVelocityMinMultiplier = SurfaceSwimmingWaterControlSettings.WaterVelocityMinMultiplier;
-		Params.PlayerVelocity = StartingSyncState->GetVelocity_WorldSpace();
-		Params.CapsuleHalfHeight = CapsuleHalfHeight;
-		Params.PlayerLocation = StartingSyncState->GetLocation_WorldSpace();
-		
-		UWaterMovementUtils::UpdateWaterSplineData(Params, LastWaterResult);
+		SimBlackboard->TryGet(CommonBlackboard::LastWaterResult, LastWaterResult);
 	}
 	
 	FVector Velocity = StartingSyncState->GetVelocity_WorldSpace();
diff --git a/Engine/Plugins/Experimental/Mover/Source/Mover/Private/PhysicsMover/Modes/PhysicsDrivenSwimmingMode.cpp b/Engine/Plugins/Experimental/Mover/Source/Mover/Private/PhysicsMover/Modes/PhysicsDrivenSwimmingMode.cpp
--- a/Engine/Plugins/Experimental/Mover/Source/Mover/Private/PhysicsMover/Modes/PhysicsDrivenSwimmingMode.cpp
+++ b/Engine/Plugins/Experimental/Mover/Source/Mover/Private/PhysicsMover/Modes/PhysicsDrivenSwimmingMode.cpp
@@ -100,9 +100,22 @@ void UPhysicsDrivenSwimmingMode::OnSimulationTick(const FSimulationTickParams& P
 	UPhysicsMovementUtils::FloorSweep_Internal(StartingSyncState->GetLocation_WorldSpace(), StartingSyncState->GetVelocity_WorldSpace() * DeltaSeconds,
 		UpdatedPrimitive, UpDir, PawnRadius, QueryDistance, CommonLegacySettings->MaxWalkSlopeCosine, TargetHeight, FloorResult, WaterResult);
 
+	if (WaterResult.IsSwimmableVolume())
+	{
+		FUpdateWaterSplineDataParams SplineParams;
+		SplineParams.TargetImmersionDepth = CommonLegacySettings->SwimmingIdealImmersionDepth;
+		SplineParams.WaterVelocityDepthForMax = SurfaceSwimmingWaterControlSettings.WaterVelocityDepthForMax;
+		SplineParams.WaterVelocityMinMultiplier = SurfaceSwimmingWaterControlSettings.WaterVelocityMinMultiplier;
+		SplineParams.PlayerVelocity = StartingSyncState->GetVelocity_WorldSpace();
+		SplineParams.CapsuleHalfHeight = Params.MovingComps.MoverComponent->GetOwner()->GetSimpleCollisionHalfHeight();
+		SplineParams.PlayerLocation = StartingSyncState->GetLocation_WorldSpace();
+
+		UWaterMovementUtils::UpdateWaterSplineData(SplineParams, WaterResult);
+	}
+	
 	SimBlackboard->Set(CommonBlackboard::LastFloorResult, FloorResult);
 	SimBlackboard->Set(CommonBlackboard::LastWaterResult, WaterResult);
-	
+
 	if (WaterResult.IsSwimmableVolume())
 	{
 		const bool bIsWithinReach = FloorResult.FloorDist <= TargetHeight;
