Crouch | Physics Mover
* Support Crouching on Physics Mover Pawns
* Added / updated example assets to MoverExamples
* Modified the sync state rollback flow to allow capturing some information from invalidating sync states.
* Added a new `OnPreMovement` callback similar to `OnPreSimulationTick`, but on the Physics Thread. This will be complementing `OnPostMovement`.

#virtualized

diff --git a/Engine/Plugins/Experimental/Mover/Source/Mover/Private/Backends/MoverNetworkPhysicsLiaison.cpp b/Engine/Plugins/Experimental/Mover/Source/Mover/Private/Backends/MoverNetworkPhysicsLiaison.cpp
--- a/Engine/Plugins/Experimental/Mover/Source/Mover/Private/Backends/MoverNetworkPhysicsLiaison.cpp
+++ b/Engine/Plugins/Experimental/Mover/Source/Mover/Private/Backends/MoverNetworkPhysicsLiaison.cpp
@@ -826,9 +826,9 @@ void UMoverNetworkPhysicsLiaisonComponent::PostPhysicsUpdate_External()
 			MoverComponentPtr->bHasValidCachedState = true;
 
 			const FName NextMode = MoverComponentPtr->GetMovementModeName();
+			UpdateConstraintSettings();
 			if (CachedMovementMode != NextMode)
 			{
-				UpdateConstraintSettings();
 				MoverComponentPtr->OnMovementModeChanged.Broadcast(CachedMovementMode, NextMode);
 
 				if (bIsCachedModeFalling && MoverComponentPtr->HasGameplayTag(Mover_IsOnGround, true))
@@ -864,8 +864,6 @@ void UMoverNetworkPhysicsLiaisonComponent::ProcessInputs_Internal(int32 PhysicsS
 			}
 			else
 			{
-				GetCurrentStateData(Input.SyncState);
-
 				bool bIsSolverResim = false;
 				bool bIsFirstResimFrame = false;
 				if (Chaos::FPhysicsSolver* Solver = PhysScene->GetSolver())
@@ -892,10 +890,20 @@ void UMoverNetworkPhysicsLiaisonComponent::ProcessInputs_Internal(int32 PhysicsS
 				{
 					FMoverAuxStateContext UnusedInvalidAuxState;
 					FMoverAuxStateContext UnusedAuxState;
-					MoverComp->OnSimulationPreRollback(&LatestSyncState, &Input.SyncState, &UnusedInvalidAuxState, &UnusedAuxState);
+
+					MoverComp->OnSimulationPreRollback(&Input.SyncState, &LatestSyncState, &UnusedInvalidAuxState, &UnusedAuxState);
+					GetCurrentStateData(Input.SyncState);
 					MoverComp->OnSimulationRollback(&Input.SyncState, &UnusedAuxState);
 				}
+				else
+				{
+					GetCurrentStateData(Input.SyncState);
+				}
 			}
+
+			const FMoverTimeStep MoverTimeStep = bUsingAsyncPhysics ? GetCurrentAsyncMoverTimeStep_Internal() : GetCurrentMoverTimeStep(DeltaTime);
+			FMoverAuxStateContext UnusedAuxState;
+			MoverComp->OnPreMovement.Broadcast(MoverTimeStep, Input.InputCmd, Input.SyncState, UnusedAuxState);
 		}
 	}
 
diff --git a/Engine/Plugins/Experimental/Mover/Source/Mover/Private/MovementModeStateMachine.cpp b/Engine/Plugins/Experimental/Mover/Source/Mover/Private/MovementModeStateMachine.cpp
--- a/Engine/Plugins/Experimental/Mover/Source/Mover/Private/MovementModeStateMachine.cpp
+++ b/Engine/Plugins/Experimental/Mover/Source/Mover/Private/MovementModeStateMachine.cpp
@@ -566,6 +566,10 @@ void UMovementModeStateMachine::RollbackModifiers(const FMoverSyncState* Invalid
 				if (ModifierFromRollbackIt->Get()->Matches(ModifierFromCacheIt->Get()))
 				{
 					bContainsModifier = true;
+
+					// Rolled back version of the modifier will be missing the handle; we fix that here
+					ModifierFromRollbackIt->Get()->OverwriteHandleIfInvalid(ModifierFromCacheIt->Get()->GetHandle());
+
 					break;
 				}
 			}
diff --git a/Engine/Plugins/Experimental/Mover/Source/Mover/Private/MoverComponent.cpp b/Engine/Plugins/Experimental/Mover/Source/Mover/Private/MoverComponent.cpp
--- a/Engine/Plugins/Experimental/Mover/Source/Mover/Private/MoverComponent.cpp
+++ b/Engine/Plugins/Experimental/Mover/Source/Mover/Private/MoverComponent.cpp
@@ -474,6 +474,8 @@ void UMoverComponent::SimulationTick(const FMoverTimeStep& InTimeStep, const FMo
 		ModeFSM->QueueNextMode(Input->SuggestedMovementMode);
 	}
 
+	OnPreMovement.Broadcast(MoverTimeStep, SimInput.InputCmd, SimInput.SyncState, SimInput.AuxState);
+
 	// Tick the actual simulation. This is where the proposed moves are queried and executed, affecting change to the moving actor's gameplay state and captured in the output sim state
 	{
 		FScopedMovementUpdate ScopedMovementUpdate(UpdatedComponent, EScopedUpdate::DeferredUpdates);
diff --git a/Engine/Plugins/Experimental/Mover/Source/Mover/Private/PhysicsMover/MovementModifiers/PhysicsStanceModifier.cpp b/Engine/Plugins/Experimental/Mover/Source/Mover/Private/PhysicsMover/MovementModifiers/PhysicsStanceModifier.cpp
new file mode 100644
--- /dev/null
+++ b/Engine/Plugins/Experimental/Mover/Source/Mover/Private/PhysicsMover/MovementModifiers/PhysicsStanceModifier.cpp
@@ -0,0 +1,87 @@
+﻿// Copyright Epic Games, Inc. All Rights Reserved.
+
+#include "PhysicsMover/MovementModifiers/PhysicsStanceModifier.h"
+
+#include "Components/CapsuleComponent.h"
+#include "DefaultMovementSet/Settings/StanceSettings.h"
+#include "MoveLibrary/MovementUtils.h"
+#include "PhysicsMover/Modes/PhysicsDrivenWalkingMode.h"
+#include "PhysicsMover/Modes/PhysicsDrivenFallingMode.h"
+
+
+void FPhysicsStanceModifier::AdjustCapsule(UMoverComponent* MoverComp, float OldHalfHeight, float NewHalfHeight, float NewEyeHeight)
+{
+	const float HalfHeightDifference = FMath::Abs(NewHalfHeight - OldHalfHeight);
+	const bool bExpanding = OldHalfHeight < NewHalfHeight;
+
+	// Adjust the TargetHeight of all the relevant MovementModes
+	UpdateTargetHeight<UPhysicsDrivenWalkingMode>(MoverComp, OldHalfHeight, NewHalfHeight);
+	UpdateTargetHeight<UPhysicsDrivenFallingMode>(MoverComp, OldHalfHeight, NewHalfHeight);
+}
+
+void FPhysicsStanceModifier::OnPostSimulationTick(const FStanceModifier* Modifier, UMoverComponent* MoverComp, UPrimitiveComponent* UpdatedPrimitive,  bool bIsCrouching, bool& bPostProcessed, OUT bool& bStanceChanged)
+{
+	if(const UStanceSettings* StanceSettings = MoverComp->FindSharedSettings<UStanceSettings>())
+	{
+		if (UCapsuleComponent* CapsuleComponent = Cast<UCapsuleComponent>(UpdatedPrimitive))
+		{
+			const UCapsuleComponent* OriginalCapsule = UMovementUtils::GetOriginalComponentType<UCapsuleComponent>(MoverComp->GetOwner());
+
+			if (OriginalCapsule)
+			{
+				if (Modifier && bIsCrouching && !bPostProcessed) // Crouching
+				{
+					const float OldHalfHeight = OriginalCapsule->GetScaledCapsuleHalfHeight();
+					const float NewHalfHeight = StanceSettings->CrouchHalfHeight;
+					const float HalfHeightDifference = NewHalfHeight - OldHalfHeight;
+					const float ColliderOffset = -MoverComp->GetPrimaryVisualComponent()->GetRelativeLocation().Z - OldHalfHeight;
+	
+					CapsuleComponent->SetCapsuleSize(CapsuleComponent->GetUnscaledCapsuleRadius(), NewHalfHeight);
+					MoverComp->GetPrimaryVisualComponent()->SetRelativeLocation(-(NewHalfHeight + ColliderOffset) * MoverComp->GetUpDirection());
+	
+					bStanceChanged = true;
+					bPostProcessed = true;
+				}
+				else if (!Modifier && !bIsCrouching && bPostProcessed) // Uncrouching
+				{
+					const float OldHalfHeight = StanceSettings->CrouchHalfHeight;
+					const float NewHalfHeight = OriginalCapsule->GetScaledCapsuleHalfHeight();
+					const float HalfHeightDifference = NewHalfHeight - OldHalfHeight;
+					const float ColliderOffset = -MoverComp->GetPrimaryVisualComponent()->GetRelativeLocation().Z - OldHalfHeight;
+	
+					CapsuleComponent->SetCapsuleSize(CapsuleComponent->GetUnscaledCapsuleRadius(), NewHalfHeight);
+					MoverComp->GetPrimaryVisualComponent()->SetRelativeLocation(-(NewHalfHeight + ColliderOffset) * MoverComp->GetUpDirection());
+	
+					bStanceChanged = true;
+					bPostProcessed = false;
+				}
+			}
+		}
+	}
+}
+
+FMovementModifierBase* FPhysicsStanceModifier::Clone() const
+{
+	FPhysicsStanceModifier* CopyPtr = new FPhysicsStanceModifier(*this);
+	return CopyPtr;
+}
+
+void FPhysicsStanceModifier::NetSerialize(FArchive& Ar)
+{
+	Super::NetSerialize(Ar);
+}
+
+UScriptStruct* FPhysicsStanceModifier::GetScriptStruct() const
+{
+	return FPhysicsStanceModifier::StaticStruct();
+}
+
+FString FPhysicsStanceModifier::ToSimpleString() const
+{
+	return FString::Printf(TEXT("Physics Based Stance Modifier"));
+}
+
+void FPhysicsStanceModifier::AddReferencedObjects(FReferenceCollector& Collector)
+{
+	Super::AddReferencedObjects(Collector);
+}
diff --git a/Engine/Plugins/Experimental/Mover/Source/Mover/Private/PhysicsMover/PhysicsCharacterMoverComponent.cpp b/Engine/Plugins/Experimental/Mover/Source/Mover/Private/PhysicsMover/PhysicsCharacterMoverComponent.cpp
new file mode 100644
--- /dev/null
+++ b/Engine/Plugins/Experimental/Mover/Source/Mover/Private/PhysicsMover/PhysicsCharacterMoverComponent.cpp
@@ -0,0 +1,98 @@
+// Copyright Epic Games, Inc. All Rights Reserved.
+
+#include "PhysicsMover/PhysicsCharacterMoverComponent.h"
+
+#include "PhysicsMover/Modes/PhysicsDrivenFallingMode.h"
+#include "PhysicsMover/Modes/PhysicsDrivenFlyingMode.h"
+#include "PhysicsMover/Modes/PhysicsDrivenWalkingMode.h"
+#include "PhysicsMover/MovementModifiers/PhysicsStanceModifier.h"
+#include "Backends/MoverNetworkPhysicsLiaison.h"
+
+UPhysicsCharacterMoverComponent::UPhysicsCharacterMoverComponent()
+{
+	// Override with physics movement modes
+	MovementModes.Add(DefaultModeNames::Walking, CreateDefaultSubobject<UPhysicsDrivenWalkingMode>(TEXT("PhysicsDrivenWalkingMode")));
+	MovementModes.Add(DefaultModeNames::Falling, CreateDefaultSubobject<UPhysicsDrivenFallingMode>(TEXT("PhysicsDrivenFallingMode")));
+	MovementModes.Add(DefaultModeNames::Flying,  CreateDefaultSubobject<UPhysicsDrivenFlyingMode>(TEXT("PhysicsDrivenFlyingMode")));
+
+	BackendClass = UMoverNetworkPhysicsLiaisonComponent::StaticClass();
+
+	bHandleJump = false;
+}
+
+void UPhysicsCharacterMoverComponent::BeginPlay()
+{
+	Super::BeginPlay();
+
+	OnPreMovement.AddDynamic(this, &UPhysicsCharacterMoverComponent::OnMoverPreMovement);
+	OnPostSimulationTick.AddDynamic(this, &UPhysicsCharacterMoverComponent::OnMoverPostSimulationTick);
+
+	if (bHandleJump)
+	{
+		UE_LOG(LogMover, Warning, TEXT("Handle Jump flag is ignored for Physics Character Mover Component; jumps are handled via the Physics Jump Check Transition."));
+	}
+}
+
+void UPhysicsCharacterMoverComponent::OnMoverPreMovement(const FMoverTimeStep& TimeStep, const FMoverInputCmdContext& InputCmd, const FMoverSyncState& SyncState, const FMoverAuxStateContext& AuxState)
+{
+	const FStanceModifier* StanceModifier = static_cast<const FStanceModifier*>(FindMovementModifier(StanceModifierHandle));
+	// This is a fail safe in case our handle was bad - try finding the modifier by type if we can
+	if (!StanceModifier)
+	{
+		StanceModifier = FindMovementModifierByType<FStanceModifier>();
+	}
+
+	EStanceMode OldActiveStance = EStanceMode::Invalid;
+	if (StanceModifier)
+	{
+		OldActiveStance = StanceModifier->ActiveStance;
+	}
+
+	const bool bIsCrouching = HasGameplayTag(Mover_IsCrouching, true);
+	if (bIsCrouching && (!bWantsToCrouch || !CanCrouch()))
+	{
+		if (StanceModifier && StanceModifier->CanExpand(this))
+		{
+			CancelModifierFromHandle(StanceModifier->GetHandle());
+			StanceModifierHandle.Invalidate();
+
+			StanceModifier = nullptr;
+		}
+	}
+	else if (!bIsCrouching && bWantsToCrouch && CanCrouch())
+	{
+		TSharedPtr<FStanceModifier> NewStanceModifier = MakeShared<FPhysicsStanceModifier>();
+		StanceModifierHandle = QueueMovementModifier(NewStanceModifier);
+
+		StanceModifier = NewStanceModifier.Get();
+	}
+
+	// Ensure that StanceModifierHandle is consistent with the crouch state 
+	if(!bIsCrouching && !bWantsToCrouch && StanceModifierHandle.IsValid())
+	{
+		StanceModifierHandle.Invalidate();
+	}
+}
+
+void UPhysicsCharacterMoverComponent::OnMoverPostSimulationTick(const FMoverTimeStep& TimeStep)
+{
+	const FStanceModifier* StanceModifier = static_cast<const FStanceModifier*>(FindMovementModifier(StanceModifierHandle));
+	if (!StanceModifier)
+	{
+		StanceModifier = FindMovementModifierByType<FStanceModifier>();
+	}
+
+	const bool bIsCrouching = HasGameplayTag(Mover_IsCrouching, true);
+	bool bStanceChanged = false;
+
+	FPhysicsStanceModifier::OnPostSimulationTick(StanceModifier, this, UpdatedCompAsPrimitive, bIsCrouching, bStancePostProcessed, bStanceChanged);
+
+	if (bStanceChanged && bIsCrouching)
+	{
+		OnStanceChanged.Broadcast(EStanceMode::Invalid, EStanceMode::Crouch);
+	}
+	else if (bStanceChanged && !bIsCrouching)
+	{
+		OnStanceChanged.Broadcast(EStanceMode::Crouch, EStanceMode::Invalid);
+	}
+}
\ No newline at end of file
diff --git a/Engine/Plugins/Experimental/Mover/Source/Mover/Public/DefaultMovementSet/MovementModifiers/StanceModifier.h b/Engine/Plugins/Experimental/Mover/Source/Mover/Public/DefaultMovementSet/MovementModifiers/StanceModifier.h
--- a/Engine/Plugins/Experimental/Mover/Source/Mover/Public/DefaultMovementSet/MovementModifiers/StanceModifier.h
+++ b/Engine/Plugins/Experimental/Mover/Source/Mover/Public/DefaultMovementSet/MovementModifiers/StanceModifier.h
@@ -67,7 +67,7 @@ public:
 
 protected:
 	// Modifies the updated component casted to a capsule component
-	void AdjustCapsule(UMoverComponent* MoverComp, float OldHalfHeight, float NewHalfHeight, float NewEyeHeight);
+	virtual void AdjustCapsule(UMoverComponent* MoverComp, float OldHalfHeight, float NewHalfHeight, float NewEyeHeight);
 
 	// Applies any movement settings like acceleration or max speed changes
 	void ApplyMovementSettings(UMoverComponent* MoverComp);
diff --git a/Engine/Plugins/Experimental/Mover/Source/Mover/Public/MoverComponent.h b/Engine/Plugins/Experimental/Mover/Source/Mover/Public/MoverComponent.h
--- a/Engine/Plugins/Experimental/Mover/Source/Mover/Public/MoverComponent.h
+++ b/Engine/Plugins/Experimental/Mover/Source/Mover/Public/MoverComponent.h
@@ -34,6 +34,9 @@ namespace MoverComponentConstants
 // Fired just before a simulation tick, regardless of being a re-simulated frame or not.
 DECLARE_DYNAMIC_MULTICAST_DELEGATE_TwoParams(FMover_OnPreSimTick, const FMoverTimeStep&, TimeStep, const FMoverInputCmdContext&, InputCmd);
 
+// Fired during a simulation tick, after the input is processed but before the actual move calculation.
+DECLARE_DYNAMIC_MULTICAST_DELEGATE_FourParams(FMover_OnPreMovement, const FMoverTimeStep&, TimeStep, const FMoverInputCmdContext&, InputCmd, const FMoverSyncState&, SyncState, const FMoverAuxStateContext&, AuxState);
+
 // Fired during a simulation tick, after movement has occurred but before the state is finalized, allowing changes to the output state.
 DECLARE_DYNAMIC_MULTICAST_DELEGATE_ThreeParams(FMover_OnPostMovement, const FMoverTimeStep&, TimeStep, FMoverSyncState&, SyncState, FMoverAuxStateContext&, AuxState);
 
@@ -70,7 +73,8 @@ public:
 
 	virtual void TickComponent(float DeltaTime, enum ELevelTick TickType, FActorComponentTickFunction* ThisTickFunction) override;
 
-	// Broadcast before each simulation tick
+	// Broadcast before each simulation tick.
+	// Note - Guaranteed to run on the game thread (even in async simulation).
 	UPROPERTY(BlueprintAssignable, Category = Mover)
 	FMover_OnPreSimTick OnPreSimulationTick;
 
@@ -470,6 +474,11 @@ protected:
 	void FindDefaultUpdatedComponent();
 	void UpdateTickRegistration();
 
+	// Broadcast during the simulation tick after inputs have been processed, but before the actual move is performed.
+	// Note - When async simulating, the delegate would be called on the async thread, and might be broadcast multiple times.
+	UPROPERTY(BlueprintAssignable, Category = Mover)
+	FMover_OnPreMovement OnPreMovement;
+
 	/** Called when a rollback occurs, before the simulation state has been restored */
 	void OnSimulationPreRollback(const FMoverSyncState* InvalidSyncState, const FMoverSyncState* SyncState, const FMoverAuxStateContext* InvalidAuxState, const FMoverAuxStateContext* AuxState);
 	
diff --git a/Engine/Plugins/Experimental/Mover/Source/Mover/Public/PhysicsMover/MovementModifiers/PhysicsStanceModifier.h b/Engine/Plugins/Experimental/Mover/Source/Mover/Public/PhysicsMover/MovementModifiers/PhysicsStanceModifier.h
new file mode 100644
--- /dev/null
+++ b/Engine/Plugins/Experimental/Mover/Source/Mover/Public/PhysicsMover/MovementModifiers/PhysicsStanceModifier.h
@@ -0,0 +1,44 @@
+﻿// Copyright Epic Games, Inc. All Rights Reserved.
+
+#pragma once
+
+#include "MoverComponent.h"
+#include "DefaultMovementSet/MovementModifiers/StanceModifier.h"
+
+#include "PhysicsStanceModifier.generated.h"
+
+/**
+ * Physics specialized version of FStanceModifier
+ */
+USTRUCT(BlueprintType)
+struct MOVER_API FPhysicsStanceModifier : public FStanceModifier
+{
+	GENERATED_BODY()
+
+	// @return newly allocated copy of this FMovementModifier. Must be overridden by child classes
+	virtual FMovementModifierBase* Clone() const override;
+
+	virtual void NetSerialize(FArchive& Ar) override;
+
+	virtual UScriptStruct* GetScriptStruct() const override;
+
+	virtual FString ToSimpleString() const override;
+
+	virtual void AddReferencedObjects(class FReferenceCollector& Collector) override;
+
+	static void OnPostSimulationTick(const FStanceModifier* Modifier, UMoverComponent* MoverComp, UPrimitiveComponent* UpdatedPrimitive,  bool bIsCrouching, bool& bPostProcessed, OUT bool& bStanceChanged);
+
+protected:
+	// Modifies the Capsule Location by updating the target height of the movement modes
+	virtual void AdjustCapsule(UMoverComponent* MoverComp, float OldHalfHeight, float NewHalfHeight, float NewEyeHeight) override;
+
+	template <class T>
+	void UpdateTargetHeight(UMoverComponent* MoverComp, float OldHalfHeight, float NewHalfHeight)
+	{
+		if (T* MovementMode = MoverComp->FindMode_Mutable<T>())
+		{
+			const float ColliderOffset = -MoverComp->GetPrimaryVisualComponent()->GetRelativeLocation().Z - OldHalfHeight;
+			MovementMode->TargetHeight = NewHalfHeight + ColliderOffset;
+		}
+	}
+};
\ No newline at end of file
diff --git a/Engine/Plugins/Experimental/Mover/Source/Mover/Public/PhysicsMover/PhysicsCharacterMoverComponent.h b/Engine/Plugins/Experimental/Mover/Source/Mover/Public/PhysicsMover/PhysicsCharacterMoverComponent.h
new file mode 100644
--- /dev/null
+++ b/Engine/Plugins/Experimental/Mover/Source/Mover/Public/PhysicsMover/PhysicsCharacterMoverComponent.h
@@ -0,0 +1,32 @@
+// Copyright Epic Games, Inc. All Rights Reserved.
+
+#pragma once
+
+#include "DefaultMovementSet/CharacterMoverComponent.h"
+#include "PhysicsCharacterMoverComponent.generated.h"
+
+
+UCLASS(BlueprintType, meta = (BlueprintSpawnableComponent))
+class MOVER_API UPhysicsCharacterMoverComponent : public UCharacterMoverComponent
+{
+	GENERATED_BODY()
+
+public:
+
+	UPhysicsCharacterMoverComponent();
+	
+	virtual void BeginPlay() override;
+
+protected:
+
+	UFUNCTION()
+	virtual void OnMoverPreMovement(const FMoverTimeStep& TimeStep, const FMoverInputCmdContext& InputCmd, const FMoverSyncState& SyncState, const FMoverAuxStateContext& AuxState);
+
+	UFUNCTION()
+	virtual void OnMoverPostSimulationTick(const FMoverTimeStep& TimeStep);
+
+private:
+
+	// Keeps track of post processing that needs to take place on the game thread, after the stance is processed in the simulation tick.
+	bool bStancePostProcessed = false;
+};
