Mover Physics | Crouch | Non-Unity compilation issue

diff --git a/Engine/Plugins/Experimental/Mover/Source/Mover/Public/PhysicsMover/MovementModifiers/PhysicsStanceModifier.h b/Engine/Plugins/Experimental/Mover/Source/Mover/Public/PhysicsMover/MovementModifiers/PhysicsStanceModifier.h
--- a/Engine/Plugins/Experimental/Mover/Source/Mover/Public/PhysicsMover/MovementModifiers/PhysicsStanceModifier.h
+++ b/Engine/Plugins/Experimental/Mover/Source/Mover/Public/PhysicsMover/MovementModifiers/PhysicsStanceModifier.h
@@ -4,6 +4,7 @@
 
 #include "MoverComponent.h"
 #include "DefaultMovementSet/MovementModifiers/StanceModifier.h"
+#include "Components/SceneComponent.h"
 
 #include "PhysicsStanceModifier.generated.h"
 
