Mover: adds support for specifying inputs and starting state to a trajectory prediction query. If not specified, the last-known state/inputs from the MoverComponent will be used, which matches existing behavior.

diff --git a/Engine/Plugins/Experimental/Mover/Source/Mover/Private/MoverComponent.cpp b/Engine/Plugins/Experimental/Mover/Source/Mover/Private/MoverComponent.cpp
--- a/Engine/Plugins/Experimental/Mover/Source/Mover/Private/MoverComponent.cpp
+++ b/Engine/Plugins/Experimental/Mover/Source/Mover/Private/MoverComponent.cpp
@@ -1760,13 +1760,33 @@ TArray<FTrajectorySampleInfo> UMoverComponent::GetPredictedTrajectory(FMoverPred
 {
 	if (ModeFSM)
 	{
-		FMoverTickStartData StartingState;
+		FMoverTickStartData StepState;
 
-		StartingState.InputCmd = GetLastInputCmd();
-		StartingState.SyncState = CachedLastSyncState;
-		StartingState.AuxState = CachedLastAuxState;
+		// Use the last-known input if none are specified.
+		if (PredictionParams.OptionalInputCmds.IsEmpty())
+		{
+			StepState.InputCmd = GetLastInputCmd();
+		}
+
+		// Use preferred starting sync/aux state. Fall back to last-known state if not set.
+		if (PredictionParams.OptionalStartSyncState.IsSet())
+		{
+			StepState.SyncState = PredictionParams.OptionalStartSyncState.GetValue();
+		}
+		else
+		{
+			StepState.SyncState = CachedLastSyncState;
+		}
+
+		if (PredictionParams.OptionalStartAuxState.IsSet())
+		{
+			StepState.AuxState = PredictionParams.OptionalStartAuxState.GetValue();
+		}
+		else
+		{
+			StepState.AuxState = CachedLastAuxState;
+		}
 
-		FMoverTickStartData StepState = StartingState;
 
 		FMoverTimeStep FutureTimeStep;
 		FutureTimeStep.StepMs = (PredictionParams.SecondsPerSample * 1000.f);
@@ -1794,6 +1814,12 @@ TArray<FTrajectorySampleInfo> UMoverComponent::GetPredictedTrajectory(FMoverPred
 
 				for (int32 i = 0; i < PredictionParams.NumPredictionSamples; ++i)
 				{
+					// If no further inputs are specified, the previous input cmd will continue to be used
+					if (i < PredictionParams.OptionalInputCmds.Num())
+					{
+						StepState.InputCmd = PredictionParams.OptionalInputCmds[i];
+					}
+
 					// Capture sample from current step state
 					FTrajectorySampleInfo& Sample = OutSamples[i];
 
diff --git a/Engine/Plugins/Experimental/Mover/Source/Mover/Public/MoverSimulationTypes.h b/Engine/Plugins/Experimental/Mover/Source/Mover/Public/MoverSimulationTypes.h
--- a/Engine/Plugins/Experimental/Mover/Source/Mover/Public/MoverSimulationTypes.h
+++ b/Engine/Plugins/Experimental/Mover/Source/Mover/Public/MoverSimulationTypes.h
@@ -296,4 +296,48 @@ public:
 	/** Contributes additions to the input cmd for this simulation frame. Typically this is translating accumulated user input (or AI state) into parameters that affect movement. */
 	UFUNCTION(BlueprintNativeEvent)
 	void ProduceInput(int32 SimTimeMs, FMoverInputCmdContext& InputCmdResult);
+};
+
+
+/** 
+ * FMoverPredictTrajectoryParams: parameter block for querying future trajectory samples based on a starting state
+ * See UMoverComponent::GetPredictedTrajectory
+ */
+USTRUCT(BlueprintType)
+struct MOVER_API FMoverPredictTrajectoryParams
+{
+	GENERATED_BODY()
+
+	/** How many samples to predict into the future, including the first sample, which is always a snapshot of the
+	 *  starting state with 0 accumulated time. */
+	UPROPERTY(EditAnywhere, BlueprintReadWrite, Category = Mover, meta = (ClampMin = 1))
+	int32 NumPredictionSamples = 1;
+
+	/* How much time between predicted samples */
+	UPROPERTY(EditAnywhere, BlueprintReadWrite, Category = Mover, meta = (ClampMin = 0.00001))
+	float SecondsPerSample = 0.333f;
+
+	/** If true, samples are based on the visual component transform, rather than the 'updated' movement root.
+	 *  Typically, this is a mesh with its component location at the bottom of the collision primitive.
+	 *  If false, samples are from the movement root. */
+	UPROPERTY(EditAnywhere, BlueprintReadWrite, Category = Mover)
+	bool bUseVisualComponentRoot = false;
+
+	/** If true, gravity will not taken into account during prediction */
+	UPROPERTY(EditAnywhere, BlueprintReadWrite, Category = Mover)
+	bool bDisableGravity = false;
+
+ 	/** Optional starting sync state. If not set, prediction will begin from the current state. */
+ 	UPROPERTY(EditAnywhere, BlueprintReadWrite, Category = Mover)
+ 	TOptional<FMoverSyncState> OptionalStartSyncState;
+ 
+ 	/** Optional starting aux state. If not set, prediction will begin from the current state. */
+ 	UPROPERTY(EditAnywhere, BlueprintReadWrite, Category = Mover)
+ 	TOptional<FMoverAuxStateContext> OptionalStartAuxState;
+
+ 	/** Optional input cmds to use, one per sample. If none are specified, prediction will begin with last-used inputs. 
+ 	 *  If too few are specified for the number of samples, the final input in the array will be used repeatedly to cover remaining samples. */
+ 	UPROPERTY(EditAnywhere, BlueprintReadWrite, Category = Mover)
+ 	TArray<FMoverInputCmdContext> OptionalInputCmds;
+
 };
\ No newline at end of file
diff --git a/Engine/Plugins/Experimental/Mover/Source/Mover/Public/MoverTypes.h b/Engine/Plugins/Experimental/Mover/Source/Mover/Public/MoverTypes.h
--- a/Engine/Plugins/Experimental/Mover/Source/Mover/Public/MoverTypes.h
+++ b/Engine/Plugins/Experimental/Mover/Source/Mover/Public/MoverTypes.h
@@ -97,32 +97,6 @@ struct MOVER_API FMoverTimeStep
 };
 
 
-USTRUCT(BlueprintType)
-struct MOVER_API FMoverPredictTrajectoryParams
-{
-	GENERATED_BODY()
-
-	/** How many samples to predict into the future, including the first sample, which is always a snapshot of the
-	 *  starting state with 0 accumulated time. */
-	UPROPERTY(EditAnywhere, BlueprintReadWrite, Category = Mover, meta=(ClampMin=1))
-	int32 NumPredictionSamples = 1;
-
-	/* How much time between predicted samples */
-	UPROPERTY(EditAnywhere, BlueprintReadWrite, Category = Mover, meta=(ClampMin=0.00001))
-	float SecondsPerSample = 0.333f;
-
-	/** If true, samples are based on the visual component transform, rather than the 'updated' movement root. 
-	 *  Typically, this is a mesh with its component location at the bottom of the collision primitive.
-	 *  If false, samples are from the movement root. */
-	UPROPERTY(EditAnywhere, BlueprintReadWrite, Category = Mover)
-	bool bUseVisualComponentRoot = false;
-
-	/** If true, gravity will not taken into account during prediction */
-	UPROPERTY(EditAnywhere, BlueprintReadWrite, Category = Mover)
-	bool bDisableGravity = false;
-
-};
-
 
 // Base type for all data structs used to compose Mover simulation model definition dynamically (input cmd, sync state, aux state)
 // NOTE: for simulation state data (sync/aux), derive from FMoverStateData instead 
