Crouch | Physics Mover | Override OnPreSimulationTick in Physics Character Mover class

diff --git a/Engine/Plugins/Experimental/Mover/Source/Mover/Public/PhysicsMover/PhysicsCharacterMoverComponent.h b/Engine/Plugins/Experimental/Mover/Source/Mover/Public/PhysicsMover/PhysicsCharacterMoverComponent.h
--- a/Engine/Plugins/Experimental/Mover/Source/Mover/Public/PhysicsMover/PhysicsCharacterMoverComponent.h
+++ b/Engine/Plugins/Experimental/Mover/Source/Mover/Public/PhysicsMover/PhysicsCharacterMoverComponent.h
@@ -25,6 +25,8 @@ protected:
 	UFUNCTION()
 	virtual void OnMoverPostSimulationTick(const FMoverTimeStep& TimeStep);
 
+	virtual void OnMoverPreSimulationTick(const FMoverTimeStep& TimeStep, const FMoverInputCmdContext& InputCmd) override {};
+
 private:
 
 	// Keeps track of post processing that needs to take place on the game thread, after the stance is processed in the simulation tick.
