Crouch | Mover | FindMovementModifierByType expects exact child type as GetScriptStruct is overrided in the child type

diff --git a/Engine/Plugins/Experimental/Mover/Source/Mover/Private/PhysicsMover/PhysicsCharacterMoverComponent.cpp b/Engine/Plugins/Experimental/Mover/Source/Mover/Private/PhysicsMover/PhysicsCharacterMoverComponent.cpp
--- a/Engine/Plugins/Experimental/Mover/Source/Mover/Private/PhysicsMover/PhysicsCharacterMoverComponent.cpp
+++ b/Engine/Plugins/Experimental/Mover/Source/Mover/Private/PhysicsMover/PhysicsCharacterMoverComponent.cpp
@@ -39,7 +39,7 @@ void UPhysicsCharacterMoverComponent::OnMoverPreMovement(const FMoverTimeStep& T
 	// This is a fail safe in case our handle was bad - try finding the modifier by type if we can
 	if (!StanceModifier)
 	{
-		StanceModifier = FindMovementModifierByType<FStanceModifier>();
+		StanceModifier = FindMovementModifierByType<FPhysicsStanceModifier>();
 	}
 
 	EStanceMode OldActiveStance = EStanceMode::Invalid;
@@ -79,7 +79,7 @@ void UPhysicsCharacterMoverComponent::OnMoverPostSimulationTick(const FMoverTime
 	const FStanceModifier* StanceModifier = static_cast<const FStanceModifier*>(FindMovementModifier(StanceModifierHandle));
 	if (!StanceModifier)
 	{
-		StanceModifier = FindMovementModifierByType<FStanceModifier>();
+		StanceModifier = FindMovementModifierByType<FPhysicsStanceModifier>();
 	}
 
 	const bool bIsCrouching = HasGameplayTag(Mover_IsCrouching, true);
