Mover: improved utility function by skipping unnecessary conversion work when possible

#rnx

diff --git a/Engine/Plugins/Experimental/Mover/Source/Mover/Private/MoveLibrary/MovementUtils.cpp b/Engine/Plugins/Experimental/Mover/Source/Mover/Private/MoveLibrary/MovementUtils.cpp
--- a/Engine/Plugins/Experimental/Mover/Source/Mover/Private/MoveLibrary/MovementUtils.cpp
+++ b/Engine/Plugins/Experimental/Mover/Source/Mover/Private/MoveLibrary/MovementUtils.cpp
@@ -156,20 +156,18 @@ FVector UMovementUtils::ConstrainToPlane(const FVector& Vector, const FPlane& Mo
 
 FRotator UMovementUtils::ApplyGravityToOrientationIntent(const FRotator& IntendedOrientation, const FQuat& WorldToGravity, bool bStayVertical)
 {
-	FRotator NewOrientationIntent = IntendedOrientation;
-
-	FRotator GravityRelativeDesiredRotation = (WorldToGravity.Inverse() * NewOrientationIntent.Quaternion()).Rotator();
-	
-	if (bStayVertical)
+	if (!bStayVertical)
 	{
-		GravityRelativeDesiredRotation.Pitch = 0.f;
-		GravityRelativeDesiredRotation.Yaw = FRotator::NormalizeAxis(GravityRelativeDesiredRotation.Yaw);
-		GravityRelativeDesiredRotation.Roll = 0.f;
+		return IntendedOrientation;
 	}
-	
-	NewOrientationIntent = (WorldToGravity * GravityRelativeDesiredRotation.Quaternion()).Rotator();
 
-	return NewOrientationIntent;
+	FRotator GravityRelativeDesiredRotation = (WorldToGravity.Inverse() * IntendedOrientation.Quaternion()).Rotator();	// world space -> gravity-relative space
+	
+	GravityRelativeDesiredRotation.Pitch = 0.f;
+	GravityRelativeDesiredRotation.Yaw = FRotator::NormalizeAxis(GravityRelativeDesiredRotation.Yaw);
+	GravityRelativeDesiredRotation.Roll = 0.f;
+	
+	return (WorldToGravity * GravityRelativeDesiredRotation.Quaternion()).Rotator();	// gravity-relative space -> world space
 }
 
 FVector UMovementUtils::ComputeSlideDelta(const FMovingComponentSet& MovingComps, const FVector& Delta, const float PctOfDeltaToMove, const FVector& Normal, const FHitResult& Hit)
