Mover: fix crash at BP compile time with invalid / none transition objects in MoverComp's Transitions array or a movement mode's Transitions array.

diff --git a/Engine/Plugins/Experimental/Mover/Source/Mover/Private/MoverComponent.cpp b/Engine/Plugins/Experimental/Mover/Source/Mover/Private/MoverComponent.cpp
--- a/Engine/Plugins/Experimental/Mover/Source/Mover/Private/MoverComponent.cpp
+++ b/Engine/Plugins/Experimental/Mover/Source/Mover/Private/MoverComponent.cpp
@@ -1039,6 +1039,11 @@ bool UMoverComponent::ValidateSetup(FDataValidationContext& Context) const
 
 			for (const UBaseMovementModeTransition* Transition : Element.Value->Transitions)
 			{
+				if (!IsValid(Transition))
+				{
+					continue;
+				}
+
 				for (const TSubclassOf<UObject>& Type : Transition->SharedSettingsClasses)
 				{
 					if (Type.Get() == nullptr)
@@ -1082,6 +1087,7 @@ bool UMoverComponent::ValidateSetup(FDataValidationContext& Context) const
 				FText::FromString(GetNameSafe(GetOwner()))));
 
 			bDidFindAnyProblems = true;
+			continue;
 		}
 
 		for (const TSubclassOf<UObject>& Type : Transition->SharedSettingsClasses)
@@ -1209,6 +1215,11 @@ void UMoverComponent::RefreshSharedSettings()
 
 			for (const UBaseMovementModeTransition* Transition : Mode->Transitions)
 			{
+				if (!IsValid(Transition))
+				{
+					continue;
+				}
+
 				for (const TSubclassOf<UObject>& SharedSettingsType : Transition->SharedSettingsClasses)
 				{
 					if (SharedSettingsType.Get() == nullptr)
@@ -1240,6 +1251,11 @@ void UMoverComponent::RefreshSharedSettings()
 
 	for (const UBaseMovementModeTransition* Transition : Transitions)
 	{
+		if (!IsValid(Transition))
+		{
+			continue;
+		}
+
 		for (const TSubclassOf<UObject>& SharedSettingsType : Transition->SharedSettingsClasses)
 		{
 			if (SharedSettingsType.Get() == nullptr)
