Mover / MoverExamples: minor fixes for spelling and property access.

#rnx

diff --git a/Engine/Plugins/Experimental/Mover/Source/Mover/Private/MovementModeStateMachine.cpp b/Engine/Plugins/Experimental/Mover/Source/Mover/Private/MovementModeStateMachine.cpp
--- a/Engine/Plugins/Experimental/Mover/Source/Mover/Private/MovementModeStateMachine.cpp
+++ b/Engine/Plugins/Experimental/Mover/Source/Mover/Private/MovementModeStateMachine.cpp
@@ -116,7 +116,7 @@ void UMovementModeStateMachine::QueueNextMode(FName DesiredNextModeName, bool bS
 	if (DesiredNextModeName != NAME_None)
 	{ 
 		const FName NextModeName          = QueuedModeTransition->GetNextModeName();
-		const bool bShouldNextModeReenter = QueuedModeTransition->ShoulReenter();
+		const bool bShouldNextModeReenter = QueuedModeTransition->ShouldReenter();
 
 		if ((NextModeName != NAME_None) &&
 		    (NextModeName != DesiredNextModeName || bShouldReenter != bShouldNextModeReenter))
@@ -489,7 +489,7 @@ void UMovementModeStateMachine::ConstructDefaultModes()
 void UMovementModeStateMachine::AdvanceToNextMode()
 {
 	const FName NextModeName = QueuedModeTransition->GetNextModeName();
-	const bool bShouldNextModeReenter = QueuedModeTransition->ShoulReenter();
+	const bool bShouldNextModeReenter = QueuedModeTransition->ShouldReenter();
 
 	if ((NextModeName != NAME_None && Modes.Contains(NextModeName)) &&
 		(CurrentModeName != NextModeName || bShouldNextModeReenter))
diff --git a/Engine/Plugins/Experimental/Mover/Source/Mover/Public/MovementModeStateMachine.h b/Engine/Plugins/Experimental/Mover/Source/Mover/Public/MovementModeStateMachine.h
--- a/Engine/Plugins/Experimental/Mover/Source/Mover/Public/MovementModeStateMachine.h
+++ b/Engine/Plugins/Experimental/Mover/Source/Mover/Public/MovementModeStateMachine.h
@@ -132,7 +132,7 @@ private:
 
 
 	 FName GetNextModeName() const { return NextMode; }
-	 bool ShoulReenter() const { return bShouldNextModeReenter; }
+	 bool ShouldReenter() const { return bShouldNextModeReenter; }
 
  private:
 	 FName NextMode = NAME_None;
diff --git a/Engine/Plugins/Experimental/MoverExamples/Source/MoverExamples/Public/MoverExamplesCharacter.h b/Engine/Plugins/Experimental/MoverExamples/Source/MoverExamples/Public/MoverExamplesCharacter.h
--- a/Engine/Plugins/Experimental/MoverExamples/Source/MoverExamples/Public/MoverExamplesCharacter.h
+++ b/Engine/Plugins/Experimental/MoverExamples/Source/MoverExamples/Public/MoverExamplesCharacter.h
@@ -82,8 +82,8 @@ protected:
 	TObjectPtr<UInputAction> FlyInputAction;
 
 public:
-	// Whether or not we author our movement inputs relative to whatever base we're standing on, or leave them in world space
-	UPROPERTY(BlueprintReadWrite, Category=MoverExamples)
+	// Whether or not we author our movement inputs relative to whatever base we're standing on, or leave them in world space. Only applies if standing on a base of some sort.
+	UPROPERTY(EditAnywhere, BlueprintReadWrite, Category=MoverExamples)
 	bool bUseBaseRelativeMovement = true;
 	
 	/**
