Fix for PBCM MaxUnsupportedTimeBeforeFalling not working if the floor distance is larger than step height.
#rnx

diff --git a/Engine/Plugins/Experimental/Mover/Source/Mover/Private/PhysicsMover/Modes/PhysicsDrivenWalkingMode.cpp b/Engine/Plugins/Experimental/Mover/Source/Mover/Private/PhysicsMover/Modes/PhysicsDrivenWalkingMode.cpp
--- a/Engine/Plugins/Experimental/Mover/Source/Mover/Private/PhysicsMover/Modes/PhysicsDrivenWalkingMode.cpp
+++ b/Engine/Plugins/Experimental/Mover/Source/Mover/Private/PhysicsMover/Modes/PhysicsDrivenWalkingMode.cpp
@@ -414,7 +414,14 @@ void UPhysicsDrivenWalkingMode::OnSimulationTick(const FSimulationTickParams& Pa
 	else
 	{
 		// No water or floor not found
-		SwitchToState(CommonLegacySettings->AirMovementModeName, Params, OutputState);
+		float TimeSinceSupported = MaxUnsupportedTimeBeforeFalling;
+		SimBlackboard->TryGet(CommonBlackboard::TimeSinceSupported, TimeSinceSupported);
+		TimeSinceSupported += DeltaSeconds;
+		SimBlackboard->Set(CommonBlackboard::TimeSinceSupported, TimeSinceSupported);
+		if (TimeSinceSupported >= MaxUnsupportedTimeBeforeFalling)
+		{
+			SwitchToState(CommonLegacySettings->AirMovementModeName, Params, OutputState);
+		}
 	}
 }
 
