[Mover] [BugFix] Additive Mixing of Layered moves should not be overriden by [REDACTED_PERSON].

diff --git a/Engine/Plugins/Experimental/Mover/Source/Mover/Private/MoveLibrary/MovementMixer.cpp b/Engine/Plugins/Experimental/Mover/Source/Mover/Private/MoveLibrary/MovementMixer.cpp
--- a/Engine/Plugins/Experimental/Mover/Source/Mover/Private/MoveLibrary/MovementMixer.cpp
+++ b/Engine/Plugins/Experimental/Mover/Source/Mover/Private/MoveLibrary/MovementMixer.cpp
@@ -57,7 +57,6 @@ void UMovementMixer::MixLayeredMove(const FLayeredMoveBase& ActiveMove, const FP
 				OutCumulativeMove.PreferredMode = MoveStep.PreferredMode;
 			}
 
-			OutCumulativeMove.PreferredMode = MoveStep.PreferredMode;
 			OutCumulativeMove.LinearVelocity += MoveStep.LinearVelocity;
 			OutCumulativeMove.AngularVelocity += MoveStep.AngularVelocity;
 		}
