Fix Server Crash

diff --git a/Engine/Plugins/Experimental/Mover/Source/Mover/Private/Backends/MoverNetworkPhysicsLiaison.cpp b/Engine/Plugins/Experimental/Mover/Source/Mover/Private/Backends/MoverNetworkPhysicsLiaison.cpp
--- a/Engine/Plugins/Experimental/Mover/Source/Mover/Private/Backends/MoverNetworkPhysicsLiaison.cpp
+++ b/Engine/Plugins/Experimental/Mover/Source/Mover/Private/Backends/MoverNetworkPhysicsLiaison.cpp
@@ -1107,6 +1107,8 @@ void UMoverNetworkPhysicsLiaisonComponent::OnPreSimulate_Internal(const FPhysics
 			Output.FloorResult.FloorDist,
 			WalkableSlopeCosine
 			});
+
+		ConstraintHandle->SetGroundParticle(nullptr);
 	}
 	else
 	{
