Mover: fixed crash when an actor fails to spawn, caused by the NPP proxy not being uninitialized

diff --git a/Engine/Plugins/Experimental/Mover/Source/Mover/Private/Backends/MoverNetworkPredictionLiaison.cpp b/Engine/Plugins/Experimental/Mover/Source/Mover/Private/Backends/MoverNetworkPredictionLiaison.cpp
--- a/Engine/Plugins/Experimental/Mover/Source/Mover/Private/Backends/MoverNetworkPredictionLiaison.cpp
+++ b/Engine/Plugins/Experimental/Mover/Source/Mover/Private/Backends/MoverNetworkPredictionLiaison.cpp
@@ -227,6 +227,14 @@ void UMoverNetworkPredictionLiaisonComponent::InitializeComponent()
 	Super::InitializeComponent();
 }
 
+
+void UMoverNetworkPredictionLiaisonComponent::UninitializeComponent()
+{
+	NetworkPredictionProxy.EndPlay();
+
+	Super::UninitializeComponent();
+}
+
 void UMoverNetworkPredictionLiaisonComponent::OnRegister()
 {
 	Super::OnRegister();
diff --git a/Engine/Plugins/Experimental/Mover/Source/Mover/Public/Backends/MoverNetworkPredictionLiaison.h b/Engine/Plugins/Experimental/Mover/Source/Mover/Public/Backends/MoverNetworkPredictionLiaison.h
--- a/Engine/Plugins/Experimental/Mover/Source/Mover/Public/Backends/MoverNetworkPredictionLiaison.h
+++ b/Engine/Plugins/Experimental/Mover/Source/Mover/Public/Backends/MoverNetworkPredictionLiaison.h
@@ -65,6 +65,7 @@ public:
 
 	// UObject interface
 	void InitializeComponent() override;
+	void UninitializeComponent() override;
 	void OnRegister() override;
 	void RegisterComponentTickFunctions(bool bRegister) override;
 	// End UObject interface
