Mover: fixed potential crash at EndPlay time for async-based movement
#rnx

diff --git a/Engine/Plugins/Experimental/Mover/Source/Mover/Private/MoverDataModelTypes.cpp b/Engine/Plugins/Experimental/Mover/Source/Mover/Private/MoverDataModelTypes.cpp
--- a/Engine/Plugins/Experimental/Mover/Source/Mover/Private/MoverDataModelTypes.cpp
+++ b/Engine/Plugins/Experimental/Mover/Source/Mover/Private/MoverDataModelTypes.cpp
@@ -132,7 +132,7 @@ bool FMoverDefaultSyncState::NetSerialize(FArchive& Ar, UPackageMap* Map, bool&
 	Orientation.SerializeCompressedShort(Ar);
 
 	// Optional movement base
-	bool bIsUsingMovementBase = (Ar.IsSaving() ? (MovementBase != nullptr) : false);
+	bool bIsUsingMovementBase = (Ar.IsSaving() ? MovementBase.IsValid() : false);
 	Ar.SerializeBits(&bIsUsingMovementBase, 1);
 
 	if (bIsUsingMovementBase)
@@ -161,9 +161,9 @@ void FMoverDefaultSyncState::ToString(FAnsiStringBuilderBase& Out) const
 	Out.Appendf("Vel: X=%.2f Y=%.2f Z=%.2f\n", Velocity.X, Velocity.Y, Velocity.Z);
 	Out.Appendf("Orient: P=%.2f Y=%.2f R=%.2f\n", Orientation.Pitch, Orientation.Yaw, Orientation.Roll);
 
-	if (MovementBase)
+	if (const UPrimitiveComponent* MovementBasePtr = MovementBase.Get())
 	{
-		Out.Appendf("MovementBase: %s (bone %s)\n", TCHAR_TO_ANSI(*GetNameSafe(MovementBase->GetOwner())), TCHAR_TO_ANSI(*MovementBaseBoneName.ToString()));
+		Out.Appendf("MovementBase: %s (bone %s)\n", TCHAR_TO_ANSI(*GetNameSafe(MovementBasePtr->GetOwner())), TCHAR_TO_ANSI(*MovementBaseBoneName.ToString()));
 		Out.Appendf("    BasePos: %s   BaseRot: %s\n", TCHAR_TO_ANSI(*MovementBasePos.ToCompactString()), TCHAR_TO_ANSI(*MovementBaseQuat.Rotator().ToCompactString()));
 	}
 	else
@@ -179,13 +179,13 @@ bool FMoverDefaultSyncState::ShouldReconcile(const FMoverDataStructBase& Authori
 	const FMoverDefaultSyncState* AuthoritySyncState = static_cast<const FMoverDefaultSyncState*>(&AuthorityState);
 	const float DistErrorTolerance = 5.f;	// JAH TODO: define these elsewhere as CVars or data asset settings
 
-	const bool bAreInDifferentSpaces = !((MovementBase == AuthoritySyncState->MovementBase) && (MovementBaseBoneName == AuthoritySyncState->MovementBaseBoneName));
+	const bool bAreInDifferentSpaces = !((MovementBase.HasSameIndexAndSerialNumber(AuthoritySyncState->MovementBase)) && (MovementBaseBoneName == AuthoritySyncState->MovementBaseBoneName));
 
 	bool bIsNearEnough = false;
 
 	if (!bAreInDifferentSpaces)
 	{
-		if (MovementBase)
+		if (MovementBase.IsValid())
 		{
 			bIsNearEnough = GetLocation_BaseSpace().Equals(AuthoritySyncState->GetLocation_BaseSpace(), DistErrorTolerance);
 		}
@@ -224,9 +224,9 @@ void FMoverDefaultSyncState::Interpolate(const FMoverDataStructBase& From, const
 
 	
 		// Bases match (or not using based movement at all)
-		if (FromState->MovementBase == ToState->MovementBase && FromState->MovementBaseBoneName == ToState->MovementBaseBoneName)
+		if (FromState->MovementBase.HasSameIndexAndSerialNumber(ToState->MovementBase) && FromState->MovementBaseBoneName == ToState->MovementBaseBoneName)
 		{
-			if (FromState->MovementBase)
+			if (FromState->MovementBase.IsValid())
 			{
 				MovementBasePos = FMath::Lerp(FromState->MovementBasePos, ToState->MovementBasePos, Pct);
 				MovementBaseQuat = FQuat::Slerp(FromState->MovementBaseQuat, ToState->MovementBaseQuat, Pct);
@@ -237,14 +237,14 @@ void FMoverDefaultSyncState::Interpolate(const FMoverDataStructBase& From, const
 			FromVelocity_ToSpace    = FromState->Velocity;
 			FromOrientation_ToSpace = FromState->Orientation;
 		}
-		else if (ToState->MovementBase)	// if moving onto a different base, regardless of coming from a prior base or not
+		else if (ToState->MovementBase.IsValid())	// if moving onto a different base, regardless of coming from a prior base or not
 		{
 			UBasedMovementUtils::TransformLocationToLocal(ToState->MovementBasePos, ToState->MovementBaseQuat, FromState->GetLocation_WorldSpace(), OUT FromLocation_ToSpace);
 			UBasedMovementUtils::TransformDirectionToLocal(ToState->MovementBaseQuat, FromState->GetIntent_WorldSpace(), OUT FromIntent_ToSpace);
 			UBasedMovementUtils::TransformDirectionToLocal(ToState->MovementBaseQuat, FromState->GetVelocity_WorldSpace(), OUT FromVelocity_ToSpace);
 			UBasedMovementUtils::TransformRotatorToLocal(ToState->MovementBaseQuat, FromState->GetOrientation_WorldSpace(), OUT FromOrientation_ToSpace);
 		}
-		else if (FromState->MovementBase)	// if leaving a base
+		else if (FromState->MovementBase.IsValid())	// if leaving a base
 		{
 			FromLocation_ToSpace	= FromState->GetLocation_WorldSpace();
 			FromIntent_ToSpace		= FromState->GetIntent_WorldSpace();
@@ -296,9 +296,9 @@ bool FMoverDefaultSyncState::UpdateCurrentMovementBase()
 {
 	bool bDidGetBaseTransform = false;
 
-	if (MovementBase)
+	if (const UPrimitiveComponent* MovementBasePtr = MovementBase.Get())
 	{
-		bDidGetBaseTransform = UBasedMovementUtils::GetMovementBaseTransform(MovementBase, MovementBaseBoneName, OUT MovementBasePos, OUT MovementBaseQuat);
+		bDidGetBaseTransform = UBasedMovementUtils::GetMovementBaseTransform(MovementBasePtr, MovementBaseBoneName, OUT MovementBasePos, OUT MovementBaseQuat);
 	}
 
 	if (!bDidGetBaseTransform)
@@ -314,7 +314,7 @@ bool FMoverDefaultSyncState::UpdateCurrentMovementBase()
 
 FVector FMoverDefaultSyncState::GetLocation_WorldSpace() const
 {
-	if (MovementBase)
+	if (MovementBase.IsValid())
 	{
 		return FTransform(MovementBaseQuat, MovementBasePos).TransformPositionNoScale(Location);
 	}
@@ -330,7 +330,7 @@ FVector FMoverDefaultSyncState::GetLocation_BaseSpace() const
 
 FVector FMoverDefaultSyncState::GetIntent_WorldSpace() const
 {
-	if (MovementBase)
+	if (MovementBase.IsValid())
 	{
 		return MovementBaseQuat.RotateVector(MoveDirectionIntent);
 	}
@@ -345,7 +345,7 @@ FVector FMoverDefaultSyncState::GetIntent_BaseSpace() const
 
 FVector FMoverDefaultSyncState::GetVelocity_WorldSpace() const
 {
-	if (MovementBase)
+	if (MovementBase.IsValid())
 	{
 		return MovementBaseQuat.RotateVector(Velocity);
 	}
@@ -361,7 +361,7 @@ FVector FMoverDefaultSyncState::GetVelocity_BaseSpace() const
 
 FRotator FMoverDefaultSyncState::GetOrientation_WorldSpace() const
 {
-	if (MovementBase)
+	if (MovementBase.IsValid())
 	{
 		return (MovementBaseQuat * FQuat(Orientation)).Rotator();
 	}
@@ -377,7 +377,7 @@ FRotator FMoverDefaultSyncState::GetOrientation_BaseSpace() const
 
 FTransform FMoverDefaultSyncState::GetTransform_WorldSpace() const
 {
-	if (MovementBase)
+	if (MovementBase.IsValid())
 	{
 		return FTransform(Orientation, Location) * FTransform(MovementBaseQuat, MovementBasePos);
 	}
diff --git a/Engine/Plugins/Experimental/Mover/Source/Mover/Public/MoverDataModelTypes.h b/Engine/Plugins/Experimental/Mover/Source/Mover/Public/MoverDataModelTypes.h
--- a/Engine/Plugins/Experimental/Mover/Source/Mover/Public/MoverDataModelTypes.h
+++ b/Engine/Plugins/Experimental/Mover/Source/Mover/Public/MoverDataModelTypes.h
@@ -147,7 +147,7 @@ public:
 protected:
 	// Optional: when moving on a base, input may be relative to this object
 	UPROPERTY(BlueprintReadOnly, Category = Mover)
-	TObjectPtr<UPrimitiveComponent> MovementBase;
+	TWeakObjectPtr<UPrimitiveComponent> MovementBase;
 
 	// Optional: for movement bases that are skeletal meshes, this is the bone we're based on. Only valid if MovementBase is set.
 	UPROPERTY(BlueprintReadOnly, Category = Mover)
@@ -197,7 +197,7 @@ public:
 	bool UpdateCurrentMovementBase();
 
 	// Queries
-	UPrimitiveComponent* GetMovementBase() const { return MovementBase; }
+	UPrimitiveComponent* GetMovementBase() const { return MovementBase.Get(); }
 	FName GetMovementBaseBoneName() const { return MovementBaseBoneName; }
 	FVector GetCapturedMovementBasePos() const { return MovementBasePos; }
 	FQuat GetCapturedMovementBaseQuat() const { return MovementBaseQuat; }
