Mover: adding navwalking mode fallback option for cases where we can't rely on the CDO to find the original collision settings

diff --git a/Engine/Plugins/Experimental/Mover/Source/Mover/Private/DefaultMovementSet/Modes/NavWalkingMode.cpp b/Engine/Plugins/Experimental/Mover/Source/Mover/Private/DefaultMovementSet/Modes/NavWalkingMode.cpp
--- a/Engine/Plugins/Experimental/Mover/Source/Mover/Private/DefaultMovementSet/Modes/NavWalkingMode.cpp
+++ b/Engine/Plugins/Experimental/Mover/Source/Mover/Private/DefaultMovementSet/Modes/NavWalkingMode.cpp
@@ -297,17 +297,20 @@ void UNavWalkingMode::SetCollisionForNavWalking(bool bEnable)
 		{
 			if (bEnable)
 			{
+				CollideVsWorldStatic = UpdatedCompAsPrimitive->GetCollisionResponseToChannel(ECC_WorldStatic);
+				CollideVsWorldDynamic = UpdatedCompAsPrimitive->GetCollisionResponseToChannel(ECC_WorldDynamic);
+
 				// TODO NS: Eventually might want make the new collision response an option so they trigger certain things while still not colliding 
 				UpdatedCompAsPrimitive->SetCollisionResponseToChannel(ECC_WorldStatic, ECR_Ignore);
 				UpdatedCompAsPrimitive->SetCollisionResponseToChannel(ECC_WorldDynamic, ECR_Ignore);
-			
+
 				if (UNavWalkingMode* NavWalkingMode = Cast<UNavWalkingMode>(MoverComponent->FindMovementMode(UNavWalkingMode::StaticClass())))
 				{
 					if (UMoverBlackboard* SimBlackboard = MoverComponent->GetSimBlackboard_Mutable())
 					{
 						SimBlackboard->Invalidate(NavWalkingBlackBoard::ProjectedNavMeshHitResult);
 					}
-					
+
 					// Stagger timed updates so many different characters spawned at the same time don't update on the same frame.
 					// Initially we want an immediate update though, so set time to a negative randomized range.
 					NavWalkingMode->NavMeshProjectionTimer = (NavWalkingMode->NavMeshProjectionInterval > 0.f) ? FMath::FRandRange(-NavWalkingMode->NavMeshProjectionInterval, 0.f) : 0.f;
@@ -318,9 +321,12 @@ void UNavWalkingMode::SetCollisionForNavWalking(bool bEnable)
 				// Grabbing the original shape settings and reverting to our old collision responses
 				if (const UShapeComponent* OriginalShapeComp = UMovementUtils::GetOriginalComponentType<UShapeComponent>(MoverComponent->GetOwner()))
 				{
-					UpdatedCompAsPrimitive->SetCollisionResponseToChannel(ECC_WorldStatic, OriginalShapeComp->GetCollisionResponseToChannel(ECC_WorldStatic));
-					UpdatedCompAsPrimitive->SetCollisionResponseToChannel(ECC_WorldDynamic, OriginalShapeComp->GetCollisionResponseToChannel(ECC_WorldDynamic));
+					CollideVsWorldStatic = OriginalShapeComp->GetCollisionResponseToChannel(ECC_WorldStatic);
+					CollideVsWorldDynamic = OriginalShapeComp->GetCollisionResponseToChannel(ECC_WorldDynamic);
 				}
+
+				UpdatedCompAsPrimitive->SetCollisionResponseToChannel(ECC_WorldStatic, CollideVsWorldStatic);
+				UpdatedCompAsPrimitive->SetCollisionResponseToChannel(ECC_WorldDynamic, CollideVsWorldDynamic);
 			}
 		}
 	}
diff --git a/Engine/Plugins/Experimental/Mover/Source/Mover/Public/DefaultMovementSet/Modes/NavWalkingMode.h b/Engine/Plugins/Experimental/Mover/Source/Mover/Public/DefaultMovementSet/Modes/NavWalkingMode.h
--- a/Engine/Plugins/Experimental/Mover/Source/Mover/Public/DefaultMovementSet/Modes/NavWalkingMode.h
+++ b/Engine/Plugins/Experimental/Mover/Source/Mover/Public/DefaultMovementSet/Modes/NavWalkingMode.h
@@ -3,6 +3,7 @@
 #pragma once
 
 #include "CoreMinimal.h"
+#include "Engine/EngineTypes.h"
 #include "MovementMode.h"
 #include "MoveLibrary/BasedMovementUtils.h"
 #include "AI/Navigation/NavigationTypes.h"
@@ -132,4 +133,8 @@ protected:
 	FRelativeBaseInfo UpdateFloorAndBaseInfo(const FFloorCheckResult& FloorResult) const;
 	
 	TObjectPtr<const UCommonLegacyMovementSettings> CommonLegacySettings;
+
+	// Collision responses saved at the activation of this mode, to be restored as upon deactivation if the CDO can't be used
+	ECollisionResponse CollideVsWorldStatic = ECollisionResponse::ECR_Block;
+	ECollisionResponse CollideVsWorldDynamic = ECollisionResponse::ECR_Block;
 };
