Add support for ShouldRemainUpright in MoverNetworkPhysicsLiason
#rnx

diff --git a/Engine/Plugins/Experimental/Mover/Source/Mover/Private/Backends/MoverNetworkPhysicsLiaison.cpp b/Engine/Plugins/Experimental/Mover/Source/Mover/Private/Backends/MoverNetworkPhysicsLiaison.cpp
--- a/Engine/Plugins/Experimental/Mover/Source/Mover/Private/Backends/MoverNetworkPhysicsLiaison.cpp
+++ b/Engine/Plugins/Experimental/Mover/Source/Mover/Private/Backends/MoverNetworkPhysicsLiaison.cpp
@@ -1061,7 +1061,14 @@ void UMoverNetworkPhysicsLiaisonComponent::OnPreSimulate_Internal(const FPhysics
 	float TargetDeltaFacing = FMath::DegreesToRadians(Remainder.Yaw);
 	if (TickParams.DeltaTimeSeconds > UE_SMALL_NUMBER)
 	{
-		CharacterParticle->SetW((TargetDeltaFacing / TickParams.DeltaTimeSeconds) * Chaos::FVec3::ZAxisVector);
+		if (CommonMovementSettings->bShouldRemainVertical)
+		{
+			CharacterParticle->SetW((TargetDeltaFacing / TickParams.DeltaTimeSeconds) * Chaos::FVec3::ZAxisVector);
+		}
+		else
+		{
+			CharacterParticle->SetW(Remainder.Quaternion().ToRotationVector() / TickParams.DeltaTimeSeconds);
+		}
 	}
 
 	// Update the constraint data based on the floor result
