Mover: fixing cases where nav agent properties were never getting initially set, leading to a failure to find navigation areas in NavWalking mode

diff --git a/Engine/Plugins/Experimental/Mover/Source/Mover/Private/DefaultMovementSet/NavMoverComponent.cpp b/Engine/Plugins/Experimental/Mover/Source/Mover/Private/DefaultMovementSet/NavMoverComponent.cpp
--- a/Engine/Plugins/Experimental/Mover/Source/Mover/Private/DefaultMovementSet/NavMoverComponent.cpp
+++ b/Engine/Plugins/Experimental/Mover/Source/Mover/Private/DefaultMovementSet/NavMoverComponent.cpp
@@ -30,6 +30,20 @@ void UNavMoverComponent::InitializeComponent()
 	}
 }
 
+void UNavMoverComponent::BeginPlay()
+{
+	Super::BeginPlay();
+
+	if (MoverComponent && MoverComponent->GetUpdatedComponent())
+	{
+		UpdateNavAgent(*MoverComponent->GetUpdatedComponent());
+	}
+	else
+	{
+		UpdateNavAgent(*GetOwner());
+	}
+}
+
 float UNavMoverComponent::GetMaxSpeedForNavMovement() const
 {
 	const UCommonLegacyMovementSettings* MovementSettings = MoverComponent->FindSharedSettings_Mutable<UCommonLegacyMovementSettings>();
diff --git a/Engine/Plugins/Experimental/Mover/Source/Mover/Public/DefaultMovementSet/NavMoverComponent.h b/Engine/Plugins/Experimental/Mover/Source/Mover/Public/DefaultMovementSet/NavMoverComponent.h
--- a/Engine/Plugins/Experimental/Mover/Source/Mover/Public/DefaultMovementSet/NavMoverComponent.h
+++ b/Engine/Plugins/Experimental/Mover/Source/Mover/Public/DefaultMovementSet/NavMoverComponent.h
@@ -60,6 +60,7 @@ private:
 	
 public:
 	virtual void InitializeComponent() override;
+	virtual void BeginPlay() override;
 
 	/** Get the owner of the object consuming nav movement */
 	virtual UObject* GetOwnerAsObject() const override { return GetOwner(); }
