Fix for rotation not being converted to radians in MoverNetworkPhysicsLiason
#rnx

diff --git a/Engine/Plugins/Experimental/Mover/Source/Mover/Private/Backends/MoverNetworkPhysicsLiaison.cpp b/Engine/Plugins/Experimental/Mover/Source/Mover/Private/Backends/MoverNetworkPhysicsLiaison.cpp
--- a/Engine/Plugins/Experimental/Mover/Source/Mover/Private/Backends/MoverNetworkPhysicsLiaison.cpp
+++ b/Engine/Plugins/Experimental/Mover/Source/Mover/Private/Backends/MoverNetworkPhysicsLiaison.cpp
@@ -1067,7 +1067,7 @@ void UMoverNetworkPhysicsLiaisonComponent::OnPreSimulate_Internal(const FPhysics
 		}
 		else
 		{
-			CharacterParticle->SetW(Remainder.Quaternion().ToRotationVector() / TickParams.DeltaTimeSeconds);
+			CharacterParticle->SetW(FMath::DegreesToRadians(Remainder).Quaternion().ToRotationVector() / TickParams.DeltaTimeSeconds);
 		}
 	}
 
