[Mover] [BugFix]
- UpdateWaterSplineData() is now moved behind a CVar which is disabled by default.
- Other pointer access checks.

diff --git a/Engine/Plugins/Experimental/Mover/Source/Mover/Private/MoveLibrary/WaterMovementUtils.cpp b/Engine/Plugins/Experimental/Mover/Source/Mover/Private/MoveLibrary/WaterMovementUtils.cpp
--- a/Engine/Plugins/Experimental/Mover/Source/Mover/Private/MoveLibrary/WaterMovementUtils.cpp
+++ b/Engine/Plugins/Experimental/Mover/Source/Mover/Private/MoveLibrary/WaterMovementUtils.cpp
@@ -32,53 +32,59 @@ FProposedMove UWaterMovementUtils::ComputeControlledWaterMove(const FWaterMovePa
 
 void UWaterMovementUtils::UpdateWaterSplineData(const FUpdateWaterSplineDataParams& UpdateWaterSplineDataParams, FWaterCheckResult& OutWaterResult)
 {
-	if (AWaterBody* WaterBody = Cast<AWaterBody>(OutWaterResult.HitResult.Component->GetOwner()))
+	if(OutWaterResult.HitResult.Component.IsValid())
 	{
-		FWaterBodyQueryResult QueryResult = WaterBody->GetWaterBodyComponent()->QueryWaterInfoClosestToWorldLocation(UpdateWaterSplineDataParams.PlayerLocation,
-				   EWaterBodyQueryFlags::ComputeLocation
-				   | EWaterBodyQueryFlags::ComputeDepth
-				   | EWaterBodyQueryFlags::ComputeNormal
-				   | EWaterBodyQueryFlags::ComputeVelocity
-				   | EWaterBodyQueryFlags::ComputeImmersionDepth);
-
-		FWaterFlowSplineData& WaterSplineData = OutWaterResult.WaterSplineData;
-
-		// Immersion depth
-		WaterSplineData.ImmersionDepth = QueryResult.GetImmersionDepth();
-
-		WaterSplineData.WaterPlaneLocation = QueryResult.GetWaterPlaneLocation();
-		WaterSplineData.WaterPlaneNormal = QueryResult.GetWaterPlaneNormal();
-
-		// Water Depth
-		WaterSplineData.WaterDepth = QueryResult.GetWaterSurfaceDepth();
-
-		// Water velocity
-		WaterSplineData.RawWaterVelocity = QueryResult.GetVelocity();
-
-		// Water velocity modified by depth.
-		const FVector2D DepthRange = FVector2D(UpdateWaterSplineDataParams.TargetImmersionDepth, UpdateWaterSplineDataParams.WaterVelocityDepthForMax);
-		const FVector2D VelocityMultiplierRange = FVector2D(UpdateWaterSplineDataParams.WaterVelocityMinMultiplier, 1.0f);
-		WaterSplineData.WaterVelocityDepthMultiplier = FMath::GetMappedRangeValueClamped(DepthRange, VelocityMultiplierRange, WaterSplineData.WaterDepth);
-		WaterSplineData.WaterVelocity = WaterSplineData.RawWaterVelocity * WaterSplineData.WaterVelocityDepthMultiplier;
-
-		// Player velocity relative to water velocity.
-		const FVector PlayerVelocity2D = FVector(UpdateWaterSplineDataParams.PlayerVelocity.X, UpdateWaterSplineDataParams.PlayerVelocity.Y, 0.f);
-		const FVector WaterVelocity2D = FVector(WaterSplineData.WaterVelocity.X, WaterSplineData.WaterVelocity.Y, 0.f);
-		const FVector WaterVelocityProjection = PlayerVelocity2D.IsNearlyZero(0.1f) ? WaterVelocity2D : WaterVelocity2D.ProjectOnToNormal(PlayerVelocity2D.GetSafeNormal2D());
-		WaterSplineData.PlayerRelativeVelocityToWater = PlayerVelocity2D - WaterVelocityProjection;
-
-		// Surface location
-		const float CapsuleBottom = UpdateWaterSplineDataParams.PlayerLocation.Z - UpdateWaterSplineDataParams.CapsuleHalfHeight;
-		const float CapsuleTop = UpdateWaterSplineDataParams.PlayerLocation.Z + UpdateWaterSplineDataParams.CapsuleHalfHeight;
-		WaterSplineData.WaterSurfaceLocation = QueryResult.GetWaterSurfaceLocation();
-
-		// Surface offset
-		WaterSplineData.WaterSurfaceOffset = WaterSplineData.WaterSurfaceLocation - UpdateWaterSplineDataParams.PlayerLocation;
-
-		// Immersion Percent
-		WaterSplineData.ImmersionPercent = FMath::Clamp((WaterSplineData.WaterSurfaceLocation.Z - CapsuleBottom) / (CapsuleTop - CapsuleBottom), 0.f, 1.f);
-
-		// Normal
-		WaterSplineData.WaterSurfaceNormal = QueryResult.GetWaterSurfaceNormal();
+		if (AWaterBody* WaterBody = Cast<AWaterBody>(OutWaterResult.HitResult.Component->GetOwner()))
+		{
+			if (WaterBody->GetWaterBodyComponent())
+			{
+				FWaterBodyQueryResult QueryResult = WaterBody->GetWaterBodyComponent()->QueryWaterInfoClosestToWorldLocation(UpdateWaterSplineDataParams.PlayerLocation,
+						   EWaterBodyQueryFlags::ComputeLocation
+						   | EWaterBodyQueryFlags::ComputeDepth
+						   | EWaterBodyQueryFlags::ComputeNormal
+						   | EWaterBodyQueryFlags::ComputeVelocity
+						   | EWaterBodyQueryFlags::ComputeImmersionDepth);
+		
+				FWaterFlowSplineData& WaterSplineData = OutWaterResult.WaterSplineData;
+		
+				// Immersion depth
+				WaterSplineData.ImmersionDepth = QueryResult.GetImmersionDepth();
+		
+				WaterSplineData.WaterPlaneLocation = QueryResult.GetWaterPlaneLocation();
+				WaterSplineData.WaterPlaneNormal = QueryResult.GetWaterPlaneNormal();
+		
+				// Water Depth
+				WaterSplineData.WaterDepth = QueryResult.GetWaterSurfaceDepth();
+		
+				// Water velocity
+				WaterSplineData.RawWaterVelocity = QueryResult.GetVelocity();
+		
+				// Water velocity modified by depth.
+				const FVector2D DepthRange = FVector2D(UpdateWaterSplineDataParams.TargetImmersionDepth, UpdateWaterSplineDataParams.WaterVelocityDepthForMax);
+				const FVector2D VelocityMultiplierRange = FVector2D(UpdateWaterSplineDataParams.WaterVelocityMinMultiplier, 1.0f);
+				WaterSplineData.WaterVelocityDepthMultiplier = FMath::GetMappedRangeValueClamped(DepthRange, VelocityMultiplierRange, WaterSplineData.WaterDepth);
+				WaterSplineData.WaterVelocity = WaterSplineData.RawWaterVelocity * WaterSplineData.WaterVelocityDepthMultiplier;
+		
+				// Player velocity relative to water velocity.
+				const FVector PlayerVelocity2D = FVector(UpdateWaterSplineDataParams.PlayerVelocity.X, UpdateWaterSplineDataParams.PlayerVelocity.Y, 0.f);
+				const FVector WaterVelocity2D = FVector(WaterSplineData.WaterVelocity.X, WaterSplineData.WaterVelocity.Y, 0.f);
+				const FVector WaterVelocityProjection = PlayerVelocity2D.IsNearlyZero(0.1f) ? WaterVelocity2D : WaterVelocity2D.ProjectOnToNormal(PlayerVelocity2D.GetSafeNormal2D());
+				WaterSplineData.PlayerRelativeVelocityToWater = PlayerVelocity2D - WaterVelocityProjection;
+		
+				// Surface location
+				const float CapsuleBottom = UpdateWaterSplineDataParams.PlayerLocation.Z - UpdateWaterSplineDataParams.CapsuleHalfHeight;
+				const float CapsuleTop = UpdateWaterSplineDataParams.PlayerLocation.Z + UpdateWaterSplineDataParams.CapsuleHalfHeight;
+				WaterSplineData.WaterSurfaceLocation = QueryResult.GetWaterSurfaceLocation();
+		
+				// Surface offset
+				WaterSplineData.WaterSurfaceOffset = WaterSplineData.WaterSurfaceLocation - UpdateWaterSplineDataParams.PlayerLocation;
+		
+				// Immersion Percent
+				WaterSplineData.ImmersionPercent = FMath::Clamp((WaterSplineData.WaterSurfaceLocation.Z - CapsuleBottom) / (CapsuleTop - CapsuleBottom), 0.f, 1.f);
+		
+				// Normal
+				WaterSplineData.WaterSurfaceNormal = QueryResult.GetWaterSurfaceNormal();
+			}
+		}
 	}
 }
\ No newline at end of file
