Mover fixing issue with nav movement consuming invalid data causing slow/bad movement when framerate differed from fixed tick rate of Mover.

diff --git a/Engine/Plugins/Experimental/Mover/Source/Mover/Private/DefaultMovementSet/NavMoverComponent.cpp b/Engine/Plugins/Experimental/Mover/Source/Mover/Private/DefaultMovementSet/NavMoverComponent.cpp
--- a/Engine/Plugins/Experimental/Mover/Source/Mover/Private/DefaultMovementSet/NavMoverComponent.cpp
+++ b/Engine/Plugins/Experimental/Mover/Source/Mover/Private/DefaultMovementSet/NavMoverComponent.cpp
@@ -13,6 +13,8 @@ UNavMoverComponent::UNavMoverComponent()
 {
 	bWantsInitializeComponent = true;
 	bAutoActivate = true;
+	PrimaryComponentTick.bCanEverTick = true;
+	PrimaryComponentTick.bStartWithTickEnabled = true;
 }
 
 void UNavMoverComponent::InitializeComponent()
@@ -44,6 +46,25 @@ void UNavMoverComponent::BeginPlay()
 	}
 }
 
+void UNavMoverComponent::TickComponent(float DeltaTime, ELevelTick TickType, FActorComponentTickFunction* ThisTickFunction)
+{
+	Super::TickComponent(DeltaTime, TickType, ThisTickFunction);
+
+	if (bShouldClearNavMoveInputs)
+	{
+		CachedNavMoveInputIntent = FVector::ZeroVector;
+		CachedNavMoveInputVelocity = FVector::ZeroVector;
+		
+		bRequestedNavMovement = false;
+		bShouldClearNavMoveInputs = false;
+	}
+}
+
+bool UNavMoverComponent::HasRequestedNavMovement() const
+{
+	return bRequestedNavMovement;
+}
+
 float UNavMoverComponent::GetMaxSpeedForNavMovement() const
 {
 	const UCommonLegacyMovementSettings* MovementSettings = MoverComponent->FindSharedSettings_Mutable<UCommonLegacyMovementSettings>();
@@ -54,6 +75,25 @@ void UNavMoverComponent::StopMovementImmediately()
 {
 	TSharedPtr<FApplyVelocityEffect> VelocityEffect = MakeShared<FApplyVelocityEffect>();
 	MoverComponent->QueueInstantMovementEffect(VelocityEffect);
+	CachedNavMoveInputIntent = FVector::ZeroVector;
+	CachedNavMoveInputVelocity = FVector::ZeroVector;
+}
+
+bool UNavMoverComponent::ConsumeNavMovementData(FVector& OutMoveInputIntent, FVector& OutMoveInputVelocity)
+{
+	if (bRequestedNavMovement)
+	{
+		UE_LOG(LogMover, VeryVerbose, TEXT("Applying %s as NavMoveInputIntent."), *CachedNavMoveInputIntent.ToString());
+		UE_LOG(LogMover, VeryVerbose, TEXT("Applying %s as NavMoveInputVelocity."), *CachedNavMoveInputVelocity.ToString());
+		OutMoveInputIntent = CachedNavMoveInputIntent;
+		OutMoveInputVelocity = CachedNavMoveInputVelocity;
+		
+		bShouldClearNavMoveInputs = true;
+		
+		return true;
+	}
+
+	return false;
 }
 
 FVector UNavMoverComponent::GetLocation() const
@@ -120,6 +160,9 @@ void UNavMoverComponent::RequestDirectMove(const FVector& MoveVelocity, bool bFo
 		return;
 	}
 
+	bRequestedNavMovement = true;
+	bShouldClearNavMoveInputs = false;
+	
 	if (IsFalling())
 	{
 		const FVector FallVelocity = MoveVelocity.GetClampedToMaxSize(GetMaxSpeedForNavMovement());
@@ -130,7 +173,6 @@ void UNavMoverComponent::RequestDirectMove(const FVector& MoveVelocity, bool bFo
 	}
 
 	CachedNavMoveInputVelocity = MoveVelocity;
-	bRequestedNavMovement = true;
 	
 	if (IsMovingOnGround())
 	{
@@ -152,7 +194,8 @@ void UNavMoverComponent::RequestPathMove(const FVector& MoveInput)
 	}
 
 	bRequestedNavMovement = true;
-	CachedNavMoveInputIntent += AdjustedMoveInput;
+	bShouldClearNavMoveInputs = false;
+	CachedNavMoveInputIntent = AdjustedMoveInput.GetSafeNormal();
 }
 
 bool UNavMoverComponent::CanStopPathFollowing() const
diff --git a/Engine/Plugins/Experimental/Mover/Source/Mover/Public/DefaultMovementSet/NavMoverComponent.h b/Engine/Plugins/Experimental/Mover/Source/Mover/Public/DefaultMovementSet/NavMoverComponent.h
--- a/Engine/Plugins/Experimental/Mover/Source/Mover/Public/DefaultMovementSet/NavMoverComponent.h
+++ b/Engine/Plugins/Experimental/Mover/Source/Mover/Public/DefaultMovementSet/NavMoverComponent.h
@@ -24,13 +24,6 @@ class MOVER_API UNavMoverComponent : public UActorComponent, public INavMovement
 public:
 	UNavMoverComponent();
 
-	FVector CachedNavMoveInputIntent = FVector::ZeroVector;
-	FVector CachedNavMoveInputVelocity = FVector::ZeroVector;
-
-	// TODO: Not sure if we need these
-	FRotator CachedTurnInput = FRotator::ZeroRotator;
-	FRotator CachedLookInput = FRotator::ZeroRotator;
-
 	/** Properties that define how the component can move. */
 	UPROPERTY(EditAnywhere, Category="Nav Movement", meta = (DisplayName = "Movement Capabilities", Keywords = "Nav Agent"))
 	FNavAgentProperties NavAgentProps;
@@ -39,14 +32,19 @@ public:
 	UPROPERTY()
 	FMovementProperties MovementState;
 	
-	/** bool for keeping track of requested nav movement - is set to true when movement is requested, is false once nav movement was consumed */
-	UPROPERTY(BlueprintReadOnly, Category="Nav Movement")
-	bool bRequestedNavMovement = false;
-	
 protected:
 	/** associated properties for nav movement */
 	UPROPERTY(EditAnywhere, BlueprintReadWrite, Category="Nav Movement")
 	FNavMovementProperties NavMovementProperties;
+
+	/** bool for keeping track of requested nav movement - is set to true when movement is requested, is false once nav movement was consumed */
+	UPROPERTY(BlueprintReadOnly, Category="Nav Movement")
+	bool bRequestedNavMovement = false;
+	
+	// Used to store movement input intent from requested moves
+	FVector CachedNavMoveInputIntent = FVector::ZeroVector;
+	// Used to store movement input velocity from requested moves
+	FVector CachedNavMoveInputVelocity = FVector::ZeroVector;
 	
 private:
 	/** object implementing IPathFollowingAgentInterface. Is private to control access to it.
@@ -57,11 +55,18 @@ private:
 	/** Associated Movement component that will actually move the actor */ 
 	UPROPERTY()
 	TObjectPtr<UMoverComponent> MoverComponent;
+
+	UPROPERTY(Transient)
+	bool bShouldClearNavMoveInputs = false;
 	
 public:
 	virtual void InitializeComponent() override;
 	virtual void BeginPlay() override;
 
+	virtual void TickComponent(float DeltaTime, ELevelTick TickType, FActorComponentTickFunction* ThisTickFunction) override;
+
+	bool HasRequestedNavMovement() const;
+	
 	/** Get the owner of the object consuming nav movement */
 	virtual UObject* GetOwnerAsObject() const override { return GetOwner(); }
 	
@@ -88,6 +93,10 @@ public:
 	// Overridden to also call StopActiveMovement().
 	virtual void StopMovementImmediately() override;
 
+	// Writes internal cached requested velocities to the MoveInput passed in. Returns true if it had move input to write.
+	UFUNCTION(BlueprintCallable, Category="AI|Components|NavMovement")
+	bool ConsumeNavMovementData(FVector& OutMoveInputIntent, FVector& OutMoveInputVelocity);
+	
 	/** Returns location of controlled actor - meaning center of collision bounding box */
 	virtual FVector GetLocation() const override;
 	/** Returns location of controlled actor's "feet" meaning center of bottom of collision bounding box */
diff --git a/Engine/Plugins/Experimental/MoverExamples/Source/MoverExamples/Private/MoverExamplesCharacter.cpp b/Engine/Plugins/Experimental/MoverExamples/Source/MoverExamples/Private/MoverExamplesCharacter.cpp
--- a/Engine/Plugins/Experimental/MoverExamples/Source/MoverExamples/Private/MoverExamplesCharacter.cpp
+++ b/Engine/Plugins/Experimental/MoverExamples/Source/MoverExamples/Private/MoverExamplesCharacter.cpp
@@ -189,16 +189,7 @@ void AMoverExamplesCharacter::OnProduceInput(float DeltaMs, FMoverInputCmdContex
 	bool bRequestedNavMovement = false;
 	if (NavMoverComponent)
 	{
-		bRequestedNavMovement = NavMoverComponent->bRequestedNavMovement;
-		if (bRequestedNavMovement)
-		{
-			CachedMoveInputIntent = NavMoverComponent->CachedNavMoveInputIntent;
-			CachedMoveInputVelocity = NavMoverComponent->CachedNavMoveInputVelocity;
-			
-			NavMoverComponent->bRequestedNavMovement = false;
-			NavMoverComponent->CachedNavMoveInputIntent = FVector::ZeroVector;
-			NavMoverComponent->CachedNavMoveInputVelocity = FVector::ZeroVector;
-		}
+		bRequestedNavMovement = NavMoverComponent->ConsumeNavMovementData(CachedMoveInputIntent, CachedMoveInputVelocity);
 	}
 	
 	// Favor velocity input 
