Add anim next node "GenerateMoverTrajectory" to support mover trajectory prediction running as part of the anim next graph evaluation.
Add anim next node to debug draw pose search trajectory via the debug logger.

#tests win64 editor animsandbox

#rnx

#virtualized

diff --git a/Engine/Build/InstalledEngineFilters.xml b/Engine/Build/InstalledEngineFilters.xml
--- a/Engine/Build/InstalledEngineFilters.xml
+++ b/Engine/Build/InstalledEngineFilters.xml
@@ -156,6 +156,7 @@
 		<!-- Additional documentation -->
 		Engine/Plugins/Runtime/NetworkPrediction/readme.txt
 		Engine/Plugins/Experimental/Mover/README.md
+		Engine/Plugins/Experimental/MoverAnimNext/README.md
 		Engine/Plugins/Runtime/GameplayAbilities/README.md
 		Engine/Plugins/Experimental/RazerChromaDevices/README.txt
 
diff --git a/Engine/Plugins/Experimental/AnimNextPoseSearch/Source/AnimNextPoseSearch/Private/RigUnit_DebugDrawTrajectory.cpp b/Engine/Plugins/Experimental/AnimNextPoseSearch/Source/AnimNextPoseSearch/Private/RigUnit_DebugDrawTrajectory.cpp
new file mode 100644
--- /dev/null
+++ b/Engine/Plugins/Experimental/AnimNextPoseSearch/Source/AnimNextPoseSearch/Private/RigUnit_DebugDrawTrajectory.cpp
@@ -0,0 +1,25 @@
+// Copyright Epic Games, Inc. All Rights Reserved.
+
+#include "RigUnit_DebugDrawTrajectory.h"
+#include "PoseSearch/PoseSearchTrajectoryLibrary.h"
+#include "Component/AnimNextComponent.h"
+#include "PoseSearch/PoseSearchDefines.h"
+
+FRigUnit_DebugDrawTrajectory_Execute()
+{
+#if ENABLE_ANIM_DEBUG
+	if (!Enabled)
+		return;
+
+	using namespace UE::AnimNext;
+
+	const FAnimNextModuleContextData& ModuleContextData = ExecuteContext.GetContextData<FAnimNextModuleContextData>();
+	FAnimNextModuleInstance& ModuleInstance = ModuleContextData.GetModuleInstance();
+
+	// TODO: this is not thread safe at the moment. Once we get the 'weak semantics' CL checked in, we can move commonly-used debug info
+	// into a module component that can be used to access GT data where needed. 
+	UAnimNextComponent* Component = Cast<UAnimNextComponent>(ModuleInstance.Object);
+
+	Trajectory.DebugDrawTrajectory(Component, LogPoseSearch, ELogVerbosity::Display, DebugThickness, DebugOffset);
+#endif
+}
diff --git a/Engine/Plugins/Experimental/AnimNextPoseSearch/Source/AnimNextPoseSearch/Private/RigUnit_DebugDrawTrajectory.h b/Engine/Plugins/Experimental/AnimNextPoseSearch/Source/AnimNextPoseSearch/Private/RigUnit_DebugDrawTrajectory.h
new file mode 100644
--- /dev/null
+++ b/Engine/Plugins/Experimental/AnimNextPoseSearch/Source/AnimNextPoseSearch/Private/RigUnit_DebugDrawTrajectory.h
@@ -0,0 +1,30 @@
+// Copyright Epic Games, Inc. All Rights Reserved.
+
+#pragma once
+
+#include "AnimNextExecuteContext.h"
+#include "Graph/RigUnit_AnimNextBase.h"
+#include "PoseSearch/PoseSearchTrajectoryTypes.h"
+
+#include "RigUnit_DebugDrawTrajectory.generated.h"
+
+USTRUCT(meta = (DisplayName = "Debug Draw Trajectory", Category = "Mover", NodeColor = "0, 1, 1", Keywords = "Mover, Trajectory, Debug"))
+struct FRigUnit_DebugDrawTrajectory: public FRigUnit_AnimNextBase
+{
+	GENERATED_BODY()
+
+	RIGVM_METHOD()
+	void Execute();
+
+	UPROPERTY(EditAnywhere, Category = "Graph", meta = (Input, Output))
+	FPoseSearchQueryTrajectory Trajectory;
+
+	UPROPERTY(EditAnywhere, Category = "Graph", meta = (Input))
+	float DebugThickness;
+
+	UPROPERTY(EditAnywhere, Category = "Graph", meta = (Input))
+	float DebugOffset;
+
+	UPROPERTY(EditAnywhere, Category = "Graph", meta = (Input))
+	bool Enabled;
+};
diff --git a/Engine/Plugins/Experimental/Mover/Source/Mover/Private/MoverPoseSearchTrajectoryPredictor.cpp b/Engine/Plugins/Experimental/Mover/Source/Mover/Private/MoverPoseSearchTrajectoryPredictor.cpp
--- a/Engine/Plugins/Experimental/Mover/Source/Mover/Private/MoverPoseSearchTrajectoryPredictor.cpp
+++ b/Engine/Plugins/Experimental/Mover/Source/Mover/Private/MoverPoseSearchTrajectoryPredictor.cpp
@@ -14,13 +14,18 @@ void UMoverTrajectoryPredictor::Predict_Implementation(FPoseSearchQueryTrajector
 		return;
 	}
 
+	Predict_Static(*MoverComponent, InOutTrajectory, NumPredictionSamples, SecondsPerPredictionSample, NumHistorySamples);
+}
+
+void UMoverTrajectoryPredictor::Predict_Static(UMoverComponent& MoverComponent, FPoseSearchQueryTrajectory& InOutTrajectory, int32 NumPredictionSamples, float SecondsPerPredictionSample, int32 NumHistorySamples)
+{
 	FMoverPredictTrajectoryParams PredictParams;
 	PredictParams.NumPredictionSamples = NumPredictionSamples;
 	PredictParams.SecondsPerSample = SecondsPerPredictionSample;
 	PredictParams.bUseVisualComponentRoot = true;
 	PredictParams.bDisableGravity = true;
 
-	TArray<FTrajectorySampleInfo> MoverPredictionSamples = MoverComponent->GetPredictedTrajectory(PredictParams);
+	TArray<FTrajectorySampleInfo> MoverPredictionSamples = MoverComponent.GetPredictedTrajectory(PredictParams);
 
 	if (InOutTrajectory.Samples.Num() < (NumHistorySamples + MoverPredictionSamples.Num()))
 	{
@@ -67,7 +72,12 @@ void UMoverTrajectoryPredictor::GetCurrentState_Implementation(FVector& OutPosit
 		return;
 	}
 
-	const USceneComponent* VisualComp = MoverComponent->GetPrimaryVisualComponent();
+	GetCurrentState_Static(*MoverComponent, OutPosition, OutFacing, OutVelocity);
+}
+
+void UMoverTrajectoryPredictor::GetCurrentState_Static(UMoverComponent& MoverComponent, FVector& OutPosition, FQuat& OutFacing, FVector& OutVelocity)
+{
+	const USceneComponent* VisualComp = MoverComponent.GetPrimaryVisualComponent();
 
 	if (VisualComp)
 	{
@@ -75,7 +85,7 @@ void UMoverTrajectoryPredictor::GetCurrentState_Implementation(FVector& OutPosit
 	}
 	else
 	{
-		OutPosition = MoverComponent->GetUpdatedComponentTransform().GetLocation();
+		OutPosition = MoverComponent.GetUpdatedComponentTransform().GetLocation();
 	}
 
 	const bool bOrientRotationToMovement = true;
@@ -88,7 +98,7 @@ void UMoverTrajectoryPredictor::GetCurrentState_Implementation(FVector& OutPosit
 		}
 		else
 		{
-			OutFacing = MoverComponent->GetUpdatedComponentTransform().GetRotation();
+			OutFacing = MoverComponent.GetUpdatedComponentTransform().GetRotation();
 		}
 	}
 	else
@@ -97,7 +107,7 @@ void UMoverTrajectoryPredictor::GetCurrentState_Implementation(FVector& OutPosit
 		//OutFacing = FQuat::MakeFromRotator(FRotator(0, TrajectoryDataState.DesiredControllerYawLastUpdate, 0)) * TrajectoryDataDerived.MeshCompRelativeRotation;
 	}
 
-	OutVelocity = MoverComponent->GetVelocity();
+	OutVelocity = MoverComponent.GetVelocity();
 }
 
 
diff --git a/Engine/Plugins/Experimental/Mover/Source/Mover/Public/MoverPoseSearchTrajectoryPredictor.h b/Engine/Plugins/Experimental/Mover/Source/Mover/Public/MoverPoseSearchTrajectoryPredictor.h
--- a/Engine/Plugins/Experimental/Mover/Source/Mover/Public/MoverPoseSearchTrajectoryPredictor.h
+++ b/Engine/Plugins/Experimental/Mover/Source/Mover/Public/MoverPoseSearchTrajectoryPredictor.h
@@ -29,6 +29,12 @@ public:
 
 	virtual void GetVelocity_Implementation(FVector& OutVelocity) override;
 
+	// Share code in static helper functions so we can call it from AnimNext, which doesn't support interfaces
+	// For Review: what is naming convention here?
+	static void Predict_Static(UMoverComponent& Mover, FPoseSearchQueryTrajectory& InOutTrajectory, int32 NumPredictionSamples, float SecondsPerPredictionSample, int32 NumHistorySamples);
+	static void GetCurrentState_Static(UMoverComponent& Mover, FVector& OutPosition, FQuat& OutFacing, FVector& OutVelocity);
+
+
 protected:
 	TObjectPtr<UMoverComponent> MoverComponent;
 
diff --git a/Engine/Plugins/Experimental/MoverAnimNext/MoverAnimNext.uplugin b/Engine/Plugins/Experimental/MoverAnimNext/MoverAnimNext.uplugin
new file mode 100644
--- /dev/null
+++ b/Engine/Plugins/Experimental/MoverAnimNext/MoverAnimNext.uplugin
@@ -0,0 +1,34 @@
+{
+	"FileVersion": 1,
+	"Version": 1,
+	"VersionName": "1.0",
+	"FriendlyName": "Mover AnimNext",
+	"Description": "This plugin adds AnimNext support for the Mover plugin.",
+	"Category": "Gameplay",
+	"CreatedBy": "Epic Games, Inc.",
+	"CreatedByURL": "https://epicgames.com",
+	"DocsURL": "",
+	"MarketplaceURL": "",
+	"SupportURL": "",
+	"CanContainContent": true,
+	"IsBetaVersion": false,
+	"IsExperimentalVersion": true,
+	"Installed": false,
+	"Modules": [
+		{
+			"Name": "MoverAnimNext",
+			"Type": "Runtime",
+			"LoadingPhase": "Default"
+		}
+	],
+	"Plugins": [
+		{
+			"Name": "Mover",
+			"Enabled": true
+		},
+		{
+			"Name": "AnimNext",
+			"Enabled": true
+		}
+	]
+}
\ No newline at end of file
diff --git a/Engine/Plugins/Experimental/MoverAnimNext/README.md b/Engine/Plugins/Experimental/MoverAnimNext/README.md
new file mode 100644
--- /dev/null
+++ b/Engine/Plugins/Experimental/MoverAnimNext/README.md
@@ -0,0 +1,6 @@
+# Mover AnimNext Plugin
+
+The Mover AnimNext Plugin adds Mover specific nodes to anim next. See Mover and AnimNext Plugin documentation for details.
+
+**The Mover and Move AnimNext plugin are Experimental. Many features are incomplete or missing. APIs and data formats are subject to change at any time.**
+
diff --git a/Engine/Plugins/Experimental/MoverAnimNext/Source/MoverAnimNext/MoverAnimNext.Build.cs b/Engine/Plugins/Experimental/MoverAnimNext/Source/MoverAnimNext/MoverAnimNext.Build.cs
new file mode 100644
--- /dev/null
+++ b/Engine/Plugins/Experimental/MoverAnimNext/Source/MoverAnimNext/MoverAnimNext.Build.cs
@@ -0,0 +1,45 @@
+// Copyright Epic Games, Inc. All Rights Reserved.
+
+using UnrealBuildTool;
+
+public class MoverAnimNext : ModuleRules
+{
+	public MoverAnimNext(ReadOnlyTargetRules Target) : base(Target)
+	{
+
+		PublicDependencyModuleNames.AddRange(
+			new string[]
+			{
+				"Core",
+				"Mover",
+				"AnimNext",
+				"PoseSearch"
+			}
+			);
+			
+		
+		PrivateDependencyModuleNames.AddRange(
+			new string[]
+			{
+				"CoreUObject",
+				"RigVM",
+				"ControlRig",
+				"Engine",
+				"PhysicsCore",
+				"DeveloperSettings"
+			}
+			);
+
+		if (Target.bBuildEditor == true)
+		{
+			PrivateDependencyModuleNames.AddRange(
+				new string[]
+				{
+					"RigVMDeveloper"
+				}
+			);
+		}
+
+		SetupGameplayDebuggerSupport(Target);
+	}
+}
diff --git a/Engine/Plugins/Experimental/MoverAnimNext/Source/MoverAnimNext/Private/Graph/RigUnit_GenerateMoverTrajectory.cpp b/Engine/Plugins/Experimental/MoverAnimNext/Source/MoverAnimNext/Private/Graph/RigUnit_GenerateMoverTrajectory.cpp
new file mode 100644
--- /dev/null
+++ b/Engine/Plugins/Experimental/MoverAnimNext/Source/MoverAnimNext/Private/Graph/RigUnit_GenerateMoverTrajectory.cpp
@@ -0,0 +1,48 @@
+// Copyright Epic Games, Inc. All Rights Reserved.
+
+#include "RigUnit_GenerateMoverTrajectory.h"
+
+#include "Component/AnimNextComponent.h"
+#include "MoverComponent.h"
+#include "MoverPoseSearchTrajectoryPredictor.h"
+
+FRigUnit_GenerateMoverTrajectory_Execute()
+{
+	using namespace UE::AnimNext;
+
+	const FAnimNextModuleContextData& ModuleContextData = ExecuteContext.GetContextData<FAnimNextModuleContextData>();
+	FAnimNextModuleInstance& ModuleInstance = ModuleContextData.GetModuleInstance();
+
+	// TODO: this is not thread safe at the moment. Once we get the 'weak semantics' CL checked in, we can move commonly-used debug info
+	// into a module component that can be used to access GT data where needed. 
+	UAnimNextComponent* Component = Cast<UAnimNextComponent>(ModuleInstance.Object);
+
+	// Note this logic is based on UPoseSearchTrajectoryLibrary::PoseSearchGeneratePredictorTrajectory
+
+	FPoseSearchTrajectoryData::FSampling TrajectoryDataSampling;
+	TrajectoryDataSampling.NumHistorySamples = NumHistorySamples;
+	TrajectoryDataSampling.SecondsPerHistorySample = HistorySamplingInterval;
+	TrajectoryDataSampling.NumPredictionSamples = NumPredictionSamples;
+	TrajectoryDataSampling.SecondsPerPredictionSample = PredictionSamplingInterval;
+
+	// TODO: handle controller yaw
+	//TrajectoryDataState.DesiredControllerYawLastUpdate = InOutDesiredControllerYawLastUpdate;
+
+	FVector CurrentPosition = FVector::ZeroVector;
+	FVector CurrentVelocity = FVector::ZeroVector;
+	FQuat CurrentFacing = FQuat::Identity;
+
+	if (MoverComponent == nullptr)
+	{
+		// Cannot execute, abort
+		// Toreview: error / warning logging necessary?
+		return;
+	}
+
+	UMoverTrajectoryPredictor::GetCurrentState_Static(*MoverComponent, CurrentPosition, CurrentFacing, CurrentVelocity);
+
+	UPoseSearchTrajectoryLibrary::InitTrajectorySamples(InOutTrajectory, TrajectoryData, CurrentPosition, CurrentFacing, TrajectoryDataSampling, DeltaTime);
+	UPoseSearchTrajectoryLibrary::UpdateHistory_TransformHistory(InOutTrajectory, TrajectoryData, CurrentPosition, CurrentVelocity, TrajectoryDataSampling, DeltaTime);
+	UMoverTrajectoryPredictor::Predict_Static(*MoverComponent, InOutTrajectory,
+		NumPredictionSamples + 1, PredictionSamplingInterval, NumHistorySamples);
+}
diff --git a/Engine/Plugins/Experimental/MoverAnimNext/Source/MoverAnimNext/Private/Graph/RigUnit_GenerateMoverTrajectory.h b/Engine/Plugins/Experimental/MoverAnimNext/Source/MoverAnimNext/Private/Graph/RigUnit_GenerateMoverTrajectory.h
new file mode 100644
--- /dev/null
+++ b/Engine/Plugins/Experimental/MoverAnimNext/Source/MoverAnimNext/Private/Graph/RigUnit_GenerateMoverTrajectory.h
@@ -0,0 +1,46 @@
+// Copyright Epic Games, Inc. All Rights Reserved.
+
+#pragma once
+
+#include "AnimNextExecuteContext.h"
+#include "Graph/RigUnit_AnimNextBase.h"
+#include "PoseSearch/PoseSearchTrajectoryLibrary.h"
+
+#include "RigUnit_GenerateMoverTrajectory.generated.h"
+
+class UMoverComponent;
+
+USTRUCT(meta = (DisplayName = "Generate Trajectory from Mover", Category = "Mover", NodeColor = "0, 1, 1", Keywords = "Mover, Trajectory"))
+struct FRigUnit_GenerateMoverTrajectory : public FRigUnit_AnimNextBase
+{
+	GENERATED_BODY()
+
+	RIGVM_METHOD()
+	void Execute();
+
+	UPROPERTY(EditAnywhere, Category = "Graph", meta = (Input))
+	TObjectPtr<UMoverComponent> MoverComponent;
+
+	UPROPERTY(EditAnywhere, Category = "Graph", meta = (Input))
+	FPoseSearchTrajectoryData TrajectoryData;
+
+	UPROPERTY(EditAnywhere, Category = "Graph", meta = (Input))
+	// This should be the most recent simulation time that was used to get us to our current state
+	float DeltaTime;
+
+	UPROPERTY(EditAnywhere, Category = "Graph", meta = (Input))
+	float HistorySamplingInterval;
+
+	UPROPERTY(EditAnywhere, Category = "Graph", meta = (Input))
+	int32 NumHistorySamples;
+
+	UPROPERTY(EditAnywhere, Category = "Graph", meta = (Input))
+	float PredictionSamplingInterval;
+
+	UPROPERTY(EditAnywhere, Category = "Graph", meta = (Input))
+	int32 NumPredictionSamples;
+
+	UPROPERTY(EditAnywhere, Category = "Graph", meta = (Input, Output))
+	FPoseSearchQueryTrajectory InOutTrajectory;
+
+};
diff --git a/Engine/Plugins/Experimental/MoverAnimNext/Source/MoverAnimNext/Private/MoverAnimNextModule.cpp b/Engine/Plugins/Experimental/MoverAnimNext/Source/MoverAnimNext/Private/MoverAnimNextModule.cpp
new file mode 100644
--- /dev/null
+++ b/Engine/Plugins/Experimental/MoverAnimNext/Source/MoverAnimNext/Private/MoverAnimNextModule.cpp
@@ -0,0 +1,43 @@
+// Copyright Epic Games, Inc. All Rights Reserved.
+#include "MoverAnimNextModule.h"
+
+#include "CoreMinimal.h"
+#include "RigVMCore/RigVMRegistry.h"
+#include "MoveLibrary/MovementUtils.h"
+#include "MoverComponent.h"
+#include "MoverPoseSearchTrajectoryPredictor.h"
+#include "PoseSearch/PoseSearchTrajectoryLibrary.h"
+
+#include "Modules/ModuleManager.h"
+
+
+void FMoverAnimNextModule::StartupModule()
+{	
+	static UScriptStruct* const AllowedStructTypes[] =
+	{
+		FMoverPredictTrajectoryParams::StaticStruct(),	// Todo: remove
+		FTrajectorySampleInfo::StaticStruct(),			// Todo: remove
+		FPoseSearchTrajectoryData::StaticStruct(),
+		FPoseSearchQueryTrajectory::StaticStruct(),
+
+	};
+
+	static TPair<UClass*, FRigVMRegistry::ERegisterObjectOperation> const AllowedObjectTypes[] =
+	{
+		{ UMoverComponent::StaticClass(), FRigVMRegistry::ERegisterObjectOperation::Class },
+		{ UMoverTrajectoryPredictor::StaticClass(), FRigVMRegistry::ERegisterObjectOperation::Class },	// Doesn't seem like anim next supports interfaces
+	};
+
+	FRigVMRegistry& RigVMRegistry = FRigVMRegistry::Get();
+	RigVMRegistry.RegisterStructTypes(AllowedStructTypes);
+	RigVMRegistry.RegisterObjectTypes(AllowedObjectTypes);
+
+}
+
+void FMoverAnimNextModule::ShutdownModule()
+{
+	// This function may be called during shutdown to clean up your module.  For modules that support dynamic reloading,
+	// we call this function before unloading the module.
+}
+
+IMPLEMENT_MODULE(FMoverAnimNextModule, MoverAnimNext)
\ No newline at end of file
diff --git a/Engine/Plugins/Experimental/MoverAnimNext/Source/MoverAnimNext/Public/MoverAnimNextModule.h b/Engine/Plugins/Experimental/MoverAnimNext/Source/MoverAnimNext/Public/MoverAnimNextModule.h
new file mode 100644
--- /dev/null
+++ b/Engine/Plugins/Experimental/MoverAnimNext/Source/MoverAnimNext/Public/MoverAnimNextModule.h
@@ -0,0 +1,15 @@
+// Copyright Epic Games, Inc. All Rights Reserved.
+
+#pragma once
+
+#include "CoreMinimal.h"
+#include "Modules/ModuleManager.h"
+
+class FMoverAnimNextModule : public IModuleInterface
+{
+public:
+
+	/** IModuleInterface implementation */
+	virtual void StartupModule() override;
+	virtual void ShutdownModule() override;
+};
