Mover: added policy options in default Walking mode to control whether to perform floor checks when not moving

diff --git a/Engine/Plugins/Experimental/Mover/Source/Mover/Private/DefaultMovementSet/Modes/WalkingMode.cpp b/Engine/Plugins/Experimental/Mover/Source/Mover/Private/DefaultMovementSet/Modes/WalkingMode.cpp
--- a/Engine/Plugins/Experimental/Mover/Source/Mover/Private/DefaultMovementSet/Modes/WalkingMode.cpp
+++ b/Engine/Plugins/Experimental/Mover/Source/Mover/Private/DefaultMovementSet/Modes/WalkingMode.cpp
@@ -261,31 +261,35 @@ void UWalkingMode::OnSimulationTick(const FSimulationTickParams& Params, FMoverT
 	}
 	else
 	{
-		// If the actor isn't moving we still need to check if they have a valid floor
-		UFloorQueryUtils::FindFloor(MovingComponents, CommonLegacySettings->FloorSweepDistance, CommonLegacySettings->MaxWalkSlopeCosine, UpdatedPrimitive->GetComponentLocation(), CurrentFloor);
-		
-		FHitResult Hit(CurrentFloor.HitResult);
-		if (Hit.bStartPenetrating)
+		// If the actor isn't moving we still may need to check if they have a valid floor, such as if they're on an elevator platform moving up/down
+		if ((FloorCheckPolicy == EStaticFloorCheckPolicy::Always) || 
+		    (FloorCheckPolicy == EStaticFloorCheckPolicy::OnDynamicBaseOnly && StartingSyncState->GetMovementBase()))
 		{
-			// The floor check failed because it started in penetration
-			// We do not want to try to move downward because the downward sweep failed, rather we'd like to try to pop out of the floor.
-			Hit.TraceEnd = Hit.TraceStart + UpDirection * 2.4;
-			FVector RequestedAdjustment = UMovementUtils::ComputePenetrationAdjustment(Hit);
+			UFloorQueryUtils::FindFloor(MovingComponents, CommonLegacySettings->FloorSweepDistance, CommonLegacySettings->MaxWalkSlopeCosine, UpdatedPrimitive->GetComponentLocation(), CurrentFloor);
+		
+			FHitResult Hit(CurrentFloor.HitResult);
+			if (Hit.bStartPenetrating)
+			{
+				// The floor check failed because it started in penetration
+				// We do not want to try to move downward because the downward sweep failed, rather we'd like to try to pop out of the floor.
+				Hit.TraceEnd = Hit.TraceStart + UpDirection * 2.4;
+				FVector RequestedAdjustment = UMovementUtils::ComputePenetrationAdjustment(Hit);
 			
-			const EMoveComponentFlags IncludeBlockingOverlapsWithoutEvents = (MOVECOMP_NeverIgnoreBlockingOverlaps | MOVECOMP_DisableBlockingOverlapDispatch);
-			EMoveComponentFlags MoveComponentFlags = MOVECOMP_NoFlags;
-			MoveComponentFlags = (MoveComponentFlags | IncludeBlockingOverlapsWithoutEvents);
-			UMovementUtils::TryMoveToResolvePenetration(Params.MovingComps, MoveComponentFlags, RequestedAdjustment, Hit, UpdatedComponent->GetComponentQuat(), MoveRecord);
-		}
+				const EMoveComponentFlags IncludeBlockingOverlapsWithoutEvents = (MOVECOMP_NeverIgnoreBlockingOverlaps | MOVECOMP_DisableBlockingOverlapDispatch);
+				EMoveComponentFlags MoveComponentFlags = MOVECOMP_NoFlags;
+				MoveComponentFlags = (MoveComponentFlags | IncludeBlockingOverlapsWithoutEvents);
+				UMovementUtils::TryMoveToResolvePenetration(Params.MovingComps, MoveComponentFlags, RequestedAdjustment, Hit, UpdatedComponent->GetComponentQuat(), MoveRecord);
+			}
 		
-		if (!CurrentFloor.IsWalkableFloor() && !Hit.bStartPenetrating)
-		{
-			// No floor or not walkable, so let's let the airborne movement mode deal with it
-			OutputState.MovementEndState.NextModeName = CommonLegacySettings->AirMovementModeName;
-			OutputState.MovementEndState.RemainingMs = Params.TimeStep.StepMs;
-			MoveRecord.SetDeltaSeconds((Params.TimeStep.StepMs - OutputState.MovementEndState.RemainingMs) * 0.001f);
-			CaptureFinalState(UpdatedComponent, bDidAttemptMovement, CurrentFloor, MoveRecord, OutputSyncState);
-			return;
+			if (!CurrentFloor.IsWalkableFloor() && !Hit.bStartPenetrating)
+			{
+				// No floor or not walkable, so let's let the airborne movement mode deal with it
+				OutputState.MovementEndState.NextModeName = CommonLegacySettings->AirMovementModeName;
+				OutputState.MovementEndState.RemainingMs = Params.TimeStep.StepMs;
+				MoveRecord.SetDeltaSeconds((Params.TimeStep.StepMs - OutputState.MovementEndState.RemainingMs) * 0.001f);
+				CaptureFinalState(UpdatedComponent, bDidAttemptMovement, CurrentFloor, MoveRecord, OutputSyncState);
+				return;
+			}
 		}
 	}
 
diff --git a/Engine/Plugins/Experimental/Mover/Source/Mover/Public/DefaultMovementSet/Modes/WalkingMode.h b/Engine/Plugins/Experimental/Mover/Source/Mover/Public/DefaultMovementSet/Modes/WalkingMode.h
--- a/Engine/Plugins/Experimental/Mover/Source/Mover/Public/DefaultMovementSet/Modes/WalkingMode.h
+++ b/Engine/Plugins/Experimental/Mover/Source/Mover/Public/DefaultMovementSet/Modes/WalkingMode.h
@@ -12,6 +12,19 @@ struct FFloorCheckResult;
 struct FRelativeBaseInfo;
 struct FMovementRecord;
 
+// Behavior policy for performing floor checks in walking mode when no movement is occurring. 
+UENUM(BlueprintType)
+enum class EStaticFloorCheckPolicy : uint8
+{
+	// Always perform floor checks, even when not moving. You may want this if static bases may disappear from underneath.
+	Always = 0				UMETA(DisplayName = "Always"),
+
+	// Only perform floor checks when not moving IF we're on a dynamic movement base
+	OnDynamicBaseOnly = 1	UMETA(DisplayName = "OnDynamicBaseOnly"),
+};
+
+
+
 /**
  * WalkingMode: a default movement mode for traversing surfaces and movement bases (walking, running, sneaking, etc.)
  */
@@ -37,6 +50,11 @@ public:
 	void SetTurnGeneratorClass(UPARAM(meta=(MustImplement="/Script/Mover.TurnGeneratorInterface", AllowAbstract="false")) TSubclassOf<UObject> TurnGeneratorClass);
 
 protected:
+
+	/** Choice of behavior for floor checks while not moving.  */
+	UPROPERTY(EditAnywhere, BlueprintReadWrite, Category = Mover)
+	EStaticFloorCheckPolicy FloorCheckPolicy = EStaticFloorCheckPolicy::OnDynamicBaseOnly;
+
 	/** Optional modular object for generating rotation towards desired orientation. If not specified, linear interpolation will be used. */
 	UPROPERTY(EditAnywhere, Instanced, Category=Mover, meta=(MustImplement="/Script/Mover.TurnGeneratorInterface"))
 	TObjectPtr<UObject> TurnGenerator;
