Fix multiple axis angular velocity calculation in MoverNetworkPhysicsLiason

diff --git a/Engine/Plugins/Experimental/Mover/Source/Mover/Private/Backends/MoverNetworkPhysicsLiaison.cpp b/Engine/Plugins/Experimental/Mover/Source/Mover/Private/Backends/MoverNetworkPhysicsLiaison.cpp
--- a/Engine/Plugins/Experimental/Mover/Source/Mover/Private/Backends/MoverNetworkPhysicsLiaison.cpp
+++ b/Engine/Plugins/Experimental/Mover/Source/Mover/Private/Backends/MoverNetworkPhysicsLiaison.cpp
@@ -1061,14 +1061,16 @@ void UMoverNetworkPhysicsLiaisonComponent::OnPreSimulate_Internal(const FPhysics
 	float TargetDeltaFacing = FMath::DegreesToRadians(Remainder.Yaw);
 	if (TickParams.DeltaTimeSeconds > UE_SMALL_NUMBER)
 	{
-		if (CommonMovementSettings->bShouldRemainVertical)
+		FVector AngularVelocity = (TargetDeltaFacing / TickParams.DeltaTimeSeconds) * Chaos::FVec3::ZAxisVector;
+		
+		if (!CommonMovementSettings->bShouldRemainVertical)
 		{
-			CharacterParticle->SetW((TargetDeltaFacing / TickParams.DeltaTimeSeconds) * Chaos::FVec3::ZAxisVector);
-		}
-		else
-		{
-			CharacterParticle->SetW(FMath::DegreesToRadians(Remainder).Quaternion().ToRotationVector() / TickParams.DeltaTimeSeconds);
+			FVector PreviousUp = CharacterParticle->GetR().RotateVector(FVector::UpVector).GetSafeNormal();
+			FVector TargetUp = OutputSyncState.GetOrientation_WorldSpace().Quaternion().RotateVector(FVector::UpVector).GetSafeNormal();
+			AngularVelocity += (PreviousUp.Cross(TargetUp).GetSafeNormal() * FMath::Acos(TargetUp.Dot(PreviousUp))) / TickParams.DeltaTimeSeconds;
 		}
+		
+		CharacterParticle->SetW(AngularVelocity);
 	}
 
 	// Update the constraint data based on the floor result
