Mover fixing issue with Velocity input only being used to calculate direction input. It now is taken into account and Mover tries to match the velocity given it's own constraints.

diff --git a/Engine/Plugins/Experimental/Mover/Source/Mover/Private/DefaultMovementSet/Modes/FallingMode.cpp b/Engine/Plugins/Experimental/Mover/Source/Mover/Private/DefaultMovementSet/Modes/FallingMode.cpp
--- a/Engine/Plugins/Experimental/Mover/Source/Mover/Private/DefaultMovementSet/Modes/FallingMode.cpp
+++ b/Engine/Plugins/Experimental/Mover/Source/Mover/Private/DefaultMovementSet/Modes/FallingMode.cpp
@@ -88,6 +88,7 @@ void UFallingMode::OnGenerateMove(const FMoverTickStartData& StartState, const F
 	Params.Acceleration = CommonLegacySettings->Acceleration;
 	Params.Deceleration = FallingDeceleration;
 	Params.WorldToGravityQuat = MoverComp->GetWorldToGravityTransform();
+	Params.bUseAccelerationForVelocityMove = CommonLegacySettings->bUseAccelerationForVelocityMove;
 
 	// Check if any current velocity values are over our terminal velocity - if so limit the move input in that direction and apply OverTerminalVelocityFallingDeceleration
 	if (Params.MoveInput.Dot(StartVelocity) > 0 && StartHorizontalVelocity.Size() >= TerminalMovementPlaneSpeed)
diff --git a/Engine/Plugins/Experimental/Mover/Source/Mover/Private/DefaultMovementSet/Modes/FlyingMode.cpp b/Engine/Plugins/Experimental/Mover/Source/Mover/Private/DefaultMovementSet/Modes/FlyingMode.cpp
--- a/Engine/Plugins/Experimental/Mover/Source/Mover/Private/DefaultMovementSet/Modes/FlyingMode.cpp
+++ b/Engine/Plugins/Experimental/Mover/Source/Mover/Private/DefaultMovementSet/Modes/FlyingMode.cpp
@@ -61,6 +61,7 @@ void UFlyingMode::OnGenerateMove(const FMoverTickStartData& StartState, const FM
 	Params.Deceleration = CommonLegacySettings->Deceleration;
 	Params.DeltaSeconds = DeltaSeconds;
 	Params.WorldToGravityQuat = MoverComp->GetWorldToGravityTransform();
+	Params.bUseAccelerationForVelocityMove = CommonLegacySettings->bUseAccelerationForVelocityMove;
 	
 	OutProposedMove = UAirMovementUtils::ComputeControlledFreeMove(Params);
 }
diff --git a/Engine/Plugins/Experimental/Mover/Source/Mover/Private/DefaultMovementSet/Modes/NavWalkingMode.cpp b/Engine/Plugins/Experimental/Mover/Source/Mover/Private/DefaultMovementSet/Modes/NavWalkingMode.cpp
--- a/Engine/Plugins/Experimental/Mover/Source/Mover/Private/DefaultMovementSet/Modes/NavWalkingMode.cpp
+++ b/Engine/Plugins/Experimental/Mover/Source/Mover/Private/DefaultMovementSet/Modes/NavWalkingMode.cpp
@@ -91,7 +91,8 @@ void UNavWalkingMode::OnGenerateMove(const FMoverTickStartData& StartState, cons
 	Params.DeltaSeconds = DeltaSeconds;
 	Params.WorldToGravityQuat = MoverComp->GetWorldToGravityTransform();
 	Params.UpDirection = UpDirection;
-	
+	Params.bUseAccelerationForVelocityMove = CommonLegacySettings->bUseAccelerationForVelocityMove;
+
 	if (Params.MoveInput.SizeSquared() > 0.f && !UMovementUtils::IsExceedingMaxSpeed(Params.PriorVelocity, CommonLegacySettings->MaxSpeed))
 	{
 		Params.Friction = CommonLegacySettings->GroundFriction;
diff --git a/Engine/Plugins/Experimental/Mover/Source/Mover/Private/DefaultMovementSet/Modes/WalkingMode.cpp b/Engine/Plugins/Experimental/Mover/Source/Mover/Private/DefaultMovementSet/Modes/WalkingMode.cpp
--- a/Engine/Plugins/Experimental/Mover/Source/Mover/Private/DefaultMovementSet/Modes/WalkingMode.cpp
+++ b/Engine/Plugins/Experimental/Mover/Source/Mover/Private/DefaultMovementSet/Modes/WalkingMode.cpp
@@ -89,6 +89,7 @@ void UWalkingMode::OnGenerateMove(const FMoverTickStartData& StartState, const F
 	Params.DeltaSeconds = DeltaSeconds;
 	Params.WorldToGravityQuat = MoverComp->GetWorldToGravityTransform();
 	Params.UpDirection = UpDirection;
+	Params.bUseAccelerationForVelocityMove = CommonLegacySettings->bUseAccelerationForVelocityMove;
 	
 	if (Params.MoveInput.SizeSquared() > 0.f && !UMovementUtils::IsExceedingMaxSpeed(Params.PriorVelocity, CommonLegacySettings->MaxSpeed))
 	{
diff --git a/Engine/Plugins/Experimental/Mover/Source/Mover/Private/MoveLibrary/AirMovementUtils.cpp b/Engine/Plugins/Experimental/Mover/Source/Mover/Private/MoveLibrary/AirMovementUtils.cpp
--- a/Engine/Plugins/Experimental/Mover/Source/Mover/Private/MoveLibrary/AirMovementUtils.cpp
+++ b/Engine/Plugins/Experimental/Mover/Source/Mover/Private/MoveLibrary/AirMovementUtils.cpp
@@ -12,7 +12,7 @@ FProposedMove UAirMovementUtils::ComputeControlledFreeMove(const FFreeMoveParams
 
 	const FPlane MovementPlane(FVector::ZeroVector, FVector::UpVector);
 
-	OutMove.DirectionIntent = UMovementUtils::ComputeDirectionIntent(InParams.MoveInput, InParams.MoveInputType);
+	OutMove.DirectionIntent = UMovementUtils::ComputeDirectionIntent(InParams.MoveInput, InParams.MoveInputType, InParams.MaxSpeed);
 	OutMove.bHasDirIntent = !OutMove.DirectionIntent.IsNearlyZero();
 
 	FComputeVelocityParams ComputeVelocityParams;
@@ -23,6 +23,9 @@ FProposedMove UAirMovementUtils::ComputeControlledFreeMove(const FFreeMoveParams
 	ComputeVelocityParams.TurningBoost = InParams.TurningBoost;
 	ComputeVelocityParams.Deceleration = InParams.Deceleration;
 	ComputeVelocityParams.Acceleration = InParams.Acceleration;
+	ComputeVelocityParams.MoveInputType = InParams.MoveInputType;
+	ComputeVelocityParams.MoveInput = InParams.MoveInput; 
+	ComputeVelocityParams.bUseAccelerationForVelocityMove = InParams.bUseAccelerationForVelocityMove;
 	
 	OutMove.LinearVelocity = UMovementUtils::ComputeVelocity(ComputeVelocityParams);
 	OutMove.AngularVelocity = UMovementUtils::ComputeAngularVelocity(InParams.PriorOrientation, InParams.OrientationIntent, InParams.WorldToGravityQuat, InParams.DeltaSeconds, InParams.TurningRate);
diff --git a/Engine/Plugins/Experimental/Mover/Source/Mover/Private/MoveLibrary/GroundMovementUtils.cpp b/Engine/Plugins/Experimental/Mover/Source/Mover/Private/MoveLibrary/GroundMovementUtils.cpp
--- a/Engine/Plugins/Experimental/Mover/Source/Mover/Private/MoveLibrary/GroundMovementUtils.cpp
+++ b/Engine/Plugins/Experimental/Mover/Source/Mover/Private/MoveLibrary/GroundMovementUtils.cpp
@@ -14,7 +14,7 @@ FProposedMove UGroundMovementUtils::ComputeControlledGroundMove(const FGroundMov
 {
 	FProposedMove OutMove;
 
-	const FVector MoveDirIntent = UMovementUtils::ComputeDirectionIntent(InParams.MoveInput, InParams.MoveInputType);
+	const FVector MoveDirIntent = UMovementUtils::ComputeDirectionIntent(InParams.MoveInput, InParams.MoveInputType, InParams.MaxSpeed);
 
 	const FPlane MovementPlane(FVector::ZeroVector, InParams.UpDirection);
 	FVector MoveDirIntentInMovementPlane = UMovementUtils::ConstrainToPlane(MoveDirIntent, MovementPlane, true);
@@ -33,6 +33,9 @@ FProposedMove UGroundMovementUtils::ComputeControlledGroundMove(const FGroundMov
 	ComputeVelocityParams.Deceleration = InParams.Deceleration;
 	ComputeVelocityParams.Acceleration = InParams.Acceleration;
 	ComputeVelocityParams.Friction = InParams.Friction;
+	ComputeVelocityParams.MoveInputType = InParams.MoveInputType;
+	ComputeVelocityParams.MoveInput = InParams.MoveInput;
+	ComputeVelocityParams.bUseAccelerationForVelocityMove = InParams.bUseAccelerationForVelocityMove;
 	
 	// Figure out linear velocity
 	const FVector Velocity = UMovementUtils::ComputeVelocity(ComputeVelocityParams);
diff --git a/Engine/Plugins/Experimental/Mover/Source/Mover/Private/MoveLibrary/MovementUtils.cpp b/Engine/Plugins/Experimental/Mover/Source/Mover/Private/MoveLibrary/MovementUtils.cpp
--- a/Engine/Plugins/Experimental/Mover/Source/Mover/Private/MoveLibrary/MovementUtils.cpp
+++ b/Engine/Plugins/Experimental/Mover/Source/Mover/Private/MoveLibrary/MovementUtils.cpp
@@ -14,6 +14,7 @@
 namespace UE::MoverUtils
 {
 	const double SMALL_MOVE_DISTANCE = 1e-3;
+	const float  VELOCITY_INPUT_NO_ACCELERATION_DIFFERENCE = 1.01f;
 }
 
 bool UMovementUtils::IsExceedingMaxSpeed(const FVector& Velocity, float InMaxSpeed)
@@ -29,43 +30,75 @@ bool UMovementUtils::IsExceedingMaxSpeed(const FVector& Velocity, float InMaxSpe
 
 FVector UMovementUtils::ComputeVelocity(const FComputeVelocityParams& InParams)
 {
-	const FVector ControlAcceleration = InParams.MoveDirectionIntent.GetClampedToMaxSize(1.f);
+	FVector Acceleration = FVector::ZeroVector;
 	FVector Velocity = InParams.InitialVelocity;
+	float DesiredSpeed = 0.0f;
 
-	const float AnalogInputModifier = (ControlAcceleration.SizeSquared() > 0.f ? ControlAcceleration.Size() : 0.f);
-	const float MaxPawnSpeed = InParams.MaxSpeed * AnalogInputModifier;
-	const bool bExceedingMaxSpeed = IsExceedingMaxSpeed(Velocity, MaxPawnSpeed);
-
-	if (AnalogInputModifier > 0.f && !bExceedingMaxSpeed)
+	if (InParams.MoveInputType == EMoveInputType::Velocity)
 	{
-		// Apply change in velocity direction
-		if (Velocity.SizeSquared() > 0.f)
+		const float RequestedSpeed = FMath::Min(InParams.MaxSpeed,  InParams.MoveInput.Size());
+		const FVector RequestedMoveDir = InParams.MoveInput.GetSafeNormal();
+		DesiredSpeed = RequestedSpeed;
+
+		// Compute acceleration. Use MaxAccel to limit speed increase
+		if (InParams.bUseAccelerationForVelocityMove && InParams.InitialVelocity.Size() < RequestedSpeed * UE::MoverUtils::VELOCITY_INPUT_NO_ACCELERATION_DIFFERENCE)
 		{
-			// Change direction faster than only using acceleration, but never increase velocity magnitude.
-			const float TimeScale = FMath::Clamp(InParams.DeltaSeconds * InParams.TurningBoost, 0.f, 1.f);
-			Velocity = Velocity + (ControlAcceleration * Velocity.Size() - Velocity) * FMath::Min(TimeScale * InParams.Friction, 1.f);
+			// Turn in the same manner as with input acceleration.
+			Velocity = Velocity - (Velocity - RequestedMoveDir * RequestedSpeed) * FMath::Min(InParams.DeltaSeconds * InParams.Friction, 1.f);
+
+			// How much do we need to accelerate to get to the new velocity?
+			Acceleration = (InParams.MoveInput - Velocity) / InParams.DeltaSeconds;
+			Acceleration = Acceleration.GetClampedToMaxSize(InParams.Acceleration);
+		}
+		else
+		{
+			// Just set velocity directly.
+			// If decelerating we do so instantly, so we don't slide through the destination if we can't brake fast enough.
+			Velocity = InParams.MoveInput;
 		}
 	}
-	else
+	else if (InParams.MoveInputType == EMoveInputType::DirectionalIntent)
 	{
-		// Dampen velocity magnitude based on deceleration.
+		const FVector ControlAcceleration = InParams.MoveDirectionIntent.GetClampedToMaxSize(1.f);
+		const float AnalogInputModifier = (ControlAcceleration.SizeSquared() > 0.f ? ControlAcceleration.Size() : 0.f);
+		DesiredSpeed = InParams.MaxSpeed * AnalogInputModifier;
+		const bool bExceedingMaxSpeed = IsExceedingMaxSpeed(Velocity, DesiredSpeed);
+
 		if (Velocity.SizeSquared() > 0.f)
 		{
-			const FVector OldVelocity = Velocity;
-			const float VelSize = FMath::Max(Velocity.Size() - FMath::Abs(InParams.Friction * Velocity.Size() + InParams.Deceleration) * InParams.DeltaSeconds, 0.f);
-			Velocity = Velocity.GetSafeNormal() * VelSize;
-
-			// Don't allow braking to lower us below max speed if we started above it.
-			if (bExceedingMaxSpeed && Velocity.SizeSquared() < FMath::Square(MaxPawnSpeed))
+			if (AnalogInputModifier > 0.f && !bExceedingMaxSpeed)
+			{
+				// Apply change in velocity direction
+				// Change direction faster than only using acceleration, but never increase velocity magnitude.
+				const float TimeScale = FMath::Clamp(InParams.DeltaSeconds * InParams.TurningBoost, 0.f, 1.f);
+				Velocity = Velocity + (ControlAcceleration * Velocity.Size() - Velocity) * FMath::Min(TimeScale * InParams.Friction, 1.f);
+			}
+			else
 			{
-				Velocity = OldVelocity.GetSafeNormal() * MaxPawnSpeed;
+				// Dampen velocity magnitude based on deceleration.
+				const FVector OldVelocity = Velocity;
+				const float VelSize = FMath::Max(Velocity.Size() - FMath::Abs(InParams.Friction * Velocity.Size() + InParams.Deceleration) * InParams.DeltaSeconds, 0.f);
+				Velocity = Velocity.GetSafeNormal() * VelSize;
+
+				// Don't allow braking to lower us below max speed if we started above it.
+				if (bExceedingMaxSpeed && Velocity.SizeSquared() < FMath::Square(DesiredSpeed))
+				{
+					Velocity = OldVelocity.GetSafeNormal() * DesiredSpeed;
+				}
 			}
 		}
+		
+		Acceleration = ControlAcceleration * FMath::Abs(InParams.Acceleration);
+	}
+	else
+	{
+		UE_LOG(LogMover, Warning, TEXT("Mover Compute Velocity has recieved a bad input type and no velocity will be generated!"));
+		return FVector::ZeroVector;
 	}
 
 	// Apply acceleration and clamp velocity magnitude.
-	const float NewMaxSpeed = (IsExceedingMaxSpeed(Velocity, MaxPawnSpeed)) ? Velocity.Size() : MaxPawnSpeed;
-	Velocity += ControlAcceleration * FMath::Abs(InParams.Acceleration) * InParams.DeltaSeconds;
+	const float NewMaxSpeed = (IsExceedingMaxSpeed(Velocity, DesiredSpeed)) ? Velocity.Size() : DesiredSpeed;
+	Velocity += Acceleration * InParams.DeltaSeconds;
 	Velocity = Velocity.GetClampedToMaxSize(NewMaxSpeed);
 
 	return Velocity;
@@ -505,7 +538,7 @@ FRotator UMovementUtils::ComputeAngularVelocity(const FRotator& FromOrientation,
 }
 
 
-FVector UMovementUtils::ComputeDirectionIntent(const FVector& MoveInput, EMoveInputType MoveInputType)
+FVector UMovementUtils::ComputeDirectionIntent(const FVector& MoveInput, EMoveInputType MoveInputType, float MaxSpeed)
 {
 	FVector ResultDirIntent(FVector::ZeroVector);
 
@@ -517,8 +550,9 @@ FVector UMovementUtils::ComputeDirectionIntent(const FVector& MoveInput, EMoveIn
 			ResultDirIntent = MoveInput;
 			break;
 
-		case EMoveInputType::Velocity:	// note: if we had access to the max velocity, we could scale intent to match
-			ResultDirIntent = MoveInput.GetSafeNormal();
+		case EMoveInputType::Velocity:
+			const float IntentScalar = FMath::Clamp(MoveInput.Size() / MaxSpeed, 0.0f, 1.0f);
+			ResultDirIntent = MoveInput.GetSafeNormal() * IntentScalar;
 			break;
 	}
 
diff --git a/Engine/Plugins/Experimental/Mover/Source/Mover/Private/MoveLibrary/WaterMovementUtils.cpp b/Engine/Plugins/Experimental/Mover/Source/Mover/Private/MoveLibrary/WaterMovementUtils.cpp
--- a/Engine/Plugins/Experimental/Mover/Source/Mover/Private/MoveLibrary/WaterMovementUtils.cpp
+++ b/Engine/Plugins/Experimental/Mover/Source/Mover/Private/MoveLibrary/WaterMovementUtils.cpp
@@ -7,7 +7,7 @@ FProposedMove UWaterMovementUtils::ComputeControlledWaterMove(const FWaterMovePa
 {
 	FProposedMove OutMove;
 
-	OutMove.DirectionIntent = UMovementUtils::ComputeDirectionIntent(InParams.MoveInput, InParams.MoveInputType);
+	OutMove.DirectionIntent = UMovementUtils::ComputeDirectionIntent(InParams.MoveInput, InParams.MoveInputType, InParams.MaxSpeed);
 	OutMove.bHasDirIntent = !OutMove.DirectionIntent.IsNearlyZero();
 
 	FComputeCombinedVelocityParams ComputeCombinedVelocityParams;
diff --git a/Engine/Plugins/Experimental/Mover/Source/Mover/Public/DefaultMovementSet/Settings/CommonLegacyMovementSettings.h b/Engine/Plugins/Experimental/Mover/Source/Mover/Public/DefaultMovementSet/Settings/CommonLegacyMovementSettings.h
--- a/Engine/Plugins/Experimental/Mover/Source/Mover/Public/DefaultMovementSet/Settings/CommonLegacyMovementSettings.h
+++ b/Engine/Plugins/Experimental/Mover/Source/Mover/Public/DefaultMovementSet/Settings/CommonLegacyMovementSettings.h
@@ -51,6 +51,14 @@ public:
 	UPROPERTY(BlueprintReadWrite, EditAnywhere, Category="General", meta = (ClampMin = "0", UIMin = "0", ForceUnits = "cm/s"))
 	float MaxSpeed = 800.f;
 
+	/**
+	 * Should use acceleration for velocity based movement intent?
+	 * If true, acceleration is applied when using velocity input to reach the target velocity.
+	 * If false, velocity is set directly, disregarding acceleration.
+	 */
+	UPROPERTY(Category="General", EditAnywhere, BlueprintReadWrite, AdvancedDisplay)
+	bool bUseAccelerationForVelocityMove = true;
+	
 	/**
 	 * Setting that affects movement control. Higher values allow faster changes in direction. This can be used to simulate slippery
 	 * surfaces such as ice or oil by lowering the value (possibly based on the material the actor is standing on).
diff --git a/Engine/Plugins/Experimental/Mover/Source/Mover/Public/MoveLibrary/AirMovementUtils.h b/Engine/Plugins/Experimental/Mover/Source/Mover/Public/MoveLibrary/AirMovementUtils.h
--- a/Engine/Plugins/Experimental/Mover/Source/Mover/Public/MoveLibrary/AirMovementUtils.h
+++ b/Engine/Plugins/Experimental/Mover/Source/Mover/Public/MoveLibrary/AirMovementUtils.h
@@ -50,6 +50,9 @@ struct MOVER_API FFreeMoveParams
 
 	UPROPERTY(BlueprintReadOnly, EditAnywhere, Category = Mover)
 	FQuat WorldToGravityQuat = FQuat::Identity;
+
+	UPROPERTY(BlueprintReadOnly, EditAnywhere, Category = Mover)
+	bool bUseAccelerationForVelocityMove = true;
 };
 
 /**
diff --git a/Engine/Plugins/Experimental/Mover/Source/Mover/Public/MoveLibrary/GroundMovementUtils.h b/Engine/Plugins/Experimental/Mover/Source/Mover/Public/MoveLibrary/GroundMovementUtils.h
--- a/Engine/Plugins/Experimental/Mover/Source/Mover/Public/MoveLibrary/GroundMovementUtils.h
+++ b/Engine/Plugins/Experimental/Mover/Source/Mover/Public/MoveLibrary/GroundMovementUtils.h
@@ -64,6 +64,9 @@ struct MOVER_API FGroundMoveParams
 	
 	UPROPERTY(BlueprintReadOnly, EditAnywhere, Category = Mover)
 	FQuat WorldToGravityQuat = FQuat::Identity;
+
+	UPROPERTY(BlueprintReadOnly, EditAnywhere, Category = Mover)
+	bool bUseAccelerationForVelocityMove = true;
 };
 
 /**
diff --git a/Engine/Plugins/Experimental/Mover/Source/Mover/Public/MoveLibrary/MovementUtils.h b/Engine/Plugins/Experimental/Mover/Source/Mover/Public/MoveLibrary/MovementUtils.h
--- a/Engine/Plugins/Experimental/Mover/Source/Mover/Public/MoveLibrary/MovementUtils.h
+++ b/Engine/Plugins/Experimental/Mover/Source/Mover/Public/MoveLibrary/MovementUtils.h
@@ -66,6 +66,15 @@ struct MOVER_API FComputeVelocityParams
 	
 	UPROPERTY(BlueprintReadOnly, EditAnywhere, Category = Mover)
 	float Acceleration = 0.f;
+	
+	UPROPERTY(BlueprintReadOnly, EditAnywhere, Category = Mover)
+	FVector MoveInput = FVector::ZeroVector;
+	
+	UPROPERTY(BlueprintReadOnly, EditAnywhere, Category = Mover)
+	EMoveInputType MoveInputType = EMoveInputType::DirectionalIntent;
+	
+	UPROPERTY(BlueprintReadOnly, EditAnywhere, Category = Mover)
+	bool bUseAccelerationForVelocityMove = true;
 };
 
 // Input parameters for ComputeCombinedVelocity()
@@ -201,13 +210,12 @@ public:
 
 	/** Computes the directional movement intent based on input vector and associated type */
 	UFUNCTION(BlueprintCallable, Category = Mover)
-	static FVector ComputeDirectionIntent(const FVector& MoveInput, EMoveInputType MoveInputType);
+	static FVector ComputeDirectionIntent(const FVector& MoveInput, EMoveInputType MoveInputType, float MaxSpeed);
 
 	/** Returns whether this rotator representing angular velocity has any non-zero values. This function exists due to FRotator::IsZero queries performing undesired wrapping and clamping. */
 	UFUNCTION(BlueprintCallable, Category = Mover)
 	static bool IsAngularVelocityZero(const FRotator& AngularVelocity);
 
-
 	// Internal functions - not meant to be called outside of this library
 	
 	/** Internal function that other move functions use to perform all actual component movement and retrieve results
