Mover: added option to constrain root motion rotation from attributes to the movement plane, to help mitigate small errors that can build up and make large pitching/rolling movements. Enabled by default.

diff --git a/Engine/Plugins/Experimental/Mover/Source/Mover/Private/DefaultMovementSet/LayeredMoves/RootMotionAttributeLayeredMove.cpp b/Engine/Plugins/Experimental/Mover/Source/Mover/Private/DefaultMovementSet/LayeredMoves/RootMotionAttributeLayeredMove.cpp
--- a/Engine/Plugins/Experimental/Mover/Source/Mover/Private/DefaultMovementSet/LayeredMoves/RootMotionAttributeLayeredMove.cpp
+++ b/Engine/Plugins/Experimental/Mover/Source/Mover/Private/DefaultMovementSet/LayeredMoves/RootMotionAttributeLayeredMove.cpp
@@ -133,7 +133,14 @@ bool FLayeredMove_RootMotionAttribute::GenerateMove(const FMoverTickStartData& S
 		// Note that we're forcing the use of the sync state's actor transform data. This is necessary when the movement simulation 
 		// is running ahead of the actor's visual representation and may be rotated differently, such as in an async physics sim.
 		const FTransform SimActorTransform = FTransform(SyncState->GetOrientation_WorldSpace().Quaternion(), SyncState->GetLocation_WorldSpace());
-		const FTransform WorldSpaceRootMotion = MoverComp->ConvertLocalRootMotionToWorld(LocalRootMotion, DeltaSeconds, &SimActorTransform, (bHasValidWarpingContext ? &WarpingContext : nullptr));
+		FTransform WorldSpaceRootMotion = MoverComp->ConvertLocalRootMotionToWorld(LocalRootMotion, DeltaSeconds, &SimActorTransform, (bHasValidWarpingContext ? &WarpingContext : nullptr));
+
+		if (bConstrainWorldRotToMovementPlane)
+		{
+			// up direction is the locked axis. The original rotated direction is effectively projected onto the movement plane defined by the up direction.
+			const FMatrix ConstrainedRot = FRotationMatrix::MakeFromZX(MoverComp->GetUpDirection(), WorldSpaceRootMotion.GetRotation().GetForwardVector());
+			WorldSpaceRootMotion.SetRotation(ConstrainedRot.ToQuat());
+		}
 
 		OutProposedMove = FProposedMove();
 		OutProposedMove.MixMode = MixMode;
diff --git a/Engine/Plugins/Experimental/Mover/Source/Mover/Public/DefaultMovementSet/LayeredMoves/RootMotionAttributeLayeredMove.h b/Engine/Plugins/Experimental/Mover/Source/Mover/Public/DefaultMovementSet/LayeredMoves/RootMotionAttributeLayeredMove.h
--- a/Engine/Plugins/Experimental/Mover/Source/Mover/Public/DefaultMovementSet/LayeredMoves/RootMotionAttributeLayeredMove.h
+++ b/Engine/Plugins/Experimental/Mover/Source/Mover/Public/DefaultMovementSet/LayeredMoves/RootMotionAttributeLayeredMove.h
@@ -13,11 +13,15 @@
 USTRUCT(BlueprintType)
 struct MOVER_API FLayeredMove_RootMotionAttribute : public FLayeredMoveBase
 {
-	GENERATED_USTRUCT_BODY()
+	GENERATED_BODY()
 
 	FLayeredMove_RootMotionAttribute();
 	virtual ~FLayeredMove_RootMotionAttribute() {}
 
+	// If true, any root motion rotations will be projected onto the movement plane (in worldspace), relative to the "up" direction. Otherwise, they'll be taken as-is.
+	UPROPERTY(EditAnywhere, BlueprintReadWrite, Category = Mover)
+	bool bConstrainWorldRotToMovementPlane = true;
+
 protected:
 	// These member variables are NOT replicated. They are used if we rollback and resimulate when the root motion attribute is no longer in sync.
 	bool bDidAttrHaveRootMotionForResim = false;
