Mover: adding optional input cmd readout to Gameplay Debugger printing

diff --git a/Engine/Plugins/Experimental/Mover/Source/Mover/Private/Debug/GameplayDebuggerCategory_Mover.cpp b/Engine/Plugins/Experimental/Mover/Source/Mover/Private/Debug/GameplayDebuggerCategory_Mover.cpp
--- a/Engine/Plugins/Experimental/Mover/Source/Mover/Private/Debug/GameplayDebuggerCategory_Mover.cpp
+++ b/Engine/Plugins/Experimental/Mover/Source/Mover/Private/Debug/GameplayDebuggerCategory_Mover.cpp
@@ -19,6 +19,8 @@ namespace FGameplayDebuggerCategoryTweakables
 	bool bShowTrajectory = true;
 	bool bShowTrail = false;
 	bool bShowCorrections = true;
+	bool bShowStateArrowViz = true;			// Toggles state parameter visualization with arrows draw on the actor
+	bool bShowInputArrowViz = false;		// Toggles input parameter visualization with arrows draw on the actor
 }
 
 namespace
@@ -44,6 +46,15 @@ namespace
 		FAutoConsoleVariableRef(TEXT("mover.debug.ShowCorrections"),
 			FGameplayDebuggerCategoryTweakables::bShowCorrections,
 			TEXT("Shows network corrections of the selected actor. NOTE: This should only be used on actors controlled by the server. For showing corrections on the local player use Mover.LocalPlayer.ShowCorrections"), ECVF_Default),
+
+		FAutoConsoleVariableRef(TEXT("mover.debug.ShowStateArrows"),
+			FGameplayDebuggerCategoryTweakables::bShowStateArrowViz,
+			TEXT("If enabled, in-world arrows will be drawn to show certain state information in the Gameplay Debugger visualization."), ECVF_Default),
+
+		FAutoConsoleVariableRef(TEXT("mover.debug.ShowInputArrows"),
+			FGameplayDebuggerCategoryTweakables::bShowInputArrowViz,
+			TEXT("If enabled, in-world arrows will be drawn to show certain input cmd information in the Gameplay Debugger visualization."), ECVF_Default),
+
 	};
 }
 
@@ -99,6 +110,10 @@ void FGameplayDebuggerCategory_Mover::CollectData(APlayerController* OwnerPC, AA
 	DataPack.ActiveLayeredMoves.Empty();
 	DataPack.ActiveModifiers.Empty();
 	DataPack.ModeMap.Empty();
+	DataPack.MoveInputType = 0;
+	DataPack.MoveInput = FVector::ZeroVector;
+	DataPack.OrientIntentDir = FVector::ZeroVector;
+	DataPack.SuggestedModeName = FString("invalid");
 
 	if (MyMoverComponent)
 	{
@@ -142,6 +157,20 @@ void FGameplayDebuggerCategory_Mover::CollectData(APlayerController* OwnerPC, AA
 				DataPack.ActiveModifiers.Add(*it->Get()->ToSimpleString());
 			}
 		}
+
+		if (MyMoverComponent->HasValidCachedInputCmd())
+		{
+			const FMoverInputCmdContext& LastInputCmd = MyMoverComponent->GetLastInputCmd();
+
+			if (const FCharacterDefaultInputs* DefaultInputs = LastInputCmd.InputCollection.FindDataByType<FCharacterDefaultInputs>())
+			{
+				DataPack.MoveInputType = (int8)DefaultInputs->GetMoveInputType();
+				DataPack.MoveInput = DefaultInputs->GetMoveInput_WorldSpace();
+				DataPack.OrientIntentDir = DefaultInputs->GetOrientationIntentDir_WorldSpace();
+				DataPack.SuggestedModeName = DefaultInputs->SuggestedMovementMode.ToString();
+			}
+		}
+
 	}
 }
 
@@ -158,17 +187,30 @@ void FGameplayDebuggerCategory_Mover::DrawData(APlayerController* OwnerPC, FGame
 		DrawInWorldInfo(*FocusedActor, CanvasContext);
 	}
 	
-	CanvasContext.Printf(TEXT("{yellow}%s\n{grey}Local Role: {white}%s\n{grey}Mode: {white}%s\n{grey}Velocity: {white}%s\n{grey}Speed: {white}%.2f\n{yellow}Active Moves: {white}\n%s\n{yellow}Active Modifiers: {white}\n%s\n{yellow}Mode Map: \n{white}%s\n{yellow}Active Transitions: {white}\n%s"),
+	CanvasContext.Printf(TEXT("{yellow}%s\n{grey}Local Role: {white}%s\n{grey}Mode: {white}%s\n{grey}Velocity: {white}%s\n{grey}Speed: {white}%.2f"),
 		*DataPack.PawnName,
 		*DataPack.LocalRole,
 		*DataPack.MovementModeName,
 		*DataPack.Velocity.ToString(),
-		 DataPack.Velocity.Length(),
+		 DataPack.Velocity.Length()
+		);
+
+	if (DataPack.MoveInputType > 0)
+	{
+		CanvasContext.Printf(TEXT("{grey}Move Input Type: {white}%i  {grey}Vec: {white}%s\n{grey}Input Suggested Mode: {white}%s\n{grey}Input Orient Intent: {white}%s"),
+			DataPack.MoveInputType,
+			*DataPack.MoveInput.ToString(),
+			*DataPack.SuggestedModeName,
+			*DataPack.OrientIntentDir.ToString()
+		);
+	}
+
+	CanvasContext.Printf(TEXT("{yellow}Active Moves: {white}\n%s\n{yellow}Active Modifiers: {white}\n%s\n{yellow}Mode Map: \n{white}%s\n{yellow}Active Transitions: {white}\n%s"),
 		*FString::JoinBy(DataPack.ActiveLayeredMoves, TEXT("\n"), [](FString MoveAsString) { return MoveAsString; }),
 		*FString::JoinBy(DataPack.ActiveModifiers, TEXT("\n"), [](FString ModifierAsString) { return ModifierAsString; }),
 		*FString::JoinBy(DataPack.ModeMap, TEXT("\n"), [](FString ModeMappingAsString) { return ModeMappingAsString; }),
 		*FString::JoinBy(DataPack.ActiveTransitions, TEXT("\n"), [](FString TransitionAsString) { return TransitionAsString; })
-		);
+	);
 }
 
 
@@ -219,6 +261,9 @@ void FGameplayDebuggerCategory_Mover::DrawInWorldInfo(AActor& DebugActor, FGamep
 	const FVector ActorMidLocation = DebugActor.GetActorLocation();
 	const FVector ActorLowLocation = ActorMidLocation - FVector(0,0,DebugActor.GetSimpleCollisionHalfHeight()*0.95f);	// slightly above lowest point
 
+	const FVector NudgeUp(0.0, 0.0, 2.0);
+
+
 	// Draw approximate bounds
 	if (CanvasContext.IsLocationVisible(ActorMidLocation))
 	{
@@ -230,40 +275,66 @@ void FGameplayDebuggerCategory_Mover::DrawInWorldInfo(AActor& DebugActor, FGamep
 			FColor::Green);
 	}
 
-	// Draw arrow showing movement intent (direction + magnitude)
-	if (CanvasContext.IsLocationVisible(ActorLowLocation))
+	if (FGameplayDebuggerCategoryTweakables::bShowStateArrowViz)
 	{
-		DrawDebugDirectionalArrow(MyWorld,
-			ActorMidLocation,
-			ActorMidLocation + (DataPack.MoveIntent*FGameplayDebuggerCategoryTweakables::MaxMoveIntentDrawLength),
-			40.f, FColor::Blue, false, -1.f, 0, 3.f);
-	}
-
-	// Draw overlaid arrows showing target orientation and actual
-	if (MoverComp && MoverComp->HasValidCachedState())
-	{
-		const FMoverSyncState& LastState = MoverComp->GetSyncState();
-		const FMoverDefaultSyncState* MoverState = LastState.SyncStateCollection.FindDataByType<FMoverDefaultSyncState>();
+		// Draw arrow showing movement intent (direction + magnitude)
+		if (CanvasContext.IsLocationVisible(ActorLowLocation))
+		{
+			DrawDebugDirectionalArrow(MyWorld,
+				ActorMidLocation,
+				ActorMidLocation + (DataPack.MoveIntent*FGameplayDebuggerCategoryTweakables::MaxMoveIntentDrawLength),
+				40.f, FColor::Blue, false, -1.f, 0, 3.f);
+		}
 
-		if (MoverState)
+		// Draw overlaid arrows showing target orientation and actual
+		if (MoverComp && MoverComp->HasValidCachedState())
 		{
-			const FVector ActualFacingDir = MoverState->GetOrientation_WorldSpace().Vector();
-			const FVector TargetFacingDir = MoverComp->GetTargetOrientation().Vector();
+			const FMoverSyncState& LastState = MoverComp->GetSyncState();
+			const FMoverDefaultSyncState* MoverState = LastState.SyncStateCollection.FindDataByType<FMoverDefaultSyncState>();
 
-			DrawDebugDirectionalArrow(MyWorld,
-				ActorLowLocation,
-				ActorLowLocation + (TargetFacingDir * FGameplayDebuggerCategoryTweakables::OrientationDrawLength),
-				30.f, FColor::Yellow, false, -1.f, 0, 2.5f);
+			if (MoverState)
+			{
+				const FVector ActualFacingDir = MoverState->GetOrientation_WorldSpace().Vector();
+				const FVector TargetFacingDir = MoverComp->GetTargetOrientation().Vector();
 
-			const FVector NudgeUp(0.0, 0.0, 2.0);
-			DrawDebugDirectionalArrow(MyWorld,
-				ActorLowLocation + NudgeUp,
-				ActorLowLocation + NudgeUp + (ActualFacingDir * FGameplayDebuggerCategoryTweakables::OrientationDrawLength * 0.9),
-				10.f, FColor::Green, false, -1.f, 0, 1.25f);
+				DrawDebugDirectionalArrow(MyWorld,
+					ActorLowLocation,
+					ActorLowLocation + (TargetFacingDir * FGameplayDebuggerCategoryTweakables::OrientationDrawLength),
+					30.f, FColor::Yellow, false, -1.f, 0, 2.5f);
+
+				DrawDebugDirectionalArrow(MyWorld,
+					ActorLowLocation + NudgeUp,
+					ActorLowLocation + NudgeUp + (ActualFacingDir * FGameplayDebuggerCategoryTweakables::OrientationDrawLength * 0.9),
+					10.f, FColor::Green, false, -1.f, 0, 1.25f);
 
+			}
 		}
 	}
 
+	if (FGameplayDebuggerCategoryTweakables::bShowInputArrowViz)
+	{
+		// Draw arrows showing what the input cmds want to do
+		if (CanvasContext.IsLocationVisible(ActorMidLocation))
+		{
+			if (!DataPack.MoveInput.IsNearlyZero())
+			{
+				DrawDebugDirectionalArrow(MyWorld,
+					ActorMidLocation,
+					ActorMidLocation + (DataPack.MoveInput.GetSafeNormal() * FGameplayDebuggerCategoryTweakables::MaxMoveIntentDrawLength),
+					40.f, FColor::Cyan, false, -1.f, 0, 3.f);
+			}
+
+			if (!DataPack.OrientIntentDir.IsNearlyZero())
+			{
+				DrawDebugDirectionalArrow(MyWorld,
+					ActorMidLocation + NudgeUp,
+					ActorMidLocation + NudgeUp + (DataPack.OrientIntentDir.GetSafeNormal() * FGameplayDebuggerCategoryTweakables::MaxMoveIntentDrawLength),
+					30.f, FColor::Orange, false, -1.f, 0, 3.f);
+			}
+		}
+
+	}
+
 }
 
 void FGameplayDebuggerCategory_Mover::FRepData::Serialize(FArchive& Ar)
@@ -278,6 +349,11 @@ void FGameplayDebuggerCategory_Mover::FRepData::Serialize(FArchive& Ar)
 	Ar << ActiveModifiers;
 	Ar << ModeMap;
 	Ar << ActiveTransitions;
+	Ar << MoveInputType;
+	Ar << MoveInput;
+	Ar << OrientIntentDir;
+	Ar << SuggestedModeName;
+
 }
 
 #endif // WITH_GAMEPLAY_DEBUGGER
\ No newline at end of file
diff --git a/Engine/Plugins/Experimental/Mover/Source/Mover/Private/Debug/GameplayDebuggerCategory_Mover.h b/Engine/Plugins/Experimental/Mover/Source/Mover/Private/Debug/GameplayDebuggerCategory_Mover.h
--- a/Engine/Plugins/Experimental/Mover/Source/Mover/Private/Debug/GameplayDebuggerCategory_Mover.h
+++ b/Engine/Plugins/Experimental/Mover/Source/Mover/Private/Debug/GameplayDebuggerCategory_Mover.h
@@ -27,6 +27,8 @@ protected:
 	{
 		FString PawnName;
 		FString LocalRole;
+
+		// State
 		FString MovementModeName;
 		FString MovementBaseInfo;
 		FVector Velocity;
@@ -36,6 +38,12 @@ protected:
 		TArray<FString> ActiveTransitions;
 		TArray<FString> ActiveModifiers;
 
+		// Last input
+		int8 MoveInputType;
+		FVector MoveInput;
+		FVector OrientIntentDir;
+		FString SuggestedModeName;
+
 		FRepData() {}
 
 		void Serialize(FArchive& Ar);
