Physics Mover - Fix for applying a positive gravity override

diff --git a/Engine/Plugins/Experimental/Mover/Source/Mover/Private/Backends/MoverNetworkPhysicsLiaison.cpp b/Engine/Plugins/Experimental/Mover/Source/Mover/Private/Backends/MoverNetworkPhysicsLiaison.cpp
--- a/Engine/Plugins/Experimental/Mover/Source/Mover/Private/Backends/MoverNetworkPhysicsLiaison.cpp
+++ b/Engine/Plugins/Experimental/Mover/Source/Mover/Private/Backends/MoverNetworkPhysicsLiaison.cpp
@@ -1112,7 +1112,7 @@ void UMoverNetworkPhysicsLiaisonComponent::OnPreSimulate_Internal(const FPhysics
 		OutputSyncState.SetMovementBase(nullptr);
 
 		ConstraintHandle->SetData({
-			MoverComp->GetUpDirection(),
+			ConstraintHandle->GetSettings().VerticalAxis,
 			Chaos::FVec3::ZeroVector,
 			0.0,
 			1.0e10,
