Mover: reworking standalone mode liaison to break operations into distinct batched tick tasks with dependencies on each other. This opens the door to cross-system dependencies with Mover as well.

diff --git a/Engine/Plugins/Experimental/Mover/Source/Mover/Private/Backends/MoverStandaloneLiaison.cpp b/Engine/Plugins/Experimental/Mover/Source/Mover/Private/Backends/MoverStandaloneLiaison.cpp
--- a/Engine/Plugins/Experimental/Mover/Source/Mover/Private/Backends/MoverStandaloneLiaison.cpp
+++ b/Engine/Plugins/Experimental/Mover/Source/Mover/Private/Backends/MoverStandaloneLiaison.cpp
@@ -2,7 +2,8 @@
 
 #include "Backends/MoverStandaloneLiaison.h"
 #include "Engine/World.h"
-#include "GameFramework/Actor.h"
+#include "GameFramework/Pawn.h"
+#include "GameFramework/Controller.h"
 #include "MoverComponent.h"
 
 
@@ -10,7 +11,24 @@ UMoverStandaloneLiaisonComponent::UMoverStandaloneLiaisonComponent(const FObject
 	: Super(ObjectInitializer)
 {
 	PrimaryComponentTick.TickGroup = TG_PrePhysics;
-	PrimaryComponentTick.bCanEverTick = true;
+	PrimaryComponentTick.bCanEverTick = false;
+
+	ProduceInputTickFunction.bCanEverTick = true;
+	ProduceInputTickFunction.bStartWithTickEnabled = true;
+	ProduceInputTickFunction.SetTickFunctionEnable(true);
+	ProduceInputTickFunction.TickGroup = TG_PrePhysics;
+	ProduceInputTickFunction.bAllowTickBatching = true;
+
+	SimulateMovementTickFunction.bCanEverTick = true;
+	SimulateMovementTickFunction.bStartWithTickEnabled = true;
+	SimulateMovementTickFunction.SetTickFunctionEnable(true);
+	SimulateMovementTickFunction.TickGroup = TG_PrePhysics;
+
+	ApplyStateTickFunction.bCanEverTick = true;
+	ApplyStateTickFunction.bStartWithTickEnabled = true;
+	ApplyStateTickFunction.SetTickFunctionEnable(true);
+	ApplyStateTickFunction.TickGroup = TG_PrePhysics;
+	ApplyStateTickFunction.bAllowTickBatching = true;
 
 	bWantsInitializeComponent = true;
 	bAutoActivate = true;
@@ -49,7 +67,7 @@ void UMoverStandaloneLiaisonComponent::BeginPlay()
 	Super::BeginPlay();
 
 	CurrentSimTimeMs = GetWorld()->GetTimeSeconds() * 1000.0;
-	CurrentSimFrame = 0;
+	CurrentSimFrame = GFrameCounter;
 
 	if (const AActor* OwnerActor = GetOwner())
 	{
@@ -61,52 +79,232 @@ void UMoverStandaloneLiaisonComponent::BeginPlay()
 
 			MoverComp->InitMoverSimulation();
 
-			MoverComp->InitializeSimulationState(OUT &CachedLastSyncState, OUT &CachedLastAuxState);
+			MoverComp->InitializeSimulationState(OUT & CachedLastSyncState, OUT & CachedLastAuxState);
 		}
 		else
 		{
 			ensureMsgf(MoverComp, TEXT("Owning actor %s does not have a MoverComponent."), *GetNameSafe(GetOwner()));
+
+			// Disable tick if we don't have mover for some reason..
+			SetComponentTickEnabled(false);
+			ProduceInputTickFunction.SetTickFunctionEnable(false);
+			SimulateMovementTickFunction.SetTickFunctionEnable(false);
+			ApplyStateTickFunction.SetTickFunctionEnable(false);
 		}
-	}	
+	}
 }
 
-
-void UMoverStandaloneLiaisonComponent::TickComponent(float DeltaTime, enum ELevelTick TickType, FActorComponentTickFunction* ThisTickFunction)
+void UMoverStandaloneLiaisonComponent::RegisterComponentTickFunctions(bool bRegister)
 {
-	Super::TickComponent(DeltaTime, TickType, ThisTickFunction);
+	Super::RegisterComponentTickFunctions(bRegister);
 
-	if (MoverComp)
+	if (bRegister)
+	{
+		if (SetupActorComponentTickFunction(&ProduceInputTickFunction))
+		{
+			// Set up tick dependencies so this order always holds true:
+			//   1) Controller update, which triggers input events
+			//   2) Mover input production for the next movement simulation step
+			//   3) Perform simulation step to generate new state
+			//   4) Apply new state to actor/components
+
+			ProduceInputTickFunction.Target = this;
+			
+			// Input production should always wait for the controller update, and we will watch for controller changes
+			if (APawn* OwnerPawn = Cast<APawn>(GetOwner()))
+			{
+				if (AController* OwnerController = OwnerPawn->GetController())
+				{
+					ProduceInputTickFunction.AddPrerequisite(OwnerController, OwnerController->PrimaryActorTick);
+				}
+
+				OwnerPawn->ReceiveControllerChangedDelegate.AddDynamic(this, &UMoverStandaloneLiaisonComponent::OnControllerChanged);
+			}
+
+			if (SetupActorComponentTickFunction(&SimulateMovementTickFunction))
+			{
+				SimulateMovementTickFunction.Target = this;
+				SimulateMovementTickFunction.AddPrerequisite(this, this->ProduceInputTickFunction);
+
+				if (SetupActorComponentTickFunction(&ApplyStateTickFunction))
+				{
+					ApplyStateTickFunction.Target = this;
+					ApplyStateTickFunction.AddPrerequisite(this, this->SimulateMovementTickFunction);
+				}
+			}
+		}
+	}
+	else
 	{
-		// TODO Here is where we might accumulate time and perform fixed tick updates
+		if (APawn* OwnerPawn = Cast<APawn>(GetOwner()))
+		{
+			if (AController* OwnerController = OwnerPawn->GetController())
+			{
+				ProduceInputTickFunction.RemovePrerequisite(OwnerController, OwnerController->PrimaryActorTick);
+			}
 
-		const int32 DeltaTimeMs = DeltaTime * 1000.f;
+			OwnerPawn->ReceiveControllerChangedDelegate.RemoveDynamic(this, &UMoverStandaloneLiaisonComponent::OnControllerChanged);
+		}
 
-		CurrentSimTimeMs = GetWorld()->GetTimeSeconds() * 1000.0;
-		CurrentSimFrame = GFrameCounter;
+		if (ProduceInputTickFunction.IsTickFunctionRegistered())
+		{
+			ProduceInputTickFunction.UnRegisterTickFunction();
+		}
 
-		FMoverInputCmdContext InputCmd;
-		MoverComp->ProduceInput(DeltaTimeMs, OUT &InputCmd);
+		if (SimulateMovementTickFunction.IsTickFunctionRegistered())
+		{
+			SimulateMovementTickFunction.UnRegisterTickFunction();
+		}
 
-		FMoverTimeStep TimeStep;
-		TimeStep.ServerFrame = CurrentSimFrame;
-		TimeStep.BaseSimTimeMs = CurrentSimTimeMs;
-		TimeStep.StepMs = FMath::Floor(DeltaTime * 1000.f);
+		if (ApplyStateTickFunction.IsTickFunctionRegistered())
+		{
+			ApplyStateTickFunction.UnRegisterTickFunction();
+		}
+	}
+}
 
-		FMoverTickStartData StartData;
-		FMoverTickEndData EndData;
 
-		StartData.InputCmd = InputCmd;
-		StartData.SyncState = CachedLastSyncState;
-		StartData.AuxState = CachedLastAuxState;
+void UMoverStandaloneLiaisonComponent::TickInputProduction(float DeltaSeconds)
+{
+	const int32 DeltaTimeMs = DeltaSeconds * 1000.f;
 
+	CurrentSimTimeMs = GetWorld()->GetTimeSeconds() * 1000.0;
+	CurrentSimFrame = GFrameCounter;
+
+	FMoverInputCmdContext InputCmd;
+	{
+		SCOPED_NAMED_EVENT(StandaloneMoverBackendLiaison_Tick_ProduceInput, FColor::Yellow);
+		MoverComp->ProduceInput(DeltaTimeMs, OUT &LastProducedInputCmd);
+	}
+}
+
+void UMoverStandaloneLiaisonComponent::TickMovementSimulation(float DeltaSeconds)
+{
+	FMoverTimeStep TimeStep;
+	TimeStep.ServerFrame = CurrentSimFrame;
+	TimeStep.BaseSimTimeMs = CurrentSimTimeMs;
+	TimeStep.StepMs = FMath::Floor(DeltaSeconds * 1000.f);
+
+	FMoverTickStartData StartData;
+	FMoverTickEndData EndData;
+
+	StartData.InputCmd = LastProducedInputCmd;
+	StartData.SyncState = CachedLastSyncState;
+	StartData.AuxState = CachedLastAuxState;
+
+	{
+		SCOPED_NAMED_EVENT(StandaloneMoverBackendLiaison_Tick_SimulationTick, FColor::Blue);
 		MoverComp->SimulationTick(TimeStep, StartData, OUT EndData);
+	}
+
+	LastUsedInputCmd = StartData.InputCmd;
+	CachedLastSyncState = EndData.SyncState;
+	CachedLastAuxState = EndData.AuxState;
+}
 
-		CachedLastSyncState = EndData.SyncState;
-		CachedLastAuxState = EndData.AuxState;
+void UMoverStandaloneLiaisonComponent::TickApplySimulationState(float DeltaSeconds)
+{
+	{
+		SCOPED_NAMED_EVENT(SilverMoverBackendLiaison_Tick_FinalizeFrame, FColor::Green);
+		MoverComp->FinalizeFrame(&CachedLastSyncState, &CachedLastAuxState);
+	}
+}
 
-		MoverComp->FinalizeFrame(&EndData.SyncState, &EndData.AuxState);
+void UMoverStandaloneLiaisonComponent::OnControllerChanged(APawn* Pawn, AController* OldController, AController* NewController)
+{
+	if (OldController)
+	{
+		ProduceInputTickFunction.RemovePrerequisite(OldController, OldController->PrimaryActorTick);
 	}
 
+	if (NewController)
+	{
+		ProduceInputTickFunction.AddPrerequisite(NewController, NewController->PrimaryActorTick);
+	}
+}
+
+
+// FMoverStandaloneProduceInputTickFunction ---------------------
+
+void FMoverStandaloneProduceInputTickFunction::ExecuteTick(float DeltaTime, enum ELevelTick TickType, ENamedThreads::Type CurrentThread, const FGraphEventRef& MyCompletionGraphEvent)
+{
+	FActorComponentTickFunction::ExecuteTickHelper(Target.Get(), /*bTickInEditor=*/ false, DeltaTime, TickType, [this](float DilatedTime)
+		{
+			if (Target.IsValid())
+			{
+				Target->TickInputProduction(DilatedTime);
+			}
+		});
+}
+
+FString FMoverStandaloneProduceInputTickFunction::DiagnosticMessage()
+{
+	return Target->GetFullName() + TEXT("[UMoverStandaloneLiaisonComponent::ProduceInputTick]");
+}
+
+FName FMoverStandaloneProduceInputTickFunction::DiagnosticContext(bool bDetailed)
+{
+	if (bDetailed)
+	{
+		return FName(*FString::Printf(TEXT("MoverStandaloneProduceInputTickFunction/%s"), *GetFullNameSafe(Target.Get())));
+	}
+
+	return FName(TEXT("MoverStandaloneProduceInputTickFunction"));
+}
+
+// FMoverStandaloneSimulateMovementTickFunction ---------------------
+
+void FMoverStandaloneSimulateMovementTickFunction::ExecuteTick(float DeltaTime, enum ELevelTick TickType, ENamedThreads::Type CurrentThread, const FGraphEventRef& MyCompletionGraphEvent)
+{
+	FActorComponentTickFunction::ExecuteTickHelper(Target.Get(), /*bTickInEditor=*/ false, DeltaTime, TickType, [this](float DilatedTime)
+		{
+			if (Target.IsValid())
+			{
+				Target->TickMovementSimulation(DilatedTime);
+			}
+		});
 }
 
+FString FMoverStandaloneSimulateMovementTickFunction::DiagnosticMessage()
+{
+	return Target->GetFullName() + TEXT("[UMoverStandaloneLiaisonComponent::SimulateMovement]");
+}
+
+FName FMoverStandaloneSimulateMovementTickFunction::DiagnosticContext(bool bDetailed)
+{
+	if (bDetailed)
+	{
+		return FName(*FString::Printf(TEXT("MoverStandaloneSimulateMovementTickFunction/%s"), *GetFullNameSafe(Target.Get())));
+	}
+
+	return FName(TEXT("MoverStandaloneSimulateMovementTickFunction"));
+}
+
+
+// FMoverStandaloneApplyStateTickFunction ---------------------
+
+void FMoverStandaloneApplyStateTickFunction::ExecuteTick(float DeltaTime, enum ELevelTick TickType, ENamedThreads::Type CurrentThread, const FGraphEventRef& MyCompletionGraphEvent)
+{
+	FActorComponentTickFunction::ExecuteTickHelper(Target.Get(), /*bTickInEditor=*/ false, DeltaTime, TickType, [this](float DilatedTime)
+		{
+			if (Target.IsValid())
+			{
+				Target->TickApplySimulationState(DilatedTime);
+			}
+		});
+}
+
+FString FMoverStandaloneApplyStateTickFunction::DiagnosticMessage()
+{
+	return Target->GetFullName() + TEXT("[UMoverStandaloneLiaisonComponent::ApplyState]");
+}
+
+FName FMoverStandaloneApplyStateTickFunction::DiagnosticContext(bool bDetailed)
+{
+	if (bDetailed)
+	{
+		return FName(*FString::Printf(TEXT("MoverStandaloneApplyStateTickFunction/%s"), *GetFullNameSafe(Target.Get())));
+	}
 
+	return FName(TEXT("MoverStandaloneApplyStateTickFunction"));
+}
\ No newline at end of file
diff --git a/Engine/Plugins/Experimental/Mover/Source/Mover/Public/Backends/MoverStandaloneLiaison.h b/Engine/Plugins/Experimental/Mover/Source/Mover/Public/Backends/MoverStandaloneLiaison.h
--- a/Engine/Plugins/Experimental/Mover/Source/Mover/Public/Backends/MoverStandaloneLiaison.h
+++ b/Engine/Plugins/Experimental/Mover/Source/Mover/Public/Backends/MoverStandaloneLiaison.h
@@ -2,19 +2,95 @@
 
 #pragma once
 
-#include "CoreMinimal.h"
 #include "Components/ActorComponent.h"
 #include "Backends/MoverBackendLiaison.h"
-#include "MoverTypes.h"
-
 #include "MoverStandaloneLiaison.generated.h"
 
 class UMoverComponent;
+class UMoverStandaloneLiaisonComponent;
+class AController;
+
+
+
+
+
+// Tick task for producing input before the next movement simulation step
+USTRUCT()
+struct FMoverStandaloneProduceInputTickFunction : public FTickFunction
+{
+	GENERATED_BODY()
+
+	/** Standalone liaison that is the target of this tick **/
+	TWeakObjectPtr<UMoverStandaloneLiaisonComponent> Target;
+
+	virtual void ExecuteTick(float DeltaTime, enum ELevelTick TickType, ENamedThreads::Type CurrentThread, const FGraphEventRef& MyCompletionGraphEvent) override;
+	virtual FString DiagnosticMessage() override;
+	virtual FName DiagnosticContext(bool bDetailed) override;
+};
+
+template<>
+struct TStructOpsTypeTraits<FMoverStandaloneProduceInputTickFunction> : public TStructOpsTypeTraitsBase2<FMoverStandaloneProduceInputTickFunction>
+{
+	enum
+	{
+		WithCopy = false
+	};
+};
+
+
+// Tick task for advancing the movement simulation step, after input has been produced
+USTRUCT()
+struct FMoverStandaloneSimulateMovementTickFunction : public FTickFunction
+{
+	GENERATED_BODY()
+
+	/** Standalone liaison that is the target of this tick **/
+	TWeakObjectPtr<UMoverStandaloneLiaisonComponent> Target;
+
+	virtual void ExecuteTick(float DeltaTime, enum ELevelTick TickType, ENamedThreads::Type CurrentThread, const FGraphEventRef& MyCompletionGraphEvent) override;
+	virtual FString DiagnosticMessage() override;
+	virtual FName DiagnosticContext(bool bDetailed) override;
+};
+
+template<>
+struct TStructOpsTypeTraits<FMoverStandaloneSimulateMovementTickFunction> : public TStructOpsTypeTraitsBase2<FMoverStandaloneSimulateMovementTickFunction>
+{
+	enum
+	{
+		WithCopy = false
+	};
+};
+
+
+// Tick task for applying the new simulation state to the actor/components, after movement has been simulated
+USTRUCT()
+struct FMoverStandaloneApplyStateTickFunction : public FTickFunction
+{
+	GENERATED_BODY()
+
+	/** Standalone liaison that is the target of this tick **/
+	TWeakObjectPtr<UMoverStandaloneLiaisonComponent> Target;
+
+	virtual void ExecuteTick(float DeltaTime, enum ELevelTick TickType, ENamedThreads::Type CurrentThread, const FGraphEventRef& MyCompletionGraphEvent) override;
+	virtual FString DiagnosticMessage() override;
+	virtual FName DiagnosticContext(bool bDetailed) override;
+};
+
+template<>
+struct TStructOpsTypeTraits<FMoverStandaloneApplyStateTickFunction> : public TStructOpsTypeTraitsBase2<FMoverStandaloneApplyStateTickFunction>
+{
+	enum
+	{
+		WithCopy = false
+	};
+};
+
+
+
 
 /**
  * MoverStandaloneLiaison: this component acts as a backend driver for an actor's Mover component, for use in Standalone (non-networked) games.
  * This class is set on a Mover component as the "back end".
- * TODO: Support options for fixed ticking rates and state smoothing
  */
 UCLASS()
 class MOVER_API UMoverStandaloneLiaisonComponent : public UActorComponent, public IMoverBackendLiaisonInterface
@@ -32,10 +108,18 @@ public:
 	// End IMoverBackendLiaisonInterface
 
 	// Begin UActorComponent interface
+	virtual void RegisterComponentTickFunctions(bool bRegister) override;
 	virtual void BeginPlay() override;
-	virtual void TickComponent(float DeltaTime, enum ELevelTick TickType, FActorComponentTickFunction* ThisTickFunction) override;
 	// End UActorComponent interface
 
+protected:
+	void TickInputProduction(float DeltaSeconds);
+	void TickMovementSimulation(float DeltaSeconds);
+	void TickApplySimulationState(float DeltaSeconds);
+
+	// Called when controller changes, used to manage ticking dependencies
+	UFUNCTION()
+	virtual void OnControllerChanged(APawn* Pawn, AController* OldController, AController* NewController);
 
 protected:
 	TObjectPtr<UMoverComponent> MoverComp;	// the component that we're in charge of driving
@@ -43,6 +127,16 @@ protected:
 	double CurrentSimTimeMs;
 	int32 CurrentSimFrame;
 
+	FMoverInputCmdContext LastProducedInputCmd;
+	FMoverInputCmdContext LastUsedInputCmd;
 	FMoverSyncState CachedLastSyncState;
 	FMoverAuxStateContext CachedLastAuxState;
-};
\ No newline at end of file
+
+	FMoverStandaloneProduceInputTickFunction ProduceInputTickFunction;
+	FMoverStandaloneSimulateMovementTickFunction SimulateMovementTickFunction;
+	FMoverStandaloneApplyStateTickFunction ApplyStateTickFunction;
+
+	friend struct FMoverStandaloneProduceInputTickFunction;
+	friend struct FMoverStandaloneSimulateMovementTickFunction;
+	friend struct FMoverStandaloneApplyStateTickFunction;
+};
diff --git a/Engine/Plugins/Experimental/Mover/Source/Mover/Public/MoverTypes.h b/Engine/Plugins/Experimental/Mover/Source/Mover/Public/MoverTypes.h
--- a/Engine/Plugins/Experimental/Mover/Source/Mover/Public/MoverTypes.h
+++ b/Engine/Plugins/Experimental/Mover/Source/Mover/Public/MoverTypes.h
@@ -19,6 +19,21 @@ MOVER_API UE_DECLARE_GAMEPLAY_TAG_EXTERN(Mover_IsSwimming);
 MOVER_API UE_DECLARE_GAMEPLAY_TAG_EXTERN(Mover_IsCrouching);
 MOVER_API UE_DECLARE_GAMEPLAY_TAG_EXTERN(Mover_IsNavWalking);
 
+/** List of Mover tick phases where different parts of work are being performed */
+UENUM(BlueprintType)
+enum class EMoverTickPhase : uint8
+{
+	Invalid = 0             UMETA(Hidden),
+
+	/** This tick is where an input for the next movement step is authored */
+	ProduceInput = 1        UMETA(DisplayName = "ProduceInput"),
+
+	/** This tick is where movement based on {input, state} is simulated, to produce a new state */
+	SimulateMovement = 2    UMETA(DisplayName = "SimulateMovement"),
+
+	/** This tick is where the newest simulation state is applied to the actor and its components */
+	ApplyState = 3          UMETA(DisplayName = "ApplyState")
+};
 
 /** Options for how to handle smoothing frame data from the backend. Typically this is for advancing the simulation at a lower or fixed rate versus the game thread/render rate. */
 UENUM(BlueprintType)
