Minor rename of some code that should have been done as part of previous submit 38194945

#rnx

diff --git a/Engine/Plugins/Experimental/Mover/Source/Mover/Private/MoverPoseSearchTrajectoryPredictor.cpp b/Engine/Plugins/Experimental/Mover/Source/Mover/Private/MoverPoseSearchTrajectoryPredictor.cpp
--- a/Engine/Plugins/Experimental/Mover/Source/Mover/Private/MoverPoseSearchTrajectoryPredictor.cpp
+++ b/Engine/Plugins/Experimental/Mover/Source/Mover/Private/MoverPoseSearchTrajectoryPredictor.cpp
@@ -14,10 +14,10 @@ void UMoverTrajectoryPredictor::Predict_Implementation(FPoseSearchQueryTrajector
 		return;
 	}
 
-	Predict_Static(*MoverComponent, InOutTrajectory, NumPredictionSamples, SecondsPerPredictionSample, NumHistorySamples);
+	Predict(*MoverComponent, InOutTrajectory, NumPredictionSamples, SecondsPerPredictionSample, NumHistorySamples);
 }
 
-void UMoverTrajectoryPredictor::Predict_Static(UMoverComponent& MoverComponent, FPoseSearchQueryTrajectory& InOutTrajectory, int32 NumPredictionSamples, float SecondsPerPredictionSample, int32 NumHistorySamples)
+void UMoverTrajectoryPredictor::Predict(UMoverComponent& MoverComponent, FPoseSearchQueryTrajectory& InOutTrajectory, int32 NumPredictionSamples, float SecondsPerPredictionSample, int32 NumHistorySamples)
 {
 	FMoverPredictTrajectoryParams PredictParams;
 	PredictParams.NumPredictionSamples = NumPredictionSamples;
@@ -72,10 +72,10 @@ void UMoverTrajectoryPredictor::GetCurrentState_Implementation(FVector& OutPosit
 		return;
 	}
 
-	GetCurrentState_Static(*MoverComponent, OutPosition, OutFacing, OutVelocity);
+	GetCurrentState(*MoverComponent, OutPosition, OutFacing, OutVelocity);
 }
 
-void UMoverTrajectoryPredictor::GetCurrentState_Static(UMoverComponent& MoverComponent, FVector& OutPosition, FQuat& OutFacing, FVector& OutVelocity)
+void UMoverTrajectoryPredictor::GetCurrentState(UMoverComponent& MoverComponent, FVector& OutPosition, FQuat& OutFacing, FVector& OutVelocity)
 {
 	const USceneComponent* VisualComp = MoverComponent.GetPrimaryVisualComponent();
 
diff --git a/Engine/Plugins/Experimental/Mover/Source/Mover/Public/MoverPoseSearchTrajectoryPredictor.h b/Engine/Plugins/Experimental/Mover/Source/Mover/Public/MoverPoseSearchTrajectoryPredictor.h
--- a/Engine/Plugins/Experimental/Mover/Source/Mover/Public/MoverPoseSearchTrajectoryPredictor.h
+++ b/Engine/Plugins/Experimental/Mover/Source/Mover/Public/MoverPoseSearchTrajectoryPredictor.h
@@ -30,9 +30,8 @@ public:
 	virtual void GetVelocity_Implementation(FVector& OutVelocity) override;
 
 	// Share code in static helper functions so we can call it from AnimNext, which doesn't support interfaces
-	// For Review: what is naming convention here?
-	static void Predict_Static(UMoverComponent& Mover, FPoseSearchQueryTrajectory& InOutTrajectory, int32 NumPredictionSamples, float SecondsPerPredictionSample, int32 NumHistorySamples);
-	static void GetCurrentState_Static(UMoverComponent& Mover, FVector& OutPosition, FQuat& OutFacing, FVector& OutVelocity);
+	static void Predict(UMoverComponent& Mover, FPoseSearchQueryTrajectory& InOutTrajectory, int32 NumPredictionSamples, float SecondsPerPredictionSample, int32 NumHistorySamples);
+	static void GetCurrentState(UMoverComponent& Mover, FVector& OutPosition, FQuat& OutFacing, FVector& OutVelocity);
 
 
 protected:
diff --git a/Engine/Plugins/Experimental/MoverAnimNext/Source/MoverAnimNext/Private/Graph/RigUnit_GenerateMoverTrajectory.cpp b/Engine/Plugins/Experimental/MoverAnimNext/Source/MoverAnimNext/Private/Graph/RigUnit_GenerateMoverTrajectory.cpp
--- a/Engine/Plugins/Experimental/MoverAnimNext/Source/MoverAnimNext/Private/Graph/RigUnit_GenerateMoverTrajectory.cpp
+++ b/Engine/Plugins/Experimental/MoverAnimNext/Source/MoverAnimNext/Private/Graph/RigUnit_GenerateMoverTrajectory.cpp
@@ -39,10 +39,9 @@ FRigUnit_GenerateMoverTrajectory_Execute()
 		return;
 	}
 
-	UMoverTrajectoryPredictor::GetCurrentState_Static(*MoverComponent, CurrentPosition, CurrentFacing, CurrentVelocity);
-
+	UMoverTrajectoryPredictor::GetCurrentState(*MoverComponent, CurrentPosition, CurrentFacing, CurrentVelocity);
 	UPoseSearchTrajectoryLibrary::InitTrajectorySamples(InOutTrajectory, TrajectoryData, CurrentPosition, CurrentFacing, TrajectoryDataSampling, DeltaTime);
 	UPoseSearchTrajectoryLibrary::UpdateHistory_TransformHistory(InOutTrajectory, TrajectoryData, CurrentPosition, CurrentVelocity, TrajectoryDataSampling, DeltaTime);
-	UMoverTrajectoryPredictor::Predict_Static(*MoverComponent, InOutTrajectory,
+	UMoverTrajectoryPredictor::Predict(*MoverComponent, InOutTrajectory,
 		NumPredictionSamples + 1, PredictionSamplingInterval, NumHistorySamples);
 }
