Mover: small changes to reduce overhead of trajectory prediction. Eliminates use of Blueprint Native events, and avoids unnecessary default initialization of elements in the trajectory samples array.

diff --git a/Engine/Plugins/Animation/PoseSearch/Source/Runtime/Private/PoseSearchTrajectoryLibrary.cpp b/Engine/Plugins/Animation/PoseSearch/Source/Runtime/Private/PoseSearchTrajectoryLibrary.cpp
--- a/Engine/Plugins/Animation/PoseSearch/Source/Runtime/Private/PoseSearchTrajectoryLibrary.cpp
+++ b/Engine/Plugins/Animation/PoseSearch/Source/Runtime/Private/PoseSearchTrajectoryLibrary.cpp
@@ -439,18 +439,19 @@ void UPoseSearchTrajectoryLibrary::PoseSearchGeneratePredictorTrajectory(
 	FVector CurrentVelocity = FVector::ZeroVector;
 	FQuat CurrentFacing = FQuat::Identity;
 
-	if (InPredictor)
+	IPoseSearchTrajectoryPredictorInterface* PredictorObject = Cast<IPoseSearchTrajectoryPredictorInterface>(InPredictor);
+
+	if (PredictorObject)
 	{
-		IPoseSearchTrajectoryPredictorInterface::Execute_GetCurrentState(InPredictor, CurrentPosition, CurrentFacing, CurrentVelocity);
+		PredictorObject->GetCurrentState(CurrentPosition, CurrentFacing, CurrentVelocity);
 	}
 
 	InitTrajectorySamples(InOutTrajectory, InTrajectoryData, CurrentPosition, CurrentFacing, TrajectoryDataSampling, InDeltaTime);
 	UpdateHistory_TransformHistory(InOutTrajectory, InTrajectoryData, CurrentPosition, CurrentVelocity, TrajectoryDataSampling, InDeltaTime);
 
-	if (InPredictor)
+	if (PredictorObject)
 	{
-		IPoseSearchTrajectoryPredictorInterface::Execute_Predict(InPredictor, InOutTrajectory,
-			InTrajectoryPredictionCount + 1, InPredictionSamplingInterval, InTrajectoryHistoryCount);
+		PredictorObject->Predict(InOutTrajectory, InTrajectoryPredictionCount + 1, InPredictionSamplingInterval, InTrajectoryHistoryCount);
 	}
 
 	//InOutDesiredControllerYawLastUpdate = TrajectoryDataState.DesiredControllerYawLastUpdate;
diff --git a/Engine/Plugins/Animation/PoseSearch/Source/Runtime/Public/PoseSearch/PoseSearchTrajectoryPredictor.h b/Engine/Plugins/Animation/PoseSearch/Source/Runtime/Public/PoseSearch/PoseSearchTrajectoryPredictor.h
--- a/Engine/Plugins/Animation/PoseSearch/Source/Runtime/Public/PoseSearch/PoseSearchTrajectoryPredictor.h
+++ b/Engine/Plugins/Animation/PoseSearch/Source/Runtime/Public/PoseSearch/PoseSearchTrajectoryPredictor.h
@@ -8,7 +8,7 @@
 struct FPoseSearchQueryTrajectory;
 
 
-UINTERFACE(BlueprintType, Experimental)
+UINTERFACE(BlueprintType, NotBlueprintable, Experimental)
 class POSESEARCH_API UPoseSearchTrajectoryPredictorInterface : public UInterface
 {
 	GENERATED_BODY()
@@ -23,16 +23,15 @@ class IPoseSearchTrajectoryPredictorInterface : public IInterface
 
 public:
 
-	UFUNCTION(BlueprintNativeEvent)
-	void Predict(FPoseSearchQueryTrajectory& InOutTrajectory, int32 NumPredictionSamples, float SecondsPerPredictionSample, int32 NumHistorySamples);
+	UFUNCTION(BlueprintCallable, Category = "Animation|PoseSearch|Experimental")
+	virtual void Predict(FPoseSearchQueryTrajectory& InOutTrajectory, int32 NumPredictionSamples, float SecondsPerPredictionSample, int32 NumHistorySamples) = 0;
 
-	UFUNCTION(BlueprintCallable, BlueprintNativeEvent, Category="Animation|PoseSearch|Experimental")
-	void GetGravity(FVector& OutGravityAccel);
+	UFUNCTION(BlueprintCallable, Category="Animation|PoseSearch|Experimental")
+	virtual void GetGravity(FVector& OutGravityAccel) = 0;
 
-	UFUNCTION(BlueprintCallable, BlueprintNativeEvent, Category = "Animation|PoseSearch|Experimental")
-	void GetCurrentState(FVector& OutPosition, FQuat& OutFacing, FVector& OutVelocity);
-
-	UFUNCTION(BlueprintCallable, BlueprintNativeEvent, Category = "Animation|PoseSearch|Experimental")
-	void GetVelocity(FVector& OutVelocity);
+	UFUNCTION(BlueprintCallable, Category = "Animation|PoseSearch|Experimental")
+	virtual void GetCurrentState(FVector& OutPosition, FQuat& OutFacing, FVector& OutVelocity) = 0;
 
+	UFUNCTION(BlueprintCallable, Category = "Animation|PoseSearch|Experimental")
+	virtual void GetVelocity(FVector& OutVelocity) = 0;
 };
diff --git a/Engine/Plugins/Experimental/Mover/Source/Mover/Private/MoverComponent.cpp b/Engine/Plugins/Experimental/Mover/Source/Mover/Private/MoverComponent.cpp
--- a/Engine/Plugins/Experimental/Mover/Source/Mover/Private/MoverComponent.cpp
+++ b/Engine/Plugins/Experimental/Mover/Source/Mover/Private/MoverComponent.cpp
@@ -1811,19 +1811,19 @@ TArray<FTrajectorySampleInfo> UMoverComponent::GetPredictedTrajectory(FMoverPred
 
 		if (const UBaseMovementMode* CurrentMovementMode = ModeFSM->GetCurrentMode())
 		{
-			const bool bOrigHasGravityOverride = bHasGravityOverride;
-			const FVector OrigGravityAccelOverride = GravityAccelOverride;
-
-			if (PredictionParams.bDisableGravity)
+			if (FMoverDefaultSyncState* StepSyncState = StepState.SyncState.SyncStateCollection.FindMutableDataByType<FMoverDefaultSyncState>())
 			{
-				SetGravityOverride(true, FVector::ZeroVector);
-			}
+				const bool bOrigHasGravityOverride = bHasGravityOverride;
+				const FVector OrigGravityAccelOverride = GravityAccelOverride;
 
-			TArray<FTrajectorySampleInfo> OutSamples;
-			OutSamples.AddDefaulted(PredictionParams.NumPredictionSamples);
+				if (PredictionParams.bDisableGravity)
+				{
+					SetGravityOverride(true, FVector::ZeroVector);
+				}
+
+				TArray<FTrajectorySampleInfo> OutSamples;
+				OutSamples.SetNumUninitialized(PredictionParams.NumPredictionSamples);
 
-			if (FMoverDefaultSyncState* StepSyncState = StepState.SyncState.SyncStateCollection.FindMutableDataByType<FMoverDefaultSyncState>())
-			{
 				FVector PriorLocation = StepSyncState->GetLocation_WorldSpace();
 				FRotator PriorOrientation = StepSyncState->GetOrientation_WorldSpace();
 				FVector PriorVelocity = StepSyncState->GetVelocity_WorldSpace();
@@ -1839,7 +1839,7 @@ TArray<FTrajectorySampleInfo> UMoverComponent::GetPredictedTrajectory(FMoverPred
 					// Capture sample from current step state
 					FTrajectorySampleInfo& Sample = OutSamples[i];
 
-					Sample.Transform.SetLocation(StepSyncState->GetLocation_WorldSpace());
+					Sample.Transform.SetTranslationAndScale3D(StepSyncState->GetLocation_WorldSpace(), FVector::OneVector);
 					Sample.Transform.SetRotation(StepSyncState->GetOrientation_WorldSpace().Quaternion());
 					Sample.LinearVelocity = StepSyncState->GetVelocity_WorldSpace();
 					Sample.InstantaneousAcceleration = (StepSyncState->GetVelocity_WorldSpace() - PriorVelocity) / PredictionParams.SecondsPerSample;
@@ -1881,14 +1881,14 @@ TArray<FTrajectorySampleInfo> UMoverComponent::GetPredictedTrajectory(FMoverPred
 						}
 					}
 				}
-			}
+				
+				if (PredictionParams.bDisableGravity)
+				{
+					SetGravityOverride(bOrigHasGravityOverride, OrigGravityAccelOverride);
+				}
 
-			if (PredictionParams.bDisableGravity)
-			{
-				SetGravityOverride(bOrigHasGravityOverride, OrigGravityAccelOverride);
+				return OutSamples;
 			}
-
-			return OutSamples;
 		}
 	}
 
diff --git a/Engine/Plugins/Experimental/Mover/Source/Mover/Private/MoverPoseSearchTrajectoryPredictor.cpp b/Engine/Plugins/Experimental/Mover/Source/Mover/Private/MoverPoseSearchTrajectoryPredictor.cpp
--- a/Engine/Plugins/Experimental/Mover/Source/Mover/Private/MoverPoseSearchTrajectoryPredictor.cpp
+++ b/Engine/Plugins/Experimental/Mover/Source/Mover/Private/MoverPoseSearchTrajectoryPredictor.cpp
@@ -5,7 +5,7 @@
  #include "MoverComponent.h"
  #include "MoveLibrary/MovementUtils.h"
  
-void UMoverTrajectoryPredictor::Predict_Implementation(FPoseSearchQueryTrajectory& InOutTrajectory,
+void UMoverTrajectoryPredictor::Predict(FPoseSearchQueryTrajectory& InOutTrajectory,
 	int32 NumPredictionSamples, float SecondsPerPredictionSample, int NumHistorySamples)
 {
 	if (!MoverComponent)
@@ -49,7 +49,7 @@ void UMoverTrajectoryPredictor::Predict(UMoverComponent& MoverComponent, FPoseSe
 	}
 }
 
-void UMoverTrajectoryPredictor::GetGravity_Implementation(FVector& OutGravityAccel)
+void UMoverTrajectoryPredictor::GetGravity(FVector& OutGravityAccel)
 {
 	if (!MoverComponent)
 	{
@@ -62,7 +62,7 @@ void UMoverTrajectoryPredictor::GetGravity_Implementation(FVector& OutGravityAcc
 }
 
 
-void UMoverTrajectoryPredictor::GetCurrentState_Implementation(FVector& OutPosition, FQuat& OutFacing, FVector& OutVelocity)
+void UMoverTrajectoryPredictor::GetCurrentState(FVector& OutPosition, FQuat& OutFacing, FVector& OutVelocity)
 {
 	if (!MoverComponent)
 	{
@@ -111,7 +111,7 @@ void UMoverTrajectoryPredictor::GetCurrentState(UMoverComponent& MoverComponent,
 }
 
 
-void UMoverTrajectoryPredictor::GetVelocity_Implementation(FVector& OutVelocity)
+void UMoverTrajectoryPredictor::GetVelocity(FVector& OutVelocity)
 {
 	if (!MoverComponent)
 	{
diff --git a/Engine/Plugins/Experimental/Mover/Source/Mover/Public/MoverPoseSearchTrajectoryPredictor.h b/Engine/Plugins/Experimental/Mover/Source/Mover/Public/MoverPoseSearchTrajectoryPredictor.h
--- a/Engine/Plugins/Experimental/Mover/Source/Mover/Public/MoverPoseSearchTrajectoryPredictor.h
+++ b/Engine/Plugins/Experimental/Mover/Source/Mover/Public/MoverPoseSearchTrajectoryPredictor.h
@@ -21,13 +21,13 @@ public:
 	UFUNCTION(BlueprintCallable, Category = "Animation|PoseSearch|Experimental")
 	void Setup(UMoverComponent* InMoverComponent) { MoverComponent = InMoverComponent; }
 
-	virtual void Predict_Implementation(FPoseSearchQueryTrajectory& InOutTrajectory, int32 NumPredictionSamples, float SecondsPerPredictionSample, int32 NumHistorySamples) override;
+	virtual void Predict(FPoseSearchQueryTrajectory& InOutTrajectory, int32 NumPredictionSamples, float SecondsPerPredictionSample, int32 NumHistorySamples) override;
 
-	virtual void GetGravity_Implementation(FVector& OutGravityAccel) override;
+	virtual void GetGravity(FVector& OutGravityAccel) override;
 
-	virtual void GetCurrentState_Implementation(FVector& OutPosition, FQuat& OutFacing, FVector& OutVelocity) override;
+	virtual void GetCurrentState(FVector& OutPosition, FQuat& OutFacing, FVector& OutVelocity) override;
 
-	virtual void GetVelocity_Implementation(FVector& OutVelocity) override;
+	virtual void GetVelocity(FVector& OutVelocity) override;
 
 	// Share code in static helper functions so we can call it from AnimNext, which doesn't support interfaces
 	static void Predict(UMoverComponent& Mover, FPoseSearchQueryTrajectory& InOutTrajectory, int32 NumPredictionSamples, float SecondsPerPredictionSample, int32 NumHistorySamples);
