Mover: improved warning messages with more context about what actor they're referring to

#rnx

diff --git a/Engine/Plugins/Experimental/Mover/Source/Mover/Private/MovementModeStateMachine.cpp b/Engine/Plugins/Experimental/Mover/Source/Mover/Private/MovementModeStateMachine.cpp
--- a/Engine/Plugins/Experimental/Mover/Source/Mover/Private/MovementModeStateMachine.cpp
+++ b/Engine/Plugins/Experimental/Mover/Source/Mover/Private/MovementModeStateMachine.cpp
@@ -121,8 +121,7 @@ void UMovementModeStateMachine::QueueNextMode(FName DesiredNextModeName, bool bS
 		if ((NextModeName != NAME_None) &&
 		    (NextModeName != DesiredNextModeName || bShouldReenter != bShouldNextModeReenter))
 		{
-			const AActor* OwnerActor = GetOwnerActor();
-			UE_LOG(LogMover, Log, TEXT("%s (%s) Overwriting of queued mode change (%s, reenter: %i) with (%s, reenter: %i)"), *GetNameSafe(OwnerActor), *UEnum::GetValueAsString(OwnerActor->GetLocalRole()), *NextModeName.ToString(), bShouldNextModeReenter, *DesiredNextModeName.ToString(), bShouldReenter);
+			UE_LOG(LogMover, Log, TEXT("%s (%s) Overwriting of queued mode change (%s, reenter: %i) with (%s, reenter: %i)"), *GetNameSafe(GetOwnerActor()), *UEnum::GetValueAsString(GetOwnerActor()->GetLocalRole()), *NextModeName.ToString(), bShouldNextModeReenter, *DesiredNextModeName.ToString(), bShouldReenter);
 		}
 
 		if (Modes.Contains(DesiredNextModeName))
@@ -131,7 +130,7 @@ void UMovementModeStateMachine::QueueNextMode(FName DesiredNextModeName, bool bS
 		}
 		else
 		{
-			UE_LOG(LogMover, Warning, TEXT("Attempted to queue an unregistered movement mode: %s"), *DesiredNextModeName.ToString());
+			UE_LOG(LogMover, Warning, TEXT("Attempted to queue an unregistered movement mode: %s on owner %s"), *DesiredNextModeName.ToString(), *GetNameSafe(GetOwnerActor()));
 		}
 	}
 }
@@ -350,7 +349,7 @@ void UMovementModeStateMachine::OnSimulationTick(USceneComponent* UpdatedCompone
 				// if we've done this sub step a lot before go ahead and just advance time to avoid freezing editor
 				if (NumConsecutiveFullRefundedSubsteps >= MaxConsecutiveFullRefundedSubsteps)
 				{
-					UE_LOG(LogMover, Warning, TEXT("Movement mode %s and %s on %s are stuck giving time back to eachother - Advancing substep."), *CurrentModeName.ToString(), *OutputState.MovementEndState.NextModeName.ToString(), *MoverComp->GetOwner()->GetName());
+					UE_LOG(LogMover, Warning, TEXT("Movement mode %s and %s on %s are stuck giving time back to each other. Overriding to advance to next substep."), *CurrentModeName.ToString(), *OutputState.MovementEndState.NextModeName.ToString(), *MoverComp->GetOwner()->GetName());
 					SubTimeStep.BaseSimTimeMs += SubTimeStep.StepMs;
 				}
 			}
diff --git a/Engine/Plugins/Experimental/Mover/Source/Mover/Private/MoverComponent.cpp b/Engine/Plugins/Experimental/Mover/Source/Mover/Private/MoverComponent.cpp
--- a/Engine/Plugins/Experimental/Mover/Source/Mover/Private/MoverComponent.cpp
+++ b/Engine/Plugins/Experimental/Mover/Source/Mover/Private/MoverComponent.cpp
@@ -1505,12 +1505,12 @@ UBaseMovementMode* UMoverComponent::AddMovementModeFromClass(FName ModeName, TSu
 {
 	if (!MovementMode)
 	{
-		UE_LOG(LogMover, Warning, TEXT("Attempted to add a movement mode that wasn't valid. AddMovementModeFromClass did not add anything."));
+		UE_LOG(LogMover, Warning, TEXT("Attempted to add a movement mode that wasn't valid. AddMovementModeFromClass will not add anything. (%s)"), *GetNameSafe(GetOwner()));
 		return nullptr;
 	}
 	if (MovementMode->HasAnyClassFlags(CLASS_Abstract))
 	{
-		UE_LOG(LogMover, Warning, TEXT("The Movement Mode class (%s) is abstract and is not a valid class to instantiate. AddMovementModeFromClass will not do anything."), *GetNameSafe(MovementMode));
+		UE_LOG(LogMover, Warning, TEXT("The Movement Mode class (%s) is abstract and is not a valid class to instantiate. AddMovementModeFromClass will not do anything. (%s)"), *GetNameSafe(MovementMode), *GetNameSafe(GetOwner()));
 		return nullptr;
 	}
 
@@ -1524,7 +1524,7 @@ bool UMoverComponent::AddMovementModeFromObject(FName ModeName, UBaseMovementMod
 	{
 		if (MovementMode->GetClass()->HasAnyClassFlags(CLASS_Abstract))
 		{
-			UE_LOG(LogMover, Warning, TEXT("The Movement Mode class (%s) is abstract and is not a valid class to instantiate. AddMovementModeFromObject will not do anything."), *GetNameSafe(MovementMode));
+			UE_LOG(LogMover, Warning, TEXT("The Movement Mode class (%s) is abstract and is not a valid class to instantiate. AddMovementModeFromObject will not do anything. (%s)"), *GetNameSafe(MovementMode), *GetNameSafe(GetOwner()));
 			return false;
 		}
 		
@@ -1532,7 +1532,7 @@ bool UMoverComponent::AddMovementModeFromObject(FName ModeName, UBaseMovementMod
 		{
 			if (FoundMovementMode->Get()->GetClass() == MovementMode->GetClass())
 			{
-				UE_LOG(LogMover, Warning, TEXT("Added the same movement mode (%s) for a movement mode name (%s). AddMovementModeFromObject will add the mode but is likely unwanted/unnecessary behavior."), *GetNameSafe(MovementMode), *ModeName.ToString());
+				UE_LOG(LogMover, Warning, TEXT("Added the same movement mode (%s) for a movement mode name (%s). AddMovementModeFromObject will add the mode but is likely unwanted/unnecessary behavior. (%s)"), *GetNameSafe(MovementMode), *ModeName.ToString(), *GetNameSafe(GetOwner()));
 			}
 
 			RemoveMovementMode(ModeName);
@@ -1540,7 +1540,7 @@ bool UMoverComponent::AddMovementModeFromObject(FName ModeName, UBaseMovementMod
 		
 		if (MovementMode->GetOuter() != this)
 		{
-			UE_LOG(LogMover, Verbose, TEXT("Movement modes are expected to be parented to the MoverComponent. The %s movement mode was reparented to %s!"), *GetNameSafe(MovementMode), *GetNameSafe(this));
+			UE_LOG(LogMover, Verbose, TEXT("Movement modes are expected to be parented to the MoverComponent. The %s movement mode was reparented to %s! (%s)"), *GetNameSafe(MovementMode), *GetNameSafe(this), *GetNameSafe(GetOwner()));
 			MovementMode->Rename(nullptr, this, REN_DoNotDirty | REN_NonTransactional);
 		}
 		
@@ -1549,7 +1549,7 @@ bool UMoverComponent::AddMovementModeFromObject(FName ModeName, UBaseMovementMod
 	}
 	else
 	{
-		UE_LOG(LogMover, Warning, TEXT("Attempted to add %s movement mode that wasn't valid to %s. AddMovementModeFromObject did not add anything."), *GetNameSafe(MovementMode), *GetNameSafe(this));
+		UE_LOG(LogMover, Warning, TEXT("Attempted to add %s movement mode that wasn't valid to %s. AddMovementModeFromObject did not add anything. (%s)"), *GetNameSafe(MovementMode), *GetNameSafe(this), *GetNameSafe(GetOwner()));
 		return false;
 	}
 
@@ -1560,7 +1560,7 @@ bool UMoverComponent::RemoveMovementMode(FName ModeName)
 {
 	if (ModeFSM->GetCurrentModeName() == ModeName)
 	{
-		UE_LOG(LogMover, Warning, TEXT("The mode being removed (%s Movement Mode) is the mode this actor is currently in. It was removed but may cause issues. Consider waiting to remove the mode or queueing a different valid mode to avoid issues."), *ModeName.ToString());
+		UE_LOG(LogMover, Warning, TEXT("The mode being removed (%s Movement Mode) is the mode this actor (%s) is currently in. It was removed but may cause issues. Consider waiting to remove the mode or queueing a different valid mode to avoid issues."), *ModeName.ToString(), *GetNameSafe(GetOwner()));
 	}
 	
 	TObjectPtr<UBaseMovementMode>* ModeToRemove = MovementModes.Find(ModeName);
@@ -1957,7 +1957,7 @@ const FMoverSyncState& UMoverComponent::GetSyncState() const
 {
 	if (!bHasValidCachedState)
 	{
-		UE_LOG(LogMover, Warning, TEXT("Attempting direct access to the last-cached sync state before one has been set. Results will be unreliable. Use the HasValidCachedState function to check if CachedLastSyncState is valid or not."));
+		UE_LOG(LogMover, Warning, TEXT("Attempting direct access to the last-cached sync state before one has been set. Results will be unreliable. Use the HasValidCachedState function to check if CachedLastSyncState is valid or not. (%s)"), *GetNameSafe(GetOwner()));
 	}
 
 	return CachedLastSyncState;
@@ -1993,7 +1993,7 @@ const FMoverInputCmdContext& UMoverComponent::GetLastInputCmd() const
 {
 	if (!bHasValidCachedUsedInput)
 	{
-		UE_LOG(LogMover, Warning, TEXT("Attempting direct access to the last-cached used input cmd before one has been set. Results will be unreliable. Use the HasValidCachedInputCmd function to check if CachedLastUsedInputCmd is valid or not."));
+		UE_LOG(LogMover, Warning, TEXT("Attempting direct access to the last-cached used input cmd before one has been set. Results will be unreliable. Use the HasValidCachedInputCmd function to check if CachedLastUsedInputCmd is valid or not. (%s)"), *GetNameSafe(GetOwner()));
 	}
 
 	return CachedLastUsedInputCmd;
