Mover: minor cruft removal. No functional changes.

#rnx

diff --git a/Engine/Plugins/Experimental/Mover/Source/Mover/Private/DefaultMovementSet/Modes/FallingMode.cpp b/Engine/Plugins/Experimental/Mover/Source/Mover/Private/DefaultMovementSet/Modes/FallingMode.cpp
--- a/Engine/Plugins/Experimental/Mover/Source/Mover/Private/DefaultMovementSet/Modes/FallingMode.cpp
+++ b/Engine/Plugins/Experimental/Mover/Source/Mover/Private/DefaultMovementSet/Modes/FallingMode.cpp
@@ -37,7 +37,6 @@ void UFallingMode::OnGenerateMove(const FMoverTickStartData& StartState, const F
 	const FCharacterDefaultInputs* CharacterInputs = StartState.InputCmd.InputCollection.FindDataByType<FCharacterDefaultInputs>();
 	const FMoverDefaultSyncState* StartingSyncState = StartState.SyncState.SyncStateCollection.FindDataByType<FMoverDefaultSyncState>();
 	check(StartingSyncState);
-	UMoverComponent* MoverComponent = GetMoverComponent();
 
 	const float DeltaSeconds = TimeStep.StepMs * 0.001f;
 
@@ -96,7 +95,7 @@ void UFallingMode::OnGenerateMove(const FMoverTickStartData& StartState, const F
 		Params.Deceleration = OverTerminalSpeedFallingDeceleration;
 	}
 	
-	UMoverBlackboard* SimBlackboard = MoverComponent->GetSimBlackboard_Mutable();
+	UMoverBlackboard* SimBlackboard = MoverComp->GetSimBlackboard_Mutable();
 	FFloorCheckResult LastFloorResult;
 	// limit our moveinput based on the floor we're on
 	if (SimBlackboard && SimBlackboard->TryGet(CommonBlackboard::LastFloorResult, LastFloorResult))
@@ -114,7 +113,7 @@ void UFallingMode::OnGenerateMove(const FMoverTickStartData& StartState, const F
 	}
 	
 	OutProposedMove = UAirMovementUtils::ComputeControlledFreeMove(Params);
-	const FVector VelocityWithGravity = StartVelocity + UMovementUtils::ComputeVelocityFromGravity(MoverComponent->GetGravityAcceleration(), DeltaSeconds);
+	const FVector VelocityWithGravity = StartVelocity + UMovementUtils::ComputeVelocityFromGravity(MoverComp->GetGravityAcceleration(), DeltaSeconds);
 
 	//  If we are going faster than TerminalVerticalVelocity apply VerticalFallingDeceleration otherwise reset Z velocity to before we applied deceleration 
 	if (VelocityWithGravity.GetAbs().Dot(UpDirection) > TerminalVerticalSpeed)
@@ -184,9 +183,6 @@ void UFallingMode::OnSimulationTick(const FSimulationTickParams& Params, FMoverT
 
 	UMovementUtils::TrySafeMoveUpdatedComponent(Params.MovingComps, MoveDelta, TargetOrientQuat, true, Hit, ETeleportType::None, MoveRecord);
 
-	// Compute final velocity based on how long we actually go until we get a hit.
-	FVector NewFallingVelocity = StartingFallingVelocity + (ProposedMove.LinearVelocity * Hit.Time * MoveDelta);
-
 	FFloorCheckResult LandingFloor;
 
 	// Handle impact, whether it's a landing surface or something to slide on
diff --git a/Engine/Plugins/Experimental/Mover/Source/Mover/Private/DefaultMovementSet/Modes/NavWalkingMode.cpp b/Engine/Plugins/Experimental/Mover/Source/Mover/Private/DefaultMovementSet/Modes/NavWalkingMode.cpp
--- a/Engine/Plugins/Experimental/Mover/Source/Mover/Private/DefaultMovementSet/Modes/NavWalkingMode.cpp
+++ b/Engine/Plugins/Experimental/Mover/Source/Mover/Private/DefaultMovementSet/Modes/NavWalkingMode.cpp
@@ -114,7 +114,10 @@ void UNavWalkingMode::OnGenerateMove(const FMoverTickStartData& StartState, cons
 void UNavWalkingMode::OnSimulationTick(const FSimulationTickParams& Params, FMoverTickEndData& OutputState)
 {
 	const UMoverComponent* MoverComp = Cast<UMoverComponent>(GetMoverComponent());
-	ensureMsgf(MoverComp, TEXT("Nav Walking Mode couldn't find a valid MoverComponent!"));
+	if (!ensureMsgf(MoverComp, TEXT("Nav Walking Mode couldn't find a valid MoverComponent!")))
+	{
+		return;
+	}
 	
 	const FMoverTickStartData& StartState = Params.StartState;
 	USceneComponent* UpdatedComponent = Params.MovingComps.UpdatedComponent.Get();
@@ -218,9 +221,10 @@ void UNavWalkingMode::OnSimulationTick(const FSimulationTickParams& Params, FMov
 		const bool bHasNavigationData = FindNavFloor(AdjustedDest, DestNavLocation);
 		if (!bHasNavigationData)
 		{
+			// Can't find nav mesh, so fall back to normal walking mode and let it process the intended movement
 			OutputState.MovementEndState.NextModeName = DefaultModeNames::Walking;
 			OutputState.MovementEndState.RemainingMs = Params.TimeStep.StepMs;
-			MoveRecord.SetDeltaSeconds((Params.TimeStep.StepMs - OutputState.MovementEndState.RemainingMs) * 0.001f);
+			MoveRecord.SetDeltaSeconds(0.0f);
 			
 			CaptureFinalState(UpdatedComponent, true, CurrentFloor, MoveRecord, OutputSyncState);
 			return;
@@ -250,9 +254,10 @@ void UNavWalkingMode::OnSimulationTick(const FSimulationTickParams& Params, FMov
 	}
 	else
 	{
+		// Can't find nav destination, so revert to a different mode and let it process the intended movement
 		OutputState.MovementEndState.NextModeName = CommonLegacySettings->AirMovementModeName;
 		OutputState.MovementEndState.RemainingMs = Params.TimeStep.StepMs;
-		MoveRecord.SetDeltaSeconds((Params.TimeStep.StepMs - OutputState.MovementEndState.RemainingMs) * 0.001f);
+		MoveRecord.SetDeltaSeconds(0.0f);
 	}
 
 	CaptureFinalState(UpdatedComponent, true, CurrentFloor, MoveRecord, OutputSyncState);
diff --git a/Engine/Plugins/Experimental/Mover/Source/Mover/Private/DefaultMovementSet/Modes/WalkingMode.cpp b/Engine/Plugins/Experimental/Mover/Source/Mover/Private/DefaultMovementSet/Modes/WalkingMode.cpp
--- a/Engine/Plugins/Experimental/Mover/Source/Mover/Private/DefaultMovementSet/Modes/WalkingMode.cpp
+++ b/Engine/Plugins/Experimental/Mover/Source/Mover/Private/DefaultMovementSet/Modes/WalkingMode.cpp
@@ -133,17 +133,13 @@ void UWalkingMode::OnSimulationTick(const FSimulationTickParams& Params, FMoverT
 
 	const FVector OrigMoveDelta = ProposedMove.LinearVelocity * DeltaSeconds;
 
-	TObjectPtr<AActor> OwnerActor = UpdatedComponent->GetOwner();
-	check(OwnerActor);
-	TObjectPtr<USkeletalMeshComponent> OwnerSkeletalMesh = OwnerActor->FindComponentByClass<USkeletalMeshComponent>();
-
 	FMovementRecord MoveRecord;
 	MoveRecord.SetDeltaSeconds(DeltaSeconds);
 
 	FFloorCheckResult CurrentFloor;
 	UMoverBlackboard* SimBlackboard = MoverComp->GetSimBlackboard_Mutable();
 
-	FVector UpDirection = MoverComp->GetUpDirection();
+	const FVector UpDirection = MoverComp->GetUpDirection();
 	FMovingComponentSet MovingComponents(MoverComp);
 	
 	// If we don't have cached floor information, we need to search for it again
@@ -176,7 +172,6 @@ void UWalkingMode::OnSimulationTick(const FSimulationTickParams& Params, FMoverT
 	bool bDidAttemptMovement = false;
 
 	float PercentTimeAppliedSoFar = MoveHitResult.Time;
-	bool bWasFirstMoveBlocked = false;
 
 	if (!CurMoveDelta.IsNearlyZero() || bIsOrientationChanging)
 	{
@@ -189,12 +184,10 @@ void UWalkingMode::OnSimulationTick(const FSimulationTickParams& Params, FMoverT
 		{
 			// We started by being stuck in geometry and need to resolve it first
 			// TODO: try to resolve starting stuck in geometry
-			bWasFirstMoveBlocked = true;
 		}
 		else if (MoveHitResult.IsValidBlockingHit())
 		{
 			// We impacted something (possibly a ramp, possibly a barrier)
-			bWasFirstMoveBlocked = true;
 			PercentTimeAppliedSoFar = MoveHitResult.Time;		
 
 			// Check if the blockage is a walkable ramp rising in front of us
@@ -203,7 +196,7 @@ void UWalkingMode::OnSimulationTick(const FSimulationTickParams& Params, FMoverT
 			{
 				// It's a walkable ramp, so cut up the move and attempt to move the remainder of it along the ramp's surface, possibly generating another hit
 				const float PercentTimeRemaining = 1.f - PercentTimeAppliedSoFar;
-				CurMoveDelta = UGroundMovementUtils::ComputeDeflectedMoveOntoRamp(CurMoveDelta * PercentTimeRemaining, MoverComp->GetUpDirection(), MoveHitResult, CommonLegacySettings->MaxWalkSlopeCosine, CurrentFloor.bLineTrace);
+				CurMoveDelta = UGroundMovementUtils::ComputeDeflectedMoveOntoRamp(CurMoveDelta * PercentTimeRemaining, UpDirection, MoveHitResult, CommonLegacySettings->MaxWalkSlopeCosine, CurrentFloor.bLineTrace);
 				UMovementUtils::TrySafeMoveUpdatedComponent(Params.MovingComps, CurMoveDelta, TargetOrientQuat, true, MoveHitResult, ETeleportType::None, MoveRecord);
 				LastMoveSeconds = PercentTimeRemaining * LastMoveSeconds;
 
