Mover: cvars to support async movement sim, starting with  the standalone backend

diff --git a/Engine/Plugins/Experimental/Mover/Source/Mover/Private/Backends/MoverStandaloneLiaison.cpp b/Engine/Plugins/Experimental/Mover/Source/Mover/Private/Backends/MoverStandaloneLiaison.cpp
--- a/Engine/Plugins/Experimental/Mover/Source/Mover/Private/Backends/MoverStandaloneLiaison.cpp
+++ b/Engine/Plugins/Experimental/Mover/Source/Mover/Private/Backends/MoverStandaloneLiaison.cpp
@@ -6,6 +6,17 @@
 #include "GameFramework/Controller.h"
 #include "MoverComponent.h"
 
+namespace MoverStandaloneBackendCVars
+{ 
+	// Whether to allow movement simulation ticks on any thread
+	static int32 RunMovementSimOnAnyThread = 0;
+	FAutoConsoleVariableRef CVarMoverStandaloneRunSimOnAnyThread(
+		TEXT("mover.standalone.RunMovementSimOnAnyThread"),
+		RunMovementSimOnAnyThread,
+		TEXT("Whether to allow Mover simulation ticks to run on any thread. Requires use of threadsafe movement modes that do not modify scene components.\n")
+		TEXT("0: Game thread only, 1: Any thread"),
+		ECVF_Default);
+}
 
 UMoverStandaloneLiaisonComponent::UMoverStandaloneLiaisonComponent(const FObjectInitializer& ObjectInitializer)
 	: Super(ObjectInitializer)
@@ -23,6 +34,7 @@ UMoverStandaloneLiaisonComponent::UMoverStandaloneLiaisonComponent(const FObject
 	SimulateMovementTickFunction.bStartWithTickEnabled = true;
 	SimulateMovementTickFunction.SetTickFunctionEnable(true);
 	SimulateMovementTickFunction.TickGroup = TG_PrePhysics;
+	SimulateMovementTickFunction.bRunOnAnyThread = MoverStandaloneBackendCVars::RunMovementSimOnAnyThread;
 
 	ApplyStateTickFunction.bCanEverTick = true;
 	ApplyStateTickFunction.bStartWithTickEnabled = true;
diff --git a/Engine/Plugins/Experimental/Mover/Source/Mover/Private/MoverComponent.cpp b/Engine/Plugins/Experimental/Mover/Source/Mover/Private/MoverComponent.cpp
--- a/Engine/Plugins/Experimental/Mover/Source/Mover/Private/MoverComponent.cpp
+++ b/Engine/Plugins/Experimental/Mover/Source/Mover/Private/MoverComponent.cpp
@@ -35,6 +35,18 @@
 
 #define LOCTEXT_NAMESPACE "Mover"
 
+namespace MoverComponentCVars
+{
+	static int32 WarnOnPostSimDifference = 0;
+	static FAutoConsoleVariableRef CVarMoverWarnOnPostSimDifference(
+		TEXT("mover.debug.WarnOnPostSimDifference"),
+		WarnOnPostSimDifference,
+		TEXT("If != 0, then any differences between the sim sync state and the component locations just after movement simulation will emit warnings.\n")
+	);
+
+} // end MoverComponentCVars
+
+
 
 namespace MoverComponentConstants
 {
@@ -436,7 +448,7 @@ void UMoverComponent::SimulationTick(const FMoverTimeStep& InTimeStep, const FMo
 	{
 		if (!(UpdatedComponent->GetComponentLocation().Equals(StartingSyncState->GetLocation_WorldSpace()) || StartingSyncState->GetMovementBase()))
 		{
-			UE_LOG(LogMover, Warning, TEXT("%s %s: Simulation start location (%s) disagrees with actual mover component location (%s). This indicates movement of the component out-of-band with the simulation, and if happens often will lead to poor quality motion."),
+			UE_LOG(LogMover, Warning, TEXT("%s %s: Simulation start location (%s) disagrees with actual mover component location (%s). This indicates movement of the component out-of-band with the simulation, and may cause poor quality motion."),
 			*GetNameSafe(GetOwner()),
 			*StaticEnum<ENetRole>()->GetValueAsString(GetOwnerRole()),
 			*StartingSyncState->GetLocation_WorldSpace().ToCompactString(),
@@ -487,10 +499,13 @@ void UMoverComponent::SimulationTick(const FMoverTimeStep& InTimeStep, const FMo
 		const FName MovementModeAfterTick = ModeFSM->GetCurrentModeName();
 		SimOutput.SyncState.MovementMode = MovementModeAfterTick;
 
-		if (UpdatedComponent->GetComponentLocation().Equals(OutputSyncState->GetLocation_WorldSpace()) == false ||
-			UpdatedComponent->GetComponentQuat().Equals(OutputSyncState->GetOrientation_WorldSpace().Quaternion(), UE_KINDA_SMALL_NUMBER) == false)
+		if (MoverComponentCVars::WarnOnPostSimDifference)
 		{
-			UE_LOG(LogMover, Warning, TEXT("Detected pos/rot difference between Mover actor (%s) sync state and scene component after sim ticking. This indicates a movement mode may not be authoring the final state correctly."), *GetNameSafe(UpdatedComponent->GetOwner()));
+			if (UpdatedComponent->GetComponentLocation().Equals(OutputSyncState->GetLocation_WorldSpace()) == false ||
+				UpdatedComponent->GetComponentQuat().Equals(OutputSyncState->GetOrientation_WorldSpace().Quaternion(), UE_KINDA_SMALL_NUMBER) == false)
+			{
+				UE_LOG(LogMover, Warning, TEXT("Detected pos/rot difference between Mover actor (%s) sync state and scene component after sim ticking. This indicates a movement mode may not be authoring the final state correctly."), *GetNameSafe(UpdatedComponent->GetOwner()));
+			}
 		}
 	}
 
@@ -721,7 +736,7 @@ void UMoverComponent::SetFrameStateFromContext(const FMoverSyncState* SyncState,
 
 		// Apply the desired transform to the scene component
 		FTransform Transform(WorldOrientation, WorldLocation, UpdatedComponent->GetComponentTransform().GetScale3D());
-		UpdatedComponent->SetWorldTransform(Transform, false, nullptr, ETeleportType::TeleportPhysics);
+		UpdatedComponent->SetWorldTransform(Transform, /*bSweep*/false, nullptr, ETeleportType::TeleportPhysics);
 		UpdatedComponent->ComponentVelocity = WorldVelocity;
 	}
 }
