Mover: core async movement test utilities and implementation of async flying mode (the simplest default mode)

diff --git a/Engine/Plugins/Experimental/Mover/Source/Mover/Private/DefaultMovementSet/Modes/AsyncFlyingMode.cpp b/Engine/Plugins/Experimental/Mover/Source/Mover/Private/DefaultMovementSet/Modes/AsyncFlyingMode.cpp
new file mode 100644
--- /dev/null
+++ b/Engine/Plugins/Experimental/Mover/Source/Mover/Private/DefaultMovementSet/Modes/AsyncFlyingMode.cpp
@@ -0,0 +1,158 @@
+// Copyright Epic Games, Inc. All Rights Reserved.
+
+#include "DefaultMovementSet/Modes/AsyncFlyingMode.h"
+#include "MoveLibrary/AirMovementUtils.h"
+#include "MoveLibrary/MovementUtils.h"
+#include "MoveLibrary/AsyncMovementUtils.h"
+#include "MoverComponent.h"
+#include "DefaultMovementSet/Settings/CommonLegacyMovementSettings.h"
+
+
+#include UE_INLINE_GENERATED_CPP_BY_NAME(AsyncFlyingMode)
+
+
+UAsyncFlyingMode::UAsyncFlyingMode(const FObjectInitializer& ObjectInitializer)
+	: Super(ObjectInitializer)
+{
+	GameplayTags.AddTag(Mover_IsInAir);
+	GameplayTags.AddTag(Mover_IsFlying);
+}
+
+void UAsyncFlyingMode::OnGenerateMove(const FMoverTickStartData& StartState, const FMoverTimeStep& TimeStep, FProposedMove& OutProposedMove) const
+{
+	const UMoverComponent* MoverComp = GetMoverComponent();
+	const FCharacterDefaultInputs* CharacterInputs = StartState.InputCmd.InputCollection.FindDataByType<FCharacterDefaultInputs>();
+	const FMoverDefaultSyncState* StartingSyncState = StartState.SyncState.SyncStateCollection.FindDataByType<FMoverDefaultSyncState>();
+	check(StartingSyncState);
+
+	const float DeltaSeconds = TimeStep.StepMs * 0.001f;
+
+	FFreeMoveParams Params;
+	if (CharacterInputs)
+	{
+		Params.MoveInputType = CharacterInputs->GetMoveInputType();
+		const bool bMaintainInputMagnitude = true;
+		Params.MoveInput = UPlanarConstraintUtils::ConstrainDirectionToPlane(MoverComp->GetPlanarConstraint(), CharacterInputs->GetMoveInput_WorldSpace(), bMaintainInputMagnitude);
+	}
+	else
+	{
+		Params.MoveInputType = EMoveInputType::Invalid;
+		Params.MoveInput = FVector::ZeroVector;
+	}
+
+	FRotator IntendedOrientation_WorldSpace;
+	// If there's no intent from input to change orientation, use the current orientation
+	if (!CharacterInputs || CharacterInputs->OrientationIntent.IsNearlyZero())
+	{
+		IntendedOrientation_WorldSpace = StartingSyncState->GetOrientation_WorldSpace();
+	}
+	else
+	{
+		IntendedOrientation_WorldSpace = CharacterInputs->GetOrientationIntentDir_WorldSpace().ToOrientationRotator();
+	}
+
+	IntendedOrientation_WorldSpace = UMovementUtils::ApplyGravityToOrientationIntent(IntendedOrientation_WorldSpace, MoverComp->GetWorldToGravityTransform(), CommonLegacySettings->bShouldRemainVertical);
+	
+	Params.OrientationIntent = IntendedOrientation_WorldSpace;
+	Params.PriorVelocity = StartingSyncState->GetVelocity_WorldSpace();
+	Params.PriorOrientation = StartingSyncState->GetOrientation_WorldSpace();
+	Params.TurningRate = CommonLegacySettings->TurningRate;
+	Params.TurningBoost = CommonLegacySettings->TurningBoost;
+	Params.MaxSpeed = CommonLegacySettings->MaxSpeed;
+	Params.Acceleration = CommonLegacySettings->Acceleration;
+	Params.Deceleration = CommonLegacySettings->Deceleration;
+	Params.DeltaSeconds = DeltaSeconds;
+	Params.WorldToGravityQuat = MoverComp->GetWorldToGravityTransform();
+	Params.bUseAccelerationForVelocityMove = CommonLegacySettings->bUseAccelerationForVelocityMove;
+	
+	OutProposedMove = UAirMovementUtils::ComputeControlledFreeMove(Params);
+}
+
+void UAsyncFlyingMode::OnSimulationTick(const FSimulationTickParams& Params, FMoverTickEndData& OutputState)
+{
+	const UMoverComponent* MoverComp = GetMoverComponent();
+	const FMoverTickStartData& StartState = Params.StartState;
+	USceneComponent* UpdatedComponent = Params.MovingComps.UpdatedComponent.Get();
+	const FProposedMove& ProposedMove = Params.ProposedMove;
+
+	const FCharacterDefaultInputs* CharacterInputs = StartState.InputCmd.InputCollection.FindDataByType<FCharacterDefaultInputs>();
+	const FMoverDefaultSyncState* StartingSyncState = StartState.SyncState.SyncStateCollection.FindDataByType<FMoverDefaultSyncState>();
+	check(StartingSyncState);
+
+	FMoverDefaultSyncState& OutputSyncState = OutputState.SyncState.SyncStateCollection.FindOrAddMutableDataByType<FMoverDefaultSyncState>();
+
+	const float DeltaSeconds = Params.TimeStep.StepMs * 0.001f;
+
+	FMovementRecord MoveRecord;
+	MoveRecord.SetDeltaSeconds(DeltaSeconds);
+
+	UMoverBlackboard* SimBlackboard = MoverComp->GetSimBlackboard_Mutable();
+	SimBlackboard->Invalidate(CommonBlackboard::LastFloorResult);	// flying = no valid floor
+	SimBlackboard->Invalidate(CommonBlackboard::LastFoundDynamicMovementBase);
+
+	OutputSyncState.MoveDirectionIntent = (ProposedMove.bHasDirIntent ? ProposedMove.DirectionIntent : FVector::ZeroVector);
+
+	// Use the orientation intent directly. If no intent is provided, use last frame's orientation. Note that we are assuming rotation changes can't fail. 
+	const FRotator StartingOrient = StartingSyncState->GetOrientation_WorldSpace();
+	FRotator TargetOrient = StartingOrient;
+
+	bool bIsOrientationChanging = false;
+
+	// Apply orientation changes (if any)
+	if (!UMovementUtils::IsAngularVelocityZero(ProposedMove.AngularVelocity))
+	{
+		TargetOrient += (ProposedMove.AngularVelocity * DeltaSeconds);
+		bIsOrientationChanging = (TargetOrient != StartingOrient);
+	}
+	
+	const FVector StartLocation = StartingSyncState->GetLocation_WorldSpace();
+	const FVector TargetLocation = StartLocation + (ProposedMove.LinearVelocity * DeltaSeconds);
+
+	const FQuat StartRotation = StartingOrient.Quaternion();
+	const FQuat TargetRotation = TargetOrient.Quaternion();
+
+	FHitResult SweepHit(1.f);
+	FMovementRecord SweepRecord;
+	SweepRecord.SetDeltaSeconds(DeltaSeconds);
+	FVector LocationInProgress = StartLocation;
+	FQuat   RotationInProgress = StartRotation;
+
+	const bool bWouldMove = UAsyncMovementUtils::TestDepenetratingMove(Params.MovingComps, StartLocation, TargetLocation, StartRotation, TargetRotation, /* bShouldSweep */ true, OUT SweepHit, SweepRecord);
+
+	LocationInProgress = StartLocation + ((TargetLocation - StartLocation) * SweepHit.Time);
+	RotationInProgress = FQuat::Slerp(StartRotation, TargetRotation, SweepHit.Time);
+
+	if (SweepHit.IsValidBlockingHit())
+	{
+		const float PctOfTimeUsedForSliding = UAsyncMovementUtils::TestSlidingMoveAlongHitSurface(Params.MovingComps, TargetLocation - StartLocation, LocationInProgress, TargetRotation, SweepHit, SweepRecord);
+
+		if (PctOfTimeUsedForSliding > 0.f)
+		{
+			LocationInProgress = SweepHit.TraceStart + ((SweepHit.TraceEnd - SweepHit.TraceStart) * PctOfTimeUsedForSliding);
+			RotationInProgress = FQuat::Slerp(RotationInProgress, TargetRotation, PctOfTimeUsedForSliding);
+		}
+	}
+
+	OutputSyncState.SetTransforms_WorldSpace(
+		LocationInProgress,
+		RotationInProgress.Rotator(),
+		SweepRecord.GetRelevantVelocity(),
+		nullptr); // no movement base
+}
+
+
+void UAsyncFlyingMode::OnRegistered(const FName ModeName)
+{
+	Super::OnRegistered(ModeName);
+
+	CommonLegacySettings = GetMoverComponent()->FindSharedSettings<UCommonLegacyMovementSettings>();
+	ensureMsgf(CommonLegacySettings, TEXT("Failed to find instance of CommonLegacyMovementSettings on %s. Movement may not function properly."), *GetPathNameSafe(this));
+}
+
+
+void UAsyncFlyingMode::OnUnregistered()
+{
+	CommonLegacySettings = nullptr;
+
+	Super::OnUnregistered();
+}
\ No newline at end of file
diff --git a/Engine/Plugins/Experimental/Mover/Source/Mover/Private/MoveLibrary/AsyncMovementUtils.cpp b/Engine/Plugins/Experimental/Mover/Source/Mover/Private/MoveLibrary/AsyncMovementUtils.cpp
new file mode 100644
--- /dev/null
+++ b/Engine/Plugins/Experimental/Mover/Source/Mover/Private/MoveLibrary/AsyncMovementUtils.cpp
@@ -0,0 +1,402 @@
+// Copyright Epic Games, Inc. All Rights Reserved.
+
+#include "MoveLibrary/AsyncMovementUtils.h"
+#include "MoveLibrary/MovementRecord.h"
+#include "MoveLibrary/MovementUtils.h"
+#include "MoverComponent.h"
+#include "MoverLog.h"
+#include "Engine/World.h"
+
+#include UE_INLINE_GENERATED_CPP_BY_NAME(AsyncMovementUtils)
+
+
+
+// TODO: replace with real cvars
+static float HitDistanceToleranceCVar = 0.f;
+static float InitialOverlapToleranceCVar = 0.f;
+static float PenetrationOverlapCheckInflationCVar = 0.1f;
+
+
+bool UAsyncMovementUtils::TestDepenetratingMove(const FMovingComponentSet& MovingComps, const FVector& StartLocation, const FVector& TargetLocation, const FQuat& StartRotation, const FQuat& TargetRotation, bool bShouldSweep, FHitResult& OutHit, FMovementRecord& InOutMoveRecord)
+{
+	const FName TestSafeMoveSubstepName = "TestSafeMove";
+	const FName TestSafeDepenetrationSubstepName = "TestSafeResolvePenetration";
+
+	const USceneComponent* UpdatedComponent = MovingComps.UpdatedComponent.Get();
+
+	if (!UpdatedComponent)
+	{
+		OutHit.Reset(1.f);
+		return false;
+	}
+
+	bool bDidMove = false;
+	EMoveComponentFlags MoveCompFlags = MOVECOMP_NoFlags;
+
+	FVector ResolvedStartLocation = StartLocation;
+	FVector CurrentLocation = StartLocation;
+
+	// Test full move first
+	{
+		const EMoveComponentFlags IncludeBlockingOverlapsWithoutEvents = (MOVECOMP_NeverIgnoreBlockingOverlaps | MOVECOMP_DisableBlockingOverlapDispatch);
+
+		MoveCompFlags = (MoveCompFlags | IncludeBlockingOverlapsWithoutEvents);
+		bDidMove = TestMoveComponent_Internal(MovingComps, ResolvedStartLocation, TargetLocation, StartRotation, TargetRotation, bShouldSweep, MoveCompFlags, OutHit);
+
+		CurrentLocation = ResolvedStartLocation + ((TargetLocation-ResolvedStartLocation) * OutHit.Time);
+
+		UE_LOG(LogMover, VeryVerbose, TEXT("TestDepenetratingMove: %s (role %i) Delta=%s DidMove=%i"),
+			*GetNameSafe(UpdatedComponent->GetOwner()), UpdatedComponent->GetOwnerRole(), *(TargetLocation- ResolvedStartLocation).ToCompactString(), bDidMove);
+	}
+
+	// If we are starting out in penetration, try to resolve it and then retry
+	if (OutHit.bStartPenetrating)
+	{
+		FVector AdjustmentToResolve = FVector::ZeroVector;
+		if (FindMoveToResolveInitialPenetration_Internal(MovingComps, StartLocation, StartRotation, OutHit, MoveCompFlags, AdjustmentToResolve))
+		{
+			ResolvedStartLocation = StartLocation + AdjustmentToResolve;
+			CurrentLocation = ResolvedStartLocation;	// keeping track of where we've moved so far
+
+			InOutMoveRecord.Append(FMovementSubstep(TestSafeDepenetrationSubstepName, AdjustmentToResolve, false));
+
+			// Retry original move
+			bDidMove |= TestMoveComponent_Internal(MovingComps, ResolvedStartLocation, TargetLocation, StartRotation, TargetRotation, bShouldSweep, MoveCompFlags, OutHit);
+
+			UE_LOG(LogMover, VeryVerbose, TEXT("TestDepenetratingMove retry: %s (role %i) Delta=%s DidMove=%i"),
+				*GetNameSafe(UpdatedComponent->GetOwner()), UpdatedComponent->GetOwnerRole(), *(TargetLocation - ResolvedStartLocation).ToCompactString(), bDidMove);
+
+			if (bDidMove)
+			{
+				CurrentLocation = ResolvedStartLocation + ((TargetLocation - ResolvedStartLocation) * OutHit.Time);
+			}
+		}
+	}
+
+	if (bDidMove)
+	{
+		InOutMoveRecord.Append(FMovementSubstep(TestSafeMoveSubstepName, CurrentLocation - ResolvedStartLocation, true));
+	}
+
+	return bDidMove;
+}
+
+
+/* static */ 
+float UAsyncMovementUtils::TestSlidingMoveAlongHitSurface(const FMovingComponentSet& MovingComps, const FVector& OriginalMoveDelta, const FVector& LocationAtHit, const FQuat& TargetRotation, FHitResult& InOutHit, FMovementRecord& InOutMoveRecord)
+{
+	if (!InOutHit.bBlockingHit)
+	{
+		return 0.f;
+	}
+
+	float PctOfTimeUsed = 0.f;
+	float PctOfOrigDeltaToSlide = 1.f - InOutHit.Time;
+	const FVector OrigHitNormal = InOutHit.Normal;
+
+	FVector SlideDelta = UMovementUtils::ComputeSlideDelta(MovingComps, OriginalMoveDelta, PctOfOrigDeltaToSlide, OrigHitNormal, InOutHit);
+
+	if ((SlideDelta | OriginalMoveDelta) > 0.f)
+	{
+		TestDepenetratingMove(MovingComps, LocationAtHit, LocationAtHit+SlideDelta, TargetRotation, TargetRotation, /*bShouldSweep*/ true, InOutHit, InOutMoveRecord);
+
+		PctOfTimeUsed = InOutHit.Time;
+
+		if (InOutHit.IsValidBlockingHit())
+		{
+			// TODO: maybe record impact here
+
+			// Compute new slide normal when hitting multiple surfaces.
+			SlideDelta = UMovementUtils::ComputeTwoWallAdjustedDelta(MovingComps, SlideDelta, InOutHit, OrigHitNormal);
+
+			// Only proceed if the new direction is of significant length and not in reverse of original attempted move.
+			if (!SlideDelta.IsNearlyZero(UE::MoverUtils::SMALL_MOVE_DISTANCE) && (SlideDelta | OriginalMoveDelta) > 0.f)
+			{
+				const FVector LocationAfter1stSlide = InOutHit.TraceStart + ((InOutHit.TraceEnd-InOutHit.TraceStart) * InOutHit.Time);
+
+				// Perform second move
+				TestDepenetratingMove(MovingComps, LocationAfter1stSlide, LocationAfter1stSlide + SlideDelta, TargetRotation, TargetRotation, /*bShouldSweep*/ true, InOutHit, InOutMoveRecord);
+				PctOfTimeUsed += (InOutHit.Time * (1.f - PctOfTimeUsed));
+
+				// TODO: maybe record second impact here
+			}
+		}
+
+		return FMath::Clamp(PctOfTimeUsed, 0.f, 1.f);
+	}
+
+	return 0.f;
+
+}
+
+
+/* static */
+bool UAsyncMovementUtils::FindMoveToResolveInitialPenetration_Internal(const FMovingComponentSet& MovingComps, const FVector& StartLocation, const FQuat& StartRotation, const FHitResult& PenetratingHit, EMoveComponentFlags MoveCompFlags, FVector& OutAdjustmentDelta)
+{
+	if (ensureMsgf(PenetratingHit.bStartPenetrating, TEXT("Expected a hit that started penetrating. Will not attempt adjustment.")))
+	{
+		const USceneComponent* UpdatedComponent = MovingComps.UpdatedComponent.Get();
+		const UPrimitiveComponent* UpdatedPrimitive = MovingComps.UpdatedPrimitive.Get();
+		const AActor* ActorOwner = UpdatedPrimitive->GetOwner();
+		if (!ActorOwner)
+		{
+			return false;
+		}
+
+		FVector ProposedAdjustment = UMovementUtils::ComputePenetrationAdjustment(PenetratingHit);
+		ProposedAdjustment = UPlanarConstraintUtils::ConstrainDirectionToPlane(MovingComps.MoverComponent->GetPlanarConstraint(), ProposedAdjustment);
+
+		// SceneComponent can't be in penetration, so this function really only applies to PrimitiveComponent.
+		if (!ProposedAdjustment.IsZero() && UpdatedPrimitive)
+		{
+			// Attempts to adjust:
+			//	1) Check if our component will fit with the proposed adjustment. If so, accept the proposed adjustment.
+			//  2) Try sweep moving out while ignoring blocking overlaps. If move is allowed, capture how far we actually moved and accept that as the proposed adjustment.
+			//  3) If still stuck in penetration, add a second penetration adjustment
+			//  4) If still stuck, try moving the entire amount
+
+			// See if we can fit at the adjusted location without overlapping anything.
+
+
+			// We really want to make sure that precision differences or differences between the overlap test and sweep tests don't put us into another overlap,
+			// so make the overlap test a bit more restrictive.
+			const float OverlapInflation = PenetrationOverlapCheckInflationCVar;
+			bool bEncroached = UMovementUtils::OverlapTest(UpdatedComponent, UpdatedPrimitive, PenetratingHit.TraceStart + ProposedAdjustment, StartRotation, UpdatedPrimitive->GetCollisionObjectType(), UpdatedPrimitive->GetCollisionShape(OverlapInflation), ActorOwner);
+
+			if (!bEncroached)
+			{
+				// No more overlapping, so we've got an acceptable adjustment
+				OutAdjustmentDelta = ProposedAdjustment;
+				return true;
+			}
+
+			// Disable MOVECOMP_NeverIgnoreBlockingOverlaps if it is enabled, otherwise we wouldn't be able to sweep out of the object to fix the penetration.
+			TGuardValue<EMoveComponentFlags> ScopedFlagRestore(MoveCompFlags, EMoveComponentFlags(MoveCompFlags & (~MOVECOMP_NeverIgnoreBlockingOverlaps)));
+
+			// Try sweeping out as far as possible...
+
+			FHitResult SweepOutHit(1.f);
+
+			bool bMoved = TestMoveComponent_Internal(MovingComps, StartLocation, StartLocation + ProposedAdjustment, StartRotation, StartRotation, /* bShouldSweep */ true, MoveCompFlags, SweepOutHit);
+
+			UE_LOG(LogMover, VeryVerbose, TEXT("FindMoveToResolveInitialPenetration_Internal: %s (role %i) Adjustment=%s DidMove=%i"),
+				*GetNameSafe(UpdatedComponent->GetOwner()), UpdatedComponent->GetOwnerRole(), *ProposedAdjustment.ToCompactString(), bMoved);
+
+			// Still stuck in penetration?
+			if (!bMoved && SweepOutHit.bStartPenetrating)
+			{
+				// Combine two MTD results to get a new direction that gets out of multiple surfaces.
+				const FVector SecondMTD = UMovementUtils::ComputePenetrationAdjustment(SweepOutHit);
+				const FVector CombinedMTD = ProposedAdjustment + SecondMTD;
+
+				if (SecondMTD != ProposedAdjustment && !CombinedMTD.IsZero())
+				{
+					bMoved = TestMoveComponent_Internal(MovingComps, StartLocation, StartLocation+CombinedMTD, StartRotation, StartRotation, /*bShouldSweep*/ true, MoveCompFlags, SweepOutHit);
+
+					UE_LOG(LogMover, VeryVerbose, TEXT("FindMoveToResolveInitialPenetration_Internal combined: %s (role %i) CombinedAdjustment=%s DidMove=%i"),
+						*GetNameSafe(UpdatedComponent->GetOwner()), UpdatedComponent->GetOwnerRole(), *CombinedMTD.ToCompactString(), bMoved);
+				}
+
+			}
+
+			// Still stuck?
+			if (!bMoved)
+			{
+				const UMoverComponent* MoverComp = MovingComps.MoverComponent.Get();
+
+				// Try moving the proposed adjustment plus the attempted move direction. This can sometimes get out of penetrations with multiple objects
+				const FVector FullMoveDelta = UPlanarConstraintUtils::ConstrainDirectionToPlane(MoverComp->GetPlanarConstraint(), (PenetratingHit.TraceEnd - PenetratingHit.TraceStart));
+				if (!FullMoveDelta.IsZero())
+				{
+					const FVector TargetLocationWithAdjustment = StartLocation + ProposedAdjustment + FullMoveDelta;
+					bMoved = TestMoveComponent_Internal(MovingComps, StartLocation, TargetLocationWithAdjustment, StartRotation, StartRotation, /*bShouldSweep*/ true, MoveCompFlags, SweepOutHit);
+
+					UE_LOG(LogMover, VeryVerbose, TEXT("FindMoveToResolveInitialPenetration_Internal multiple: %s (role %i) TargetLocationWithAdjustment=%s DidMove=%i"),
+						*GetNameSafe(UpdatedComponent->GetOwner()), UpdatedComponent->GetOwnerRole(), *TargetLocationWithAdjustment.ToCompactString(), bMoved);
+				}
+			}
+
+			if (bMoved)
+			{
+				const FVector FinalAdjustmentLocation = SweepOutHit.TraceStart + ((SweepOutHit.TraceEnd-SweepOutHit.TraceStart) * SweepOutHit.Time);
+				OutAdjustmentDelta = FinalAdjustmentLocation - StartLocation;
+				return true;
+			}
+		}
+	}
+
+	OutAdjustmentDelta = FVector::ZeroVector;
+	return false;
+}
+
+static bool ShouldIgnoreHitResult(const UWorld* InWorld, FHitResult const& TestHit, FVector const& MovementDirDenormalized, const AActor* MovingActor, EMoveComponentFlags MoveFlags)
+{
+	if (TestHit.bBlockingHit)
+	{
+		// check "ignore bases" functionality
+		if ((MoveFlags & MOVECOMP_IgnoreBases) && MovingActor)	//we let overlap components go through because their overlap is still needed and will cause beginOverlap/endOverlap events
+		{
+			// ignore if there's a base relationship between moving actor and hit actor
+			AActor const* const HitActor = TestHit.HitObjectHandle.FetchActor();
+			if (HitActor)
+			{
+				if (MovingActor->IsBasedOnActor(HitActor) || HitActor->IsBasedOnActor(MovingActor))
+				{
+					return true;
+				}
+			}
+		}
+
+		// If we started penetrating, we may want to ignore it if we are moving out of penetration.
+		// This helps prevent getting stuck in walls.
+		if ((TestHit.Distance < HitDistanceToleranceCVar || TestHit.bStartPenetrating) && !(MoveFlags & MOVECOMP_NeverIgnoreBlockingOverlaps))
+		{
+			const float DotTolerance = InitialOverlapToleranceCVar;
+
+			// Dot product of movement direction against 'exit' direction
+			const FVector MovementDir = MovementDirDenormalized.GetSafeNormal();
+			const float MoveDot = (TestHit.ImpactNormal | MovementDir);
+
+			const bool bMovingOut = MoveDot > DotTolerance;
+
+			// If we are moving out, ignore this result!
+			if (bMovingOut)
+			{
+				return true;
+			}
+		}
+	}
+
+	return false;
+}
+
+static void PullBackHit(FHitResult& Hit, const FVector& Start, const FVector& End, const float Dist)
+{
+	const float DesiredTimeBack = FMath::Clamp(0.1f, 0.1f / Dist, 1.f / Dist) + 0.001f;
+	Hit.Time = FMath::Clamp(Hit.Time - DesiredTimeBack, 0.f, 1.f);
+}
+
+
+
+// Returns true if any of the sweeping movement could occur
+bool UAsyncMovementUtils::TestMoveComponent_Internal(const FMovingComponentSet& MovingComps, const FVector& StartLocation, const FVector& TargetLocation, const FQuat& StartRotation, const FQuat& TargetRotation, bool bShouldSweep, EMoveComponentFlags MoveCompFlags, FHitResult& OutHit)
+{
+	UPrimitiveComponent* UpdatedPrimitive = MovingComps.UpdatedPrimitive.Get();
+	if (!UpdatedPrimitive)
+	{
+		OutHit.Reset(1.f);
+		return false;
+	}
+
+	FVector ConstrainedDelta = TargetLocation - StartLocation;
+
+	if (MovingComps.MoverComponent.IsValid())
+	{
+		ConstrainedDelta = UPlanarConstraintUtils::ConstrainDirectionToPlane(MovingComps.MoverComponent->GetPlanarConstraint(), ConstrainedDelta);
+	}
+
+	const AActor* OwningActor = UpdatedPrimitive->GetOwner();
+
+	const FVector TraceStart = StartLocation;
+	const FVector TraceEnd = StartLocation + ConstrainedDelta;
+	float DeltaSizeSq = (TraceEnd - TraceStart).SizeSquared();				// Recalc here to account for precision loss of float addition
+	const FQuat InitialRotationQuat = StartRotation;
+
+	// If we aren't sweeping, always consider the full move allowed
+	if (!bShouldSweep)
+	{
+		OutHit.Init(TraceStart, TraceEnd);
+		return true;
+	}
+
+	// ComponentSweepMulti does nothing if moving < KINDA_SMALL_NUMBER in distance, so it's important to not try to sweep distances smaller than that. 
+	constexpr float MinMovementDistSq = FMath::Square(4.f * UE_KINDA_SMALL_NUMBER);
+	if (DeltaSizeSq <= MinMovementDistSq)
+	{
+		// Skip if no vector or rotation.
+		if (TargetRotation.Equals(InitialRotationQuat, SCENECOMPONENT_QUAT_TOLERANCE))
+		{
+			OutHit.Init(TraceStart, TraceEnd);
+			return true;
+		}
+
+		DeltaSizeSq = 0.f;
+	}
+
+
+	TArray<FHitResult> AllHits;
+	FVector NewLocation = TraceStart;
+
+	// Perform movement collision checking if needed for this actor.
+	const bool bCollisionEnabled = UpdatedPrimitive->IsQueryCollisionEnabled();
+	UWorld* const MyWorld = UpdatedPrimitive->GetWorld();
+	if (MyWorld && bCollisionEnabled && (DeltaSizeSq > 0.f))
+	{
+		static const FName TraceTagName = TEXT("SweepTestMoverComponent");
+
+		FComponentQueryParams Params(SCENE_QUERY_STAT(SweepTestMoverComponent), OwningActor);
+		FCollisionResponseParams ResponseParam;
+		UpdatedPrimitive->InitSweepCollisionParams(Params, ResponseParam);
+		Params.bIgnoreTouches = true;
+		Params.TraceTag = TraceTagName;
+
+		const bool bHadBlockingHit = MyWorld->ComponentSweepMulti(AllHits, UpdatedPrimitive, TraceStart, TraceEnd, InitialRotationQuat, Params);
+
+		if (AllHits.Num() > 0)
+		{
+			const float DeltaSize = FMath::Sqrt(DeltaSizeSq);
+			for (int32 HitIdx = 0; HitIdx < AllHits.Num(); HitIdx++)
+			{
+				PullBackHit(AllHits[HitIdx], TraceStart, TraceEnd, DeltaSize);
+			}
+		}
+
+
+		// Find the best blocking hit from AllHits
+		if (bHadBlockingHit)
+		{
+			int32 BlockingHitIndex = INDEX_NONE;
+			float BlockingHitNormalDotDelta = UE_BIG_NUMBER;
+
+			for (int32 HitIdx = 0; HitIdx < AllHits.Num(); ++HitIdx)
+			{
+				const FHitResult& TestHit = AllHits[HitIdx];
+				if (TestHit.bBlockingHit)
+				{
+					if (!ShouldIgnoreHitResult(MyWorld, TestHit, ConstrainedDelta, OwningActor, MoveCompFlags) && !UpdatedPrimitive->ShouldComponentIgnoreHitResult(TestHit, MoveCompFlags))
+					{
+						if (TestHit.bStartPenetrating)
+						{
+							// We may have multiple initial hits, and want to choose the one with the normal most opposed to our movement.
+							const float NormalDotDelta = (TestHit.ImpactNormal | ConstrainedDelta);
+							if (NormalDotDelta < BlockingHitNormalDotDelta)
+							{
+								BlockingHitNormalDotDelta = NormalDotDelta;
+								BlockingHitIndex = HitIdx;
+							}
+						}
+						else if (BlockingHitIndex == INDEX_NONE)
+						{
+							// First non-overlapping blocking hit should be used, if an overlapping hit was not.
+							// This should be the only non-overlapping blocking hit, and last in the results.
+							BlockingHitIndex = HitIdx;
+							break;
+						}
+					}
+				}
+			}
+
+			if (BlockingHitIndex != INDEX_NONE)
+			{
+				OutHit = AllHits[BlockingHitIndex];
+				return OutHit.Time > 0.f;	// consider there to be no movement if the blocking hit started immediately (aka start in penetration)
+			}
+		}
+	}
+
+	// No blocking hits occurred, so the full movement was allowed
+	OutHit.Init(TraceStart, TraceEnd);
+	return true;
+
+}
+
diff --git a/Engine/Plugins/Experimental/Mover/Source/Mover/Public/DefaultMovementSet/Modes/AsyncFlyingMode.h b/Engine/Plugins/Experimental/Mover/Source/Mover/Public/DefaultMovementSet/Modes/AsyncFlyingMode.h
new file mode 100644
--- /dev/null
+++ b/Engine/Plugins/Experimental/Mover/Source/Mover/Public/DefaultMovementSet/Modes/AsyncFlyingMode.h
@@ -0,0 +1,35 @@
+// Copyright Epic Games, Inc. All Rights Reserved.
+
+#pragma once
+
+#include "CoreMinimal.h"
+#include "MovementMode.h"
+#include "AsyncFlyingMode.generated.h"
+
+class UCommonLegacyMovementSettings;
+
+
+/**
+ * AsyncFlyingMode: a default movement mode for moving through the air freely, but still interacting with blocking geometry. The
+ * moving actor will remain upright vs the movement plane.
+ * This mode is threadsafe, and simulates movement without actually modifying any scene component(s).
+ */
+UCLASS(Blueprintable, BlueprintType, Experimental)
+class MOVER_API UAsyncFlyingMode : public UBaseMovementMode
+{
+	GENERATED_UCLASS_BODY()
+
+
+public:
+	UFUNCTION(BlueprintCallable, Category=Mover)
+	virtual void OnGenerateMove(const FMoverTickStartData& StartState, const FMoverTimeStep& TimeStep, FProposedMove& OutProposedMove) const override;
+
+	UFUNCTION(BlueprintCallable, Category=Mover)
+	virtual void OnSimulationTick(const FSimulationTickParams& Params, FMoverTickEndData& OutputState) override;
+
+protected:
+	virtual void OnRegistered(const FName ModeName) override;
+	virtual void OnUnregistered() override;
+
+	TObjectPtr<const UCommonLegacyMovementSettings> CommonLegacySettings;
+};
diff --git a/Engine/Plugins/Experimental/Mover/Source/Mover/Public/MoveLibrary/AsyncMovementUtils.h b/Engine/Plugins/Experimental/Mover/Source/Mover/Public/MoveLibrary/AsyncMovementUtils.h
new file mode 100644
--- /dev/null
+++ b/Engine/Plugins/Experimental/Mover/Source/Mover/Public/MoveLibrary/AsyncMovementUtils.h
@@ -0,0 +1,54 @@
+// Copyright Epic Games, Inc. All Rights Reserved.
+
+#pragma once
+
+#include "Kismet/BlueprintFunctionLibrary.h"
+#include "Components/SceneComponent.h"
+#include "MoverDataModelTypes.h"
+
+#include "AsyncMovementUtils.generated.h"
+
+struct FMovementRecord;
+class UMoverComponent;
+struct FOptionalFloorCheckResult;
+
+
+/**
+ * AsyncMovementUtils: a collection of stateless static BP-accessible functions focused on testing potential movements in a 
+ * threadsafe manner without actually causing immediate changes.
+ */
+UCLASS(Experimental)
+class MOVER_API UAsyncMovementUtils : public UBlueprintFunctionLibrary
+{
+	GENERATED_BODY()
+
+public:
+	/** Tests potential movement of a component without actually moving it, taking penetration resolution issues into account first.  
+	 * Returns true if any movement was possible
+	 * Modifies OutHit with final movement hit data
+	 * Appends to InOutMoveRecord with any movement substeps
+	 */
+	UFUNCTION(BlueprintCallable, Category = Mover)
+	static bool TestDepenetratingMove(const FMovingComponentSet& MovingComps, const FVector& StartLocation, const FVector& TargetLocation, const FQuat& StartRotation, const FQuat& TargetRotation, bool bShouldSweep, FHitResult& OutHit, FMovementRecord& InOutMoveRecord);
+
+	/** Tests potential movement of a component sliding along a surface, without actually moving it. 
+	 * Returns the percent of time applied, with 0.0 meaning no movement would occur. 
+	 * Modifies InOutHit with final movement hit data
+	 * Appends to InOutMoveRecord with any movement substeps
+	 */
+	UFUNCTION(BlueprintCallable, Category = Mover)
+	static float TestSlidingMoveAlongHitSurface(const FMovingComponentSet& MovingComps, const FVector& OriginalMoveDelta, const FVector& LocationAtHit, const FQuat& TargetRotation, FHitResult& InOutHit, FMovementRecord& InOutMoveRecord);
+
+
+	/** Attempts to find a move that would resolve an initially penetrating blockage.
+	 * Returns true if an adjustment was found. The adjustment will be in OutAdjustmentDelta.
+	 */
+	static bool FindMoveToResolveInitialPenetration_Internal(const FMovingComponentSet& MovingComps, const FVector& StartLocation, const FQuat& TargetRotation, const FHitResult& PenetratingHit, EMoveComponentFlags MoveComponentFlags, FVector& OutAdjustmentDelta);
+
+
+	/** Tests a move of a component without actually moving it. 
+	 * Returns true if any motion could occur. Detailed blocking hit results are written to OutHit.  
+	 * If not sweeping, full movement will always be allowed even if new blocking overlaps would occur.
+	 */
+	static bool TestMoveComponent_Internal(const FMovingComponentSet& MovingComps, const FVector& StartLocation, const FVector& NewLocation, const FQuat& StartRotation, const FQuat& NewRotation, bool bShouldSweep, EMoveComponentFlags MoveCompFlags, FHitResult& OutHit);
+};
