Mover: fix static analysis warning. No functional changes.

#rnx

diff --git a/Engine/Plugins/Experimental/Mover/Source/Mover/Private/MoveLibrary/AsyncMovementUtils.cpp b/Engine/Plugins/Experimental/Mover/Source/Mover/Private/MoveLibrary/AsyncMovementUtils.cpp
--- a/Engine/Plugins/Experimental/Mover/Source/Mover/Private/MoveLibrary/AsyncMovementUtils.cpp
+++ b/Engine/Plugins/Experimental/Mover/Source/Mover/Private/MoveLibrary/AsyncMovementUtils.cpp
@@ -135,8 +135,14 @@ bool UAsyncMovementUtils::FindMoveToResolveInitialPenetration_Internal(const FMo
 {
 	if (ensureMsgf(PenetratingHit.bStartPenetrating, TEXT("Expected a hit that started penetrating. Will not attempt adjustment.")))
 	{
+		if (!MovingComps.UpdatedComponent.IsValid() || !MovingComps.UpdatedPrimitive.IsValid())
+		{
+			return false;
+		}
+
 		const USceneComponent* UpdatedComponent = MovingComps.UpdatedComponent.Get();
 		const UPrimitiveComponent* UpdatedPrimitive = MovingComps.UpdatedPrimitive.Get();
+
 		const AActor* ActorOwner = UpdatedPrimitive->GetOwner();
 		if (!ActorOwner)
 		{
@@ -147,7 +153,7 @@ bool UAsyncMovementUtils::FindMoveToResolveInitialPenetration_Internal(const FMo
 		ProposedAdjustment = UPlanarConstraintUtils::ConstrainDirectionToPlane(MovingComps.MoverComponent->GetPlanarConstraint(), ProposedAdjustment);
 
 		// SceneComponent can't be in penetration, so this function really only applies to PrimitiveComponent.
-		if (!ProposedAdjustment.IsZero() && UpdatedPrimitive)
+		if (!ProposedAdjustment.IsZero())
 		{
 			// Attempts to adjust:
 			//	1) Check if our component will fit with the proposed adjustment. If so, accept the proposed adjustment.
