Mover: first pass on async falling, walking, navwalking mode implementations.
Known issues with navwalking collision settings, lack of floor check option support while not walking, not rotating while sliding along a wall (walking)

diff --git a/Engine/Plugins/Experimental/Mover/Source/Mover/Private/DefaultMovementSet/Modes/AsyncFallingMode.cpp b/Engine/Plugins/Experimental/Mover/Source/Mover/Private/DefaultMovementSet/Modes/AsyncFallingMode.cpp
new file mode 100644
--- /dev/null
+++ b/Engine/Plugins/Experimental/Mover/Source/Mover/Private/DefaultMovementSet/Modes/AsyncFallingMode.cpp
@@ -0,0 +1,368 @@
+// Copyright Epic Games, Inc. All Rights Reserved.
+
+#include "DefaultMovementSet/Modes/AsyncFallingMode.h"
+#include "Components/SkeletalMeshComponent.h"
+#include "MoverComponent.h"
+#include "MoveLibrary/AirMovementUtils.h"
+#include "MoveLibrary/AsyncMovementUtils.h"
+#include "MoveLibrary/BasedMovementUtils.h"
+#include "MoveLibrary/FloorQueryUtils.h"
+#include "MoveLibrary/MovementUtils.h"
+#include "DefaultMovementSet/Settings/CommonLegacyMovementSettings.h"
+
+#include "DrawDebugHelpers.h"
+
+
+#include UE_INLINE_GENERATED_CPP_BY_NAME(AsyncFallingMode)
+
+UAsyncFallingMode::UAsyncFallingMode(const FObjectInitializer& ObjectInitializer)
+	: Super(ObjectInitializer)
+	, bCancelVerticalSpeedOnLanding(true)
+	, AirControlPercentage(0.4f)
+	, FallingDeceleration(200.0f)
+	, OverTerminalSpeedFallingDeceleration(800.0f)
+	, TerminalMovementPlaneSpeed(1500.0f)
+	, bShouldClampTerminalVerticalSpeed(true)
+	, VerticalFallingDeceleration(4000.0f)
+	, TerminalVerticalSpeed(2000.0f)
+{
+	SharedSettingsClasses.Add(UCommonLegacyMovementSettings::StaticClass());
+
+	GameplayTags.AddTag(Mover_IsInAir);
+	GameplayTags.AddTag(Mover_IsFalling);
+}
+
+constexpr float VERTICAL_SLOPE_NORMAL_MAX_DOT = 0.001f; // Slope is vertical if Abs(UpDirection) <= this threshold. Accounts for precision problems that sometimes angle normals slightly off horizontal for vertical surface.
+
+void UAsyncFallingMode::OnGenerateMove(const FMoverTickStartData& StartState, const FMoverTimeStep& TimeStep, FProposedMove& OutProposedMove) const
+{
+	const UMoverComponent* MoverComp = GetMoverComponent();
+	const FCharacterDefaultInputs* CharacterInputs = StartState.InputCmd.InputCollection.FindDataByType<FCharacterDefaultInputs>();
+	const FMoverDefaultSyncState* StartingSyncState = StartState.SyncState.SyncStateCollection.FindDataByType<FMoverDefaultSyncState>();
+	check(StartingSyncState);
+
+	if (!CommonLegacySettings.IsValid())
+	{
+		return;
+	}
+
+	const float DeltaSeconds = TimeStep.StepMs * 0.001f;
+
+	FVector UpDirection = MoverComp->GetUpDirection();
+	
+	// We don't want velocity limits to take the falling velocity component into account, since it is handled 
+	//   separately by the terminal velocity of the environment.
+	const FVector StartVelocity = StartingSyncState->GetVelocity_WorldSpace();
+	const FVector StartHorizontalVelocity =  FVector::VectorPlaneProject(StartVelocity, UpDirection);
+
+	FFreeMoveParams Params;
+	if (CharacterInputs)
+	{
+		Params.MoveInputType = CharacterInputs->GetMoveInputType();
+		const bool bMaintainInputMagnitude = true;
+		Params.MoveInput = UPlanarConstraintUtils::ConstrainDirectionToPlane(MoverComp->GetPlanarConstraint(), CharacterInputs->GetMoveInput_WorldSpace(), bMaintainInputMagnitude);
+	}
+	else
+	{
+		Params.MoveInputType = EMoveInputType::Invalid;
+		Params.MoveInput = FVector::ZeroVector;
+	}
+
+	Params.MoveInput *= AirControlPercentage;
+	// Don't care about up axis input since falling - if up input matters that should probably be a different movement mode
+	Params.MoveInput = FVector::VectorPlaneProject(Params.MoveInput, UpDirection);
+	
+	FRotator IntendedOrientation_WorldSpace;
+	// If there's no intent from input to change orientation, use the current orientation
+	if (!CharacterInputs || CharacterInputs->OrientationIntent.IsNearlyZero())
+	{
+		IntendedOrientation_WorldSpace = StartingSyncState->GetOrientation_WorldSpace();
+	}
+	else
+	{
+		IntendedOrientation_WorldSpace = CharacterInputs->GetOrientationIntentDir_WorldSpace().ToOrientationRotator();
+	}
+
+	IntendedOrientation_WorldSpace = UMovementUtils::ApplyGravityToOrientationIntent(IntendedOrientation_WorldSpace, MoverComp->GetWorldToGravityTransform(), CommonLegacySettings->bShouldRemainVertical);
+	
+	Params.OrientationIntent = IntendedOrientation_WorldSpace;
+	Params.PriorVelocity = StartHorizontalVelocity;
+	Params.PriorOrientation = StartingSyncState->GetOrientation_WorldSpace();
+	Params.DeltaSeconds = DeltaSeconds;
+	Params.TurningRate = CommonLegacySettings->TurningRate;
+	Params.TurningBoost = CommonLegacySettings->TurningBoost;
+	Params.MaxSpeed = CommonLegacySettings->MaxSpeed;
+	Params.Acceleration = CommonLegacySettings->Acceleration;
+	Params.Deceleration = FallingDeceleration;
+	Params.WorldToGravityQuat = MoverComp->GetWorldToGravityTransform();
+	Params.bUseAccelerationForVelocityMove = CommonLegacySettings->bUseAccelerationForVelocityMove;
+
+	// Check if any current velocity values are over our terminal velocity - if so limit the move input in that direction and apply OverTerminalVelocityFallingDeceleration
+	if (Params.MoveInput.Dot(StartVelocity) > 0 && StartHorizontalVelocity.Size() >= TerminalMovementPlaneSpeed)
+	{
+		Params.Deceleration = OverTerminalSpeedFallingDeceleration;
+	}
+	
+	UMoverBlackboard* SimBlackboard = MoverComp->GetSimBlackboard_Mutable();
+	FFloorCheckResult LastFloorResult;
+	// limit our moveinput based on the floor we're on
+	if (SimBlackboard && SimBlackboard->TryGet(CommonBlackboard::LastFloorResult, LastFloorResult))
+	{
+		if (LastFloorResult.HitResult.IsValidBlockingHit() && LastFloorResult.HitResult.Normal.Dot(UpDirection) > VERTICAL_SLOPE_NORMAL_MAX_DOT && !LastFloorResult.IsWalkableFloor())
+		{
+			// If acceleration is into the wall, limit contribution.
+			if (FVector::DotProduct(Params.MoveInput, LastFloorResult.HitResult.Normal) < 0.f)
+			{
+				// Allow movement parallel to the wall, but not into it because that may push us up.
+				const FVector FallingHitNormal = FVector::VectorPlaneProject( LastFloorResult.HitResult.Normal, -UpDirection).GetSafeNormal();
+				Params.MoveInput = FVector::VectorPlaneProject(Params.MoveInput, FallingHitNormal);
+			}
+		}
+	}
+	
+	OutProposedMove = UAirMovementUtils::ComputeControlledFreeMove(Params);
+	const FVector VelocityWithGravity = StartVelocity + UMovementUtils::ComputeVelocityFromGravity(MoverComp->GetGravityAcceleration(), DeltaSeconds);
+
+	//  If we are going faster than TerminalVerticalVelocity apply VerticalFallingDeceleration otherwise reset Z velocity to before we applied deceleration 
+	if (VelocityWithGravity.GetAbs().Dot(UpDirection) > TerminalVerticalSpeed)
+	{
+		if (bShouldClampTerminalVerticalSpeed)
+		{
+			UMovementUtils::SetGravityVerticalComponent(OutProposedMove.LinearVelocity, -TerminalVerticalSpeed, UpDirection);
+		}
+		else
+		{
+			float DesiredDeceleration = FMath::Abs(TerminalVerticalSpeed - VelocityWithGravity.GetAbs().Dot(UpDirection)) / DeltaSeconds;
+			float DecelerationToApply = FMath::Min(DesiredDeceleration, VerticalFallingDeceleration);
+			DecelerationToApply = FMath::Sign(VelocityWithGravity.Dot(UpDirection)) * DecelerationToApply * DeltaSeconds;
+			FVector MaxUpDirVelocity = VelocityWithGravity * UpDirection - (UpDirection * DecelerationToApply);
+			
+			UMovementUtils::SetGravityVerticalComponent(OutProposedMove.LinearVelocity, MaxUpDirVelocity.Dot(UpDirection), UpDirection);
+		}
+	}
+	else
+	{
+		UMovementUtils::SetGravityVerticalComponent(OutProposedMove.LinearVelocity, VelocityWithGravity.Dot(UpDirection), UpDirection);
+	}
+}
+
+void UAsyncFallingMode::OnSimulationTick(const FSimulationTickParams& Params, FMoverTickEndData& OutputState)
+{
+	if (!CommonLegacySettings.IsValid())
+	{
+		return;
+	}
+
+	UMoverComponent* MoverComponent = GetMoverComponent();
+	const FMoverTickStartData& StartState = Params.StartState;
+	USceneComponent* UpdatedComponent = Params.MovingComps.UpdatedComponent.Get();
+	FProposedMove ProposedMove = Params.ProposedMove;
+
+	const FCharacterDefaultInputs* CharacterInputs = StartState.InputCmd.InputCollection.FindDataByType<FCharacterDefaultInputs>();
+	const FMoverDefaultSyncState* StartingSyncState = StartState.SyncState.SyncStateCollection.FindDataByType<FMoverDefaultSyncState>();
+	check(StartingSyncState);
+
+	FMoverDefaultSyncState& OutputSyncState = OutputState.SyncState.SyncStateCollection.FindOrAddMutableDataByType<FMoverDefaultSyncState>();
+
+	const float DeltaSeconds = Params.TimeStep.StepMs * 0.001f;
+	float PctTimeApplied = 0.f;
+
+	FMovementRecord MoveRecord;
+	MoveRecord.SetDeltaSeconds(DeltaSeconds);
+	
+	UMoverBlackboard* SimBlackboard = MoverComponent->GetSimBlackboard_Mutable();
+
+	SimBlackboard->Invalidate(CommonBlackboard::LastFloorResult);	// falling = no valid floor
+	SimBlackboard->Invalidate(CommonBlackboard::LastFoundDynamicMovementBase);
+
+	OutputSyncState.MoveDirectionIntent = (ProposedMove.bHasDirIntent ? ProposedMove.DirectionIntent : FVector::ZeroVector);
+
+	// Use the orientation intent directly. If no intent is provided, use last frame's orientation. Note that we are assuming rotation changes can't fail. 
+	const FRotator StartingOrient = StartingSyncState->GetOrientation_WorldSpace();
+	FRotator TargetOrient = StartingOrient;
+
+	// Apply orientation changes (if any)
+	if (!UMovementUtils::IsAngularVelocityZero(ProposedMove.AngularVelocity))
+	{
+		TargetOrient += (ProposedMove.AngularVelocity * DeltaSeconds);
+	}
+
+	const FVector StartingFallingVelocity = StartingSyncState->GetVelocity_WorldSpace();
+
+	//FVector MoveDelta = 0.5f * (PriorFallingVelocity + ProposedMove.LinearVelocity) * DeltaSeconds; 	// TODO: revive midpoint integration
+	const FVector StartLocation = StartingSyncState->GetLocation_WorldSpace();
+	const FVector TargetLocation = StartLocation + (ProposedMove.LinearVelocity * DeltaSeconds);
+
+	const FQuat StartRotation = StartingOrient.Quaternion();
+	const FQuat TargetRotation = TargetOrient.Quaternion();
+
+
+	FHitResult SweepHit(1.f);
+
+	FVector LocationInProgress = StartLocation;
+	FQuat   RotationInProgress = StartRotation;
+
+	const FVector MoveDelta = TargetLocation - StartLocation;
+
+	const bool bWouldMove = UAsyncMovementUtils::TestDepenetratingMove(Params.MovingComps, StartLocation, TargetLocation, StartRotation, TargetRotation, /* bShouldSweep */ true, OUT SweepHit, IN OUT MoveRecord);
+
+	LocationInProgress = StartLocation + ((TargetLocation - StartLocation) * SweepHit.Time);
+	RotationInProgress = FQuat::Slerp(StartRotation, TargetRotation, SweepHit.Time);
+
+	// Compute final velocity based on how long we actually go until we get a hit.
+
+	FFloorCheckResult LandingFloor;
+
+	// Handle impact, whether it's a landing surface or something to slide on
+	if (SweepHit.IsValidBlockingHit())
+	{
+		float LastMoveTimeSlice = DeltaSeconds;
+		float SubTimeTickRemaining = LastMoveTimeSlice * (1.f - SweepHit.Time);
+
+		PctTimeApplied += SweepHit.Time * (1.f - PctTimeApplied);
+
+		// Check for hitting a landing surface
+		if (UAirMovementUtils::IsValidLandingSpot(Params.MovingComps, LocationInProgress,
+				SweepHit, CommonLegacySettings->FloorSweepDistance, CommonLegacySettings->MaxWalkSlopeCosine, OUT LandingFloor))
+		{
+			CaptureFinalState(LocationInProgress, RotationInProgress.Rotator(), LandingFloor, DeltaSeconds, DeltaSeconds * PctTimeApplied, OutputSyncState, OutputState, IN OUT MoveRecord);
+			return;
+		}
+
+		
+		LandingFloor.HitResult = SweepHit;
+		SimBlackboard->Set(CommonBlackboard::LastFloorResult, LandingFloor);
+		
+		FMoverOnImpactParams ImpactParams(DefaultModeNames::Falling, SweepHit, MoveDelta);
+		MoverComponent->HandleImpact(ImpactParams);
+
+		// We didn't land on a walkable surface, so let's try to slide along it
+		const float PctOfTimeUsedForSliding = UAirMovementUtils::TestFallingMoveAlongHitSurface(
+			Params.MovingComps, TargetLocation-StartLocation, LocationInProgress, TargetRotation, 
+			/*bHandleImpact=*/true, CommonLegacySettings->FloorSweepDistance, CommonLegacySettings->MaxWalkSlopeCosine,
+			IN OUT SweepHit, OUT LandingFloor, IN OUT MoveRecord);
+
+		if (PctOfTimeUsedForSliding > 0.f)
+		{
+			LocationInProgress = SweepHit.TraceStart + ((SweepHit.TraceEnd - SweepHit.TraceStart) * PctOfTimeUsedForSliding);
+			RotationInProgress = FQuat::Slerp(RotationInProgress, TargetRotation, PctOfTimeUsedForSliding);
+		}
+
+		PctTimeApplied += SweepHit.Time * (1.f - PctTimeApplied);
+
+		if (LandingFloor.IsWalkableFloor())
+		{
+			CaptureFinalState(LocationInProgress, RotationInProgress.Rotator(), LandingFloor, DeltaSeconds, DeltaSeconds * PctTimeApplied, OutputSyncState, OutputState, MoveRecord);
+			return;
+		}
+	}
+	else
+	{
+		// This indicates an unimpeded full move
+		PctTimeApplied = 1.f;
+	}
+	
+	CaptureFinalState(LocationInProgress, RotationInProgress.Rotator(), LandingFloor, DeltaSeconds, DeltaSeconds* PctTimeApplied, OutputSyncState, OutputState, MoveRecord);
+}
+
+
+void UAsyncFallingMode::OnRegistered(const FName ModeName)
+{
+	Super::OnRegistered(ModeName);
+
+	CommonLegacySettings = GetMoverComponent()->FindSharedSettings<UCommonLegacyMovementSettings>();
+	ensureMsgf(CommonLegacySettings.IsValid(), TEXT("Failed to find instance of CommonLegacyMovementSettings on %s. Movement may not function properly."), *GetPathNameSafe(this));
+}
+
+
+void UAsyncFallingMode::OnUnregistered()
+{
+	CommonLegacySettings = nullptr;
+
+	Super::OnUnregistered();
+}
+
+void UAsyncFallingMode::ProcessLanded(const FFloorCheckResult& FloorResult, FVector& Velocity, FRelativeBaseInfo& BaseInfo, FMoverTickEndData& TickEndData) const
+{
+	// TODO: Refactor landed events for async movement. For now, leave the immediate event, but listeners are limited in what they can do from a worker thread.
+
+	const UMoverComponent* MoverComp = GetMoverComponent();
+	UMoverBlackboard* SimBlackboard = MoverComp->GetSimBlackboard_Mutable();
+
+	FName NextMovementMode = NAME_None; 
+	// if we can walk on the floor we landed on
+	if (FloorResult.IsWalkableFloor())
+	{
+		if (bCancelVerticalSpeedOnLanding)
+		{
+			const FPlane MovementPlane(FVector::ZeroVector, MoverComp->GetUpDirection());
+			Velocity = UMovementUtils::ConstrainToPlane(Velocity, MovementPlane, false);
+		}
+		else
+		{
+			Velocity = FVector::VectorPlaneProject(Velocity, FloorResult.HitResult.Normal);
+		}
+		
+		// Transfer to LandingMovementMode (usually walking), and cache any floor / movement base info
+		NextMovementMode = CommonLegacySettings->GroundMovementModeName;
+
+		SimBlackboard->Set(CommonBlackboard::LastFloorResult, FloorResult);
+
+		if (UBasedMovementUtils::IsADynamicBase(FloorResult.HitResult.GetComponent()))
+		{
+			BaseInfo.SetFromFloorResult(FloorResult);
+		}
+	}
+	// we could check for other surfaces here (i.e. when swimming is implemented we can check the floor hit here and see if we need to go into swimming)
+
+	// This would also be a good spot for implementing some falling physics interactions (i.e. falling into a movable object and pushing it based off of this actors velocity)
+	
+	// if a new mode was set go ahead and switch to it after this tick and broadcast we landed
+	if (!NextMovementMode.IsNone())
+	{
+		TickEndData.MovementEndState.NextModeName = NextMovementMode;
+		OnLanded.Broadcast(NextMovementMode, FloorResult.HitResult);
+	}
+}
+
+void UAsyncFallingMode::CaptureFinalState(const FVector FinalLocation, const FRotator FinalRotation, const FFloorCheckResult& FloorResult, float DeltaSeconds, float DeltaSecondsUsed, FMoverDefaultSyncState& OutputSyncState, FMoverTickEndData& TickEndData, FMovementRecord& Record) const
+{
+	UMoverBlackboard* SimBlackboard = GetMoverComponent()->GetSimBlackboard_Mutable();
+
+	// Check for time refunds
+	constexpr float MinRemainingSecondsToRefund = 0.0001f;	// If we have this amount of time (or more) remaining, give it to the next simulation step.
+
+	if ((DeltaSeconds - DeltaSecondsUsed) >= MinRemainingSecondsToRefund)
+	{
+		const float PctOfTimeRemaining = (1.0f - (DeltaSecondsUsed / DeltaSeconds));
+		TickEndData.MovementEndState.RemainingMs = PctOfTimeRemaining * DeltaSeconds * 1000.f;
+	}
+	else
+	{
+		TickEndData.MovementEndState.RemainingMs = 0.f;
+	}
+	
+	Record.SetDeltaSeconds( DeltaSecondsUsed );
+	
+	FVector EffectiveVelocity = Record.GetRelevantVelocity();
+
+	FRelativeBaseInfo MovementBaseInfo;
+	ProcessLanded(FloorResult, EffectiveVelocity, MovementBaseInfo, TickEndData);
+
+	if (MovementBaseInfo.HasRelativeInfo())
+	{
+		SimBlackboard->Set(CommonBlackboard::LastFoundDynamicMovementBase, MovementBaseInfo);
+
+		OutputSyncState.SetTransforms_WorldSpace( FinalLocation,
+												  FinalRotation,
+												  EffectiveVelocity,
+												  MovementBaseInfo.MovementBase.Get(), MovementBaseInfo.BoneName);
+	}
+	else
+	{
+		OutputSyncState.SetTransforms_WorldSpace( FinalLocation,
+												  FinalRotation,
+												  EffectiveVelocity,
+												  nullptr); // no movement base
+	}
+}
diff --git a/Engine/Plugins/Experimental/Mover/Source/Mover/Private/DefaultMovementSet/Modes/AsyncNavWalkingMode.cpp b/Engine/Plugins/Experimental/Mover/Source/Mover/Private/DefaultMovementSet/Modes/AsyncNavWalkingMode.cpp
new file mode 100644
--- /dev/null
+++ b/Engine/Plugins/Experimental/Mover/Source/Mover/Private/DefaultMovementSet/Modes/AsyncNavWalkingMode.cpp
@@ -0,0 +1,619 @@
+﻿// Copyright Epic Games, Inc. All Rights Reserved.
+
+
+#include "DefaultMovementSet/Modes/AsyncNavWalkingMode.h"
+#include "MoverComponent.h"
+#include "AI/NavigationSystemBase.h"
+#include "AI/Navigation/NavigationDataInterface.h"
+#include "Components/ShapeComponent.h"
+#include "DefaultMovementSet/NavMoverComponent.h"
+#include "DefaultMovementSet/Settings/CommonLegacyMovementSettings.h"
+#include "MoveLibrary/AsyncMovementUtils.h"
+#include "MoveLibrary/FloorQueryUtils.h"
+#include "MoveLibrary/GroundMovementUtils.h"
+#include "MoveLibrary/ModularMovement.h"
+#include "MoveLibrary/MovementUtils.h"
+#include "DrawDebugHelpers.h"
+
+#include UE_INLINE_GENERATED_CPP_BY_NAME(AsyncNavWalkingMode)
+
+UAsyncNavWalkingMode::UAsyncNavWalkingMode()
+	: bSweepWhileNavWalking(true)
+	, bProjectNavMeshWalking(false)
+	, NavMeshProjectionHeightScaleUp(0.67f)
+	, NavMeshProjectionHeightScaleDown(1.0f)
+	, NavMeshProjectionInterval(0.1f)
+	, NavMeshProjectionInterpSpeed(12.f)
+	, NavMeshProjectionTimer(0)
+	, NavMoverComponent(nullptr)
+	, bProjectNavMeshOnBothWorldChannels(true)
+{
+	GameplayTags.AddTag(Mover_IsOnGround);
+	GameplayTags.AddTag(Mover_IsNavWalking);
+}
+
+void UAsyncNavWalkingMode::OnGenerateMove(const FMoverTickStartData& StartState, const FMoverTimeStep& TimeStep, FProposedMove& OutProposedMove) const
+{
+	if (!CommonLegacySettings.IsValid())
+	{
+		return;
+	}
+
+	const UMoverComponent* MoverComp = GetMoverComponent();
+	const FCharacterDefaultInputs* CharacterInputs = StartState.InputCmd.InputCollection.FindDataByType<FCharacterDefaultInputs>();
+	const FMoverDefaultSyncState* StartingSyncState = StartState.SyncState.SyncStateCollection.FindDataByType<FMoverDefaultSyncState>();
+	check(StartingSyncState);
+
+	const float DeltaSeconds = TimeStep.StepMs * 0.001f;
+	FFloorCheckResult LastFloorResult;
+	FVector MovementNormal;
+	FVector UpDirection = MoverComp->GetUpDirection();
+
+	UMoverBlackboard* SimBlackboard = MoverComp->GetSimBlackboard_Mutable();
+
+	// Try to use the floor as the basis for the intended move direction (i.e. try to walk along slopes, rather than into them)
+	if (SimBlackboard && SimBlackboard->TryGet(CommonBlackboard::LastFloorResult, LastFloorResult) && LastFloorResult.IsWalkableFloor())
+	{
+		MovementNormal = LastFloorResult.HitResult.ImpactNormal;
+	}
+	else
+	{
+		MovementNormal = MoverComp->GetUpDirection();
+	}
+
+	FRotator IntendedOrientation_WorldSpace;
+	// If there's no intent from input to change orientation, use the current orientation
+	if (!CharacterInputs || CharacterInputs->OrientationIntent.IsNearlyZero())
+	{
+		IntendedOrientation_WorldSpace = StartingSyncState->GetOrientation_WorldSpace();
+	}
+	else
+	{
+		IntendedOrientation_WorldSpace = CharacterInputs->GetOrientationIntentDir_WorldSpace().ToOrientationRotator();
+	}
+
+	IntendedOrientation_WorldSpace = UMovementUtils::ApplyGravityToOrientationIntent(IntendedOrientation_WorldSpace, MoverComp->GetWorldToGravityTransform(), CommonLegacySettings->bShouldRemainVertical);
+
+	FGroundMoveParams Params;
+
+	if (CharacterInputs)
+	{
+		Params.MoveInputType = CharacterInputs->GetMoveInputType();
+		Params.MoveInput = CharacterInputs->GetMoveInput_WorldSpace();
+	}
+	else
+	{
+		Params.MoveInputType = EMoveInputType::Invalid;
+		Params.MoveInput = FVector::ZeroVector;
+	}
+
+	Params.OrientationIntent = IntendedOrientation_WorldSpace;
+	Params.PriorVelocity = FVector::VectorPlaneProject(StartingSyncState->GetVelocity_WorldSpace(), MovementNormal);
+	Params.PriorOrientation = StartingSyncState->GetOrientation_WorldSpace();
+	Params.GroundNormal = MovementNormal;
+	Params.TurningRate = CommonLegacySettings->TurningRate;
+	Params.TurningBoost = CommonLegacySettings->TurningBoost;
+	Params.MaxSpeed = CommonLegacySettings->MaxSpeed;
+	Params.Acceleration = CommonLegacySettings->Acceleration;
+	Params.Deceleration = CommonLegacySettings->Deceleration;
+	Params.DeltaSeconds = DeltaSeconds;
+	Params.WorldToGravityQuat = MoverComp->GetWorldToGravityTransform();
+	Params.UpDirection = UpDirection;
+	Params.bUseAccelerationForVelocityMove = CommonLegacySettings->bUseAccelerationForVelocityMove;
+
+	if (Params.MoveInput.SizeSquared() > 0.f && !UMovementUtils::IsExceedingMaxSpeed(Params.PriorVelocity, CommonLegacySettings->MaxSpeed))
+	{
+		Params.Friction = CommonLegacySettings->GroundFriction;
+	}
+	else
+	{
+		Params.Friction = CommonLegacySettings->bUseSeparateBrakingFriction ? CommonLegacySettings->BrakingFriction : CommonLegacySettings->GroundFriction;
+		Params.Friction *= CommonLegacySettings->BrakingFrictionFactor;
+	}
+
+	OutProposedMove = UGroundMovementUtils::ComputeControlledGroundMove(Params);
+
+	if (TurnGenerator)
+	{
+		OutProposedMove.AngularVelocity = ITurnGeneratorInterface::Execute_GetTurn(TurnGenerator, IntendedOrientation_WorldSpace, StartState, *StartingSyncState, TimeStep, OutProposedMove, SimBlackboard);
+	}
+}
+
+void UAsyncNavWalkingMode::OnSimulationTick(const FSimulationTickParams& Params, FMoverTickEndData& OutputState)
+{
+	const UMoverComponent* MoverComp = Cast<UMoverComponent>(GetMoverComponent());
+	if (!MoverComp || !NavMoverComponent.IsValid() || !CommonLegacySettings.IsValid())
+	{
+		return;
+	}
+
+	const FMoverTickStartData& StartState = Params.StartState;
+	FProposedMove ProposedMove = Params.ProposedMove;
+	FVector UpDirection = MoverComp->GetUpDirection();
+
+	const FCharacterDefaultInputs* CharacterInputs = StartState.InputCmd.InputCollection.FindDataByType<FCharacterDefaultInputs>();
+	const FMoverDefaultSyncState* StartingSyncState = StartState.SyncState.SyncStateCollection.FindDataByType<FMoverDefaultSyncState>();
+	check(StartingSyncState);
+
+	FMoverDefaultSyncState& OutputSyncState = OutputState.SyncState.SyncStateCollection.FindOrAddMutableDataByType<FMoverDefaultSyncState>();
+
+	const float DeltaSeconds = Params.TimeStep.StepMs * 0.001f;
+	const FVector OrigMoveDelta = ProposedMove.LinearVelocity * DeltaSeconds;
+
+	FMovementRecord MoveRecord;
+	MoveRecord.SetDeltaSeconds(DeltaSeconds);
+
+	OutputSyncState.MoveDirectionIntent = (ProposedMove.bHasDirIntent ? ProposedMove.DirectionIntent : FVector::ZeroVector);
+
+	const FRotator StartingOrient = StartingSyncState->GetOrientation_WorldSpace();
+	FRotator TargetOrient = StartingOrient;
+
+	// Apply orientation changes (if any)
+	if (!UMovementUtils::IsAngularVelocityZero(ProposedMove.AngularVelocity))
+	{
+		TargetOrient += (ProposedMove.AngularVelocity * DeltaSeconds);
+	}
+
+
+	const FVector StartingFeetLocation = NavMoverComponent->GetFeetLocationAt(StartingSyncState->GetLocation_WorldSpace());
+	FVector TargetFeetLocation = StartingFeetLocation + (ProposedMove.LinearVelocity * DeltaSeconds);
+
+	const FVector StartLocation = StartingSyncState->GetLocation_WorldSpace();
+
+	const FQuat StartRotation = StartingOrient.Quaternion();
+	const FQuat TargetRotation = TargetOrient.Quaternion();
+
+
+	FVector LocationInProgress = StartLocation;
+	FQuat RotationInProgress = StartRotation;
+
+	FFloorCheckResult CurrentFloor;
+	const bool bDeltaMoveNearlyZero = OrigMoveDelta.IsNearlyZero();
+	FNavLocation DestNavLocation;
+
+	float SimpleRadius = 0;
+	float SimpleHalfHeight = 0;
+	NavMoverComponent->GetSimpleCollisionCylinder(OUT SimpleRadius, OUT SimpleHalfHeight);
+
+	bool bSameNavLocation = false;
+	if (CachedNavLocation.NodeRef != INVALID_NAVNODEREF)
+	{
+		if (bProjectNavMeshWalking)
+		{
+			const float DistSq = UMovementUtils::ProjectToGravityFloor(StartingFeetLocation - CachedNavLocation.Location, UpDirection).SizeSquared();
+			const float DistDot = FMath::Abs((StartingFeetLocation - CachedNavLocation.Location).Dot(UpDirection));
+
+			const float TotalCapsuleHeight = SimpleHalfHeight * 2.0f;
+			const float ProjectionScale = (StartingFeetLocation.Dot(UpDirection) > CachedNavLocation.Location.Dot(UpDirection)) ? NavMeshProjectionHeightScaleUp : NavMeshProjectionHeightScaleDown;
+			const float DistThr = TotalCapsuleHeight * FMath::Max(0.f, ProjectionScale);
+
+			bSameNavLocation = (DistSq <= UE_KINDA_SMALL_NUMBER) && (DistDot < DistThr);
+		}
+		else
+		{
+			bSameNavLocation = CachedNavLocation.Location.Equals(StartingFeetLocation);
+		}
+
+		if (bDeltaMoveNearlyZero && bSameNavLocation)
+		{
+			if (const INavigationDataInterface* NavData = GetNavData())
+			{
+				if (!NavData->IsNodeRefValid(CachedNavLocation.NodeRef))
+				{
+					CachedNavLocation.NodeRef = INVALID_NAVNODEREF;
+					bSameNavLocation = false;
+				}
+			}
+		}
+	}
+
+	if (bDeltaMoveNearlyZero && bSameNavLocation)
+	{
+		DestNavLocation = CachedNavLocation;
+		UE_LOG(LogMover, VeryVerbose, TEXT("%s using cached navmesh location! (bProjectNavMeshWalking = %d)"), *GetNameSafe(GetMoverComponent()->GetOwner()), bProjectNavMeshWalking);
+	}
+	else
+	{
+		// Start the trace from the vertical location of the last valid trace.
+		// Otherwise if we are projecting our location to the underlying geometry and it's far above or below the navmesh,
+		// we'll follow that geometry's plane out of range of valid navigation.
+		if (bSameNavLocation && bProjectNavMeshWalking)
+		{
+			UMovementUtils::SetGravityVerticalComponent(TargetFeetLocation, CachedNavLocation.Location.Dot(UpDirection), UpDirection);
+		}
+
+		// Find the point on the NavMesh
+		const bool bHasNavigationData = FindNavFloor(TargetFeetLocation, DestNavLocation);
+		if (!bHasNavigationData)
+		{
+			// Can't find nav mesh at this location, so fall back to normal walking mode 
+			OutputState.MovementEndState.NextModeName = DefaultModeNames::Walking;
+			OutputState.MovementEndState.RemainingMs = Params.TimeStep.StepMs;
+			MoveRecord.SetDeltaSeconds(0.0f);
+
+			CaptureFinalState(LocationInProgress, RotationInProgress.Rotator(), /*bDidAttemptMovement=*/true, CurrentFloor, MoveRecord, OutputSyncState);
+			return;
+		}
+
+		CachedNavLocation = DestNavLocation;
+	}
+
+	if (DestNavLocation.NodeRef != INVALID_NAVNODEREF)
+	{
+		FVector NewFeetLocation = UMovementUtils::ProjectToGravityFloor(TargetFeetLocation, UpDirection) + UMovementUtils::GetGravityVerticalComponent(DestNavLocation.Location, UpDirection);
+		if (bProjectNavMeshWalking)
+		{
+			const float TotalCapsuleHeight = SimpleHalfHeight * 2.0f;
+			const float UpOffset = TotalCapsuleHeight * FMath::Max(0.f, NavMeshProjectionHeightScaleUp);
+			const float DownOffset = TotalCapsuleHeight * FMath::Max(0.f, NavMeshProjectionHeightScaleDown);
+			NewFeetLocation = ProjectLocationFromNavMesh(DeltaSeconds, StartingFeetLocation, NewFeetLocation, UpOffset, DownOffset);
+		}
+
+		FVector AdjustedDelta = NewFeetLocation - StartingFeetLocation;
+
+		if (!AdjustedDelta.IsNearlyZero())
+		{
+			FHitResult MoveHitResult;
+
+			if (UAsyncMovementUtils::TestDepenetratingMove(Params.MovingComps, StartLocation, StartLocation + AdjustedDelta, StartRotation, TargetRotation, bSweepWhileNavWalking, OUT MoveHitResult, IN OUT MoveRecord))
+			{
+				LocationInProgress = StartLocation + (AdjustedDelta * MoveHitResult.Time);
+				RotationInProgress = FQuat::Slerp(StartRotation, TargetRotation, MoveHitResult.Time);
+			}
+		}
+	}
+	else
+	{
+		OutputState.MovementEndState.NextModeName = CommonLegacySettings->AirMovementModeName;
+		OutputState.MovementEndState.RemainingMs = Params.TimeStep.StepMs;
+		MoveRecord.SetDeltaSeconds((Params.TimeStep.StepMs - OutputState.MovementEndState.RemainingMs) * 0.001f);
+	}
+
+	CaptureFinalState(LocationInProgress, RotationInProgress.Rotator(), /*bDidAttemptMovement=*/true, CurrentFloor, MoveRecord, OutputSyncState);
+}
+
+bool UAsyncNavWalkingMode::FindNavFloor(const FVector& TestLocation, FNavLocation& NavFloorLocation) const
+{
+	const INavigationDataInterface* NavData = GetNavData();
+	if (NavData == nullptr || !NavMoverComponent.IsValid())
+	{
+		return false;
+	}
+
+	const FNavAgentProperties& AgentProps = NavMoverComponent->GetNavAgentPropertiesRef();
+	const float SearchRadius = AgentProps.AgentRadius * 2.0f;
+	const float SearchHeight = AgentProps.AgentHeight * AgentProps.NavWalkingSearchHeightScale;
+
+	return NavData->ProjectPoint(TestLocation, NavFloorLocation, FVector(SearchRadius, SearchRadius, SearchHeight));
+}
+
+UObject* UAsyncNavWalkingMode::GetTurnGenerator()
+{
+	return TurnGenerator;
+}
+
+void UAsyncNavWalkingMode::SetTurnGeneratorClass(TSubclassOf<UObject> TurnGeneratorClass)
+{
+	if (TurnGeneratorClass)
+	{
+		TurnGenerator = NewObject<UObject>(this, TurnGeneratorClass);
+	}
+	else
+	{
+		TurnGenerator = nullptr; // Clearing the turn generator is valid - will go back to the default turn generation
+	}
+}
+
+void UAsyncNavWalkingMode::SetCollisionForNavWalking(bool bEnable)
+{
+	if (const UMoverComponent* MoverComponent = GetMoverComponent())
+	{
+		if (UPrimitiveComponent* UpdatedCompAsPrimitive = Cast<UPrimitiveComponent>(MoverComponent->GetUpdatedComponent()))
+		{
+			if (bEnable)
+			{
+				CollideVsWorldStatic = UpdatedCompAsPrimitive->GetCollisionResponseToChannel(ECC_WorldStatic);
+				CollideVsWorldDynamic = UpdatedCompAsPrimitive->GetCollisionResponseToChannel(ECC_WorldDynamic);
+
+				// TODO NS: Eventually might want make the new collision response an option so they trigger certain things while still not colliding 
+				UpdatedCompAsPrimitive->SetCollisionResponseToChannel(ECC_WorldStatic, ECR_Ignore);
+				UpdatedCompAsPrimitive->SetCollisionResponseToChannel(ECC_WorldDynamic, ECR_Ignore);
+			}
+			else
+			{
+				// Grabbing the original shape settings and reverting to our old collision responses
+				if (const UShapeComponent* OriginalShapeComp = UMovementUtils::GetOriginalComponentType<UShapeComponent>(MoverComponent->GetOwner()))
+				{
+					CollideVsWorldStatic = OriginalShapeComp->GetCollisionResponseToChannel(ECC_WorldStatic);
+					CollideVsWorldDynamic = OriginalShapeComp->GetCollisionResponseToChannel(ECC_WorldDynamic);
+				}
+
+				UpdatedCompAsPrimitive->SetCollisionResponseToChannel(ECC_WorldStatic, CollideVsWorldStatic);
+				UpdatedCompAsPrimitive->SetCollisionResponseToChannel(ECC_WorldDynamic, CollideVsWorldDynamic);
+			}
+		}
+	}
+}
+
+void UAsyncNavWalkingMode::OnActivate()
+{
+	// TODO: Need alternative to modifying collision settings
+	//SetCollisionForNavWalking(true);
+
+	if (UMoverBlackboard* SimBlackboard = GetMoverComponent()->GetSimBlackboard_Mutable())
+	{
+		SimBlackboard->Invalidate(NavWalkingBlackBoard::ProjectedNavMeshHitResult);
+	}
+
+	NavMeshProjectionTimer = (NavMeshProjectionInterval > 0.f) ? FMath::FRandRange(-NavMeshProjectionInterval, 0.f) : 0.f;
+}
+
+void UAsyncNavWalkingMode::OnDeactivate()
+{
+	// TODO: Need alternative to modifying collision settings
+}
+
+const INavigationDataInterface* UAsyncNavWalkingMode::GetNavData() const
+{
+	const UWorld* World = GetWorld();
+
+	const AActor* MoverComponentOwner;
+	if (const UMoverComponent* MoverComponent = GetMoverComponent())
+	{
+		MoverComponentOwner = MoverComponent->GetOwner();
+		if (World == nullptr || World->GetNavigationSystem() == nullptr || MoverComponentOwner == nullptr)
+		{
+			return nullptr;
+		}
+	}
+	else
+	{
+		return nullptr;
+	}
+
+
+	const INavigationDataInterface* NavData = FNavigationSystem::GetNavDataForActor(*MoverComponentOwner);
+
+	return NavData;
+}
+
+void UAsyncNavWalkingMode::FindBestNavMeshLocation(const FVector& TraceStart, const FVector& TraceEnd, const FVector& CurrentFeetLocation, const FVector& TargetNavLocation, FHitResult& OutHitResult) const
+{
+	// raycast to underlying mesh to allow us to more closely follow geometry
+	// we use static objects here as a best approximation to accept only objects that
+	// influence navmesh generation
+	FCollisionQueryParams Params(SCENE_QUERY_STAT(ProjectLocation), false);
+	Params.AddIgnoredActor(GetMoverComponent()->GetOwner());
+	
+	// blocked by world static and optionally world dynamic
+	FCollisionResponseParams ResponseParams(ECR_Ignore);
+	ResponseParams.CollisionResponse.SetResponse(ECC_WorldStatic, ECR_Overlap);
+	ResponseParams.CollisionResponse.SetResponse(ECC_WorldDynamic, bProjectNavMeshOnBothWorldChannels ? ECR_Overlap : ECR_Ignore);
+
+	TArray<FHitResult> MultiTraceHits;
+	GetWorld()->LineTraceMultiByChannel(MultiTraceHits, TraceStart, TraceEnd, ECC_WorldStatic, Params, ResponseParams);
+
+	struct FCompareFHitResultNavMeshTrace
+	{
+		explicit FCompareFHitResultNavMeshTrace(const FVector& inSourceLocation) : SourceLocation(inSourceLocation)
+		{
+		}
+
+		FORCEINLINE bool operator()(const FHitResult& A, const FHitResult& B) const
+		{
+			const float ADistSqr = (SourceLocation - A.ImpactPoint).SizeSquared();
+			const float BDistSqr = (SourceLocation - B.ImpactPoint).SizeSquared();
+
+			return (ADistSqr < BDistSqr);
+		}
+
+		const FVector& SourceLocation;
+	};
+
+	struct FRemoveNotBlockingResponseNavMeshTrace
+	{
+		FRemoveNotBlockingResponseNavMeshTrace(bool bInCheckOnlyWorldStatic) : bCheckOnlyWorldStatic(bInCheckOnlyWorldStatic) {}
+
+		FORCEINLINE bool operator()(const FHitResult& TestHit) const
+		{
+			UPrimitiveComponent* PrimComp = TestHit.GetComponent();
+			const bool bBlockOnWorldStatic = PrimComp && (PrimComp->GetCollisionResponseToChannel(ECC_WorldStatic) == ECR_Block);
+			const bool bBlockOnWorldDynamic = PrimComp && (PrimComp->GetCollisionResponseToChannel(ECC_WorldDynamic) == ECR_Block);
+
+			return !bBlockOnWorldStatic && (!bBlockOnWorldDynamic || bCheckOnlyWorldStatic);
+		}
+
+		bool bCheckOnlyWorldStatic;
+	};
+
+	MultiTraceHits.RemoveAllSwap(FRemoveNotBlockingResponseNavMeshTrace(!bProjectNavMeshOnBothWorldChannels), EAllowShrinking::No);
+	if (MultiTraceHits.Num() > 0)
+	{
+		// Sort the hits by the closest to our origin.
+		MultiTraceHits.Sort(FCompareFHitResultNavMeshTrace(TargetNavLocation));
+
+		// Cache the closest hit and treat it as a blocking hit (we used an overlap to get all the world static hits so we could sort them ourselves)
+		OutHitResult = MultiTraceHits[0];
+		OutHitResult.bBlockingHit = true;
+	}
+}
+
+FVector UAsyncNavWalkingMode::ProjectLocationFromNavMesh(float DeltaSeconds, const FVector& CurrentFeetLocation, const FVector& TargetNavLocation, float UpOffset, float DownOffset)
+{
+	FVector NewLocation = TargetNavLocation;
+
+	const float VerticalOffset = -(DownOffset + UpOffset);
+	if (VerticalOffset > -UE_SMALL_NUMBER)
+	{
+		return NewLocation;
+	}
+
+	const UMoverComponent* MoverComp = GetMoverComponent();
+	const FVector UpDirection = MoverComp->GetUpDirection();
+
+	const FVector TraceStart = TargetNavLocation + UpOffset * UpDirection;
+	const FVector TraceEnd = TargetNavLocation + DownOffset * -UpDirection;
+
+	FHitResult CachedProjectedNavMeshHitResult;
+	UMoverBlackboard* SimBlackboard = MoverComp->GetSimBlackboard_Mutable();
+	SimBlackboard->TryGet(NavWalkingBlackBoard::ProjectedNavMeshHitResult, CachedProjectedNavMeshHitResult);
+
+	// We can skip this trace if we are checking at the same location as the last trace (ie, we haven't moved).
+	const bool bCachedLocationStillValid = (CachedProjectedNavMeshHitResult.bBlockingHit &&
+		CachedProjectedNavMeshHitResult.TraceStart == TraceStart &&
+		CachedProjectedNavMeshHitResult.TraceEnd == TraceEnd);
+
+	NavMeshProjectionTimer -= DeltaSeconds;
+	if (NavMeshProjectionTimer <= 0.0f)
+	{
+		if (!bCachedLocationStillValid)
+		{
+			UE_LOG(LogMover, VeryVerbose, TEXT("ProjectLocationFromNavMesh(): %s interval: %.3f [SKIP TRACE]"), *GetNameSafe(GetMoverComponent()->GetOwner()), NavMeshProjectionInterval);
+
+			FHitResult HitResult;
+			FindBestNavMeshLocation(TraceStart, TraceEnd, CurrentFeetLocation, TargetNavLocation, HitResult);
+
+			// discard result if we were already inside something
+			if (HitResult.bStartPenetrating || !HitResult.bBlockingHit)
+			{
+				CachedProjectedNavMeshHitResult.Reset();
+				SimBlackboard->Invalidate(NavWalkingBlackBoard::ProjectedNavMeshHitResult);
+			}
+			else
+			{
+				CachedProjectedNavMeshHitResult = HitResult;
+				SimBlackboard->Set(NavWalkingBlackBoard::ProjectedNavMeshHitResult, CachedProjectedNavMeshHitResult);
+			}
+		}
+		else
+		{
+			UE_LOG(LogMover, VeryVerbose, TEXT("ProjectLocationFromNavMesh(): %s interval: %.3f [SKIP TRACE]"), *GetNameSafe(GetMoverComponent()->GetOwner()), NavMeshProjectionInterval);
+		}
+
+		// Wrap around to maintain same relative offset to tick time changes.
+		// Prevents large framerate spikes from aligning multiple characters to the same frame (if they start staggered, they will now remain staggered).
+		float ModTime = 0.f;
+		if (NavMeshProjectionInterval > UE_SMALL_NUMBER)
+		{
+			ModTime = FMath::Fmod(-NavMeshProjectionTimer, NavMeshProjectionInterval);
+		}
+
+		NavMeshProjectionTimer = NavMeshProjectionInterval - ModTime;
+	}
+
+	// Project to last plane we found.
+	if (CachedProjectedNavMeshHitResult.bBlockingHit)
+	{
+		if (bCachedLocationStillValid && FMath::IsNearlyEqual(CurrentFeetLocation.Dot(UpDirection), CachedProjectedNavMeshHitResult.ImpactPoint.Dot(UpDirection), (FVector::FReal)0.01f))
+		{
+			// Already at destination.
+			UMovementUtils::SetGravityVerticalComponent(NewLocation, CurrentFeetLocation.Dot(UpDirection), UpDirection);
+		}
+		else
+		{
+			const FVector ProjectedPoint = FMath::LinePlaneIntersection(TraceStart, TraceEnd, CachedProjectedNavMeshHitResult.ImpactPoint, CachedProjectedNavMeshHitResult.Normal);
+			FVector::FReal ProjectedVertical = ProjectedPoint.Dot(UpDirection);
+
+			// Limit to not be too far above or below NavMesh location
+			const FVector::FReal VertTraceStart = TraceStart.Dot(UpDirection);
+			const FVector::FReal VertTraceEnd = TraceEnd.Dot(UpDirection);
+			const FVector::FReal TraceMin = FMath::Min(VertTraceStart, VertTraceEnd);
+			const FVector::FReal TraceMax = FMath::Max(VertTraceStart, VertTraceEnd);
+			ProjectedVertical = FMath::Clamp(ProjectedVertical, TraceMin, TraceMax);
+
+			// Interp for smoother updates (less "pop" when trace hits something new). 0 interp speed is instant.
+			const FVector::FReal InterpSpeed = FMath::Max<FVector::FReal>(0.f, NavMeshProjectionInterpSpeed);
+			ProjectedVertical = FMath::FInterpTo(CurrentFeetLocation.Dot(UpDirection), ProjectedVertical, (FVector::FReal)DeltaSeconds, InterpSpeed);
+			ProjectedVertical = FMath::Clamp(ProjectedVertical, TraceMin, TraceMax);
+
+			// Final result
+			UMovementUtils::SetGravityVerticalComponent(NewLocation, ProjectedVertical, UpDirection);
+		}
+	}
+
+	return NewLocation;
+}
+
+void UAsyncNavWalkingMode::OnRegistered(const FName ModeName)
+{
+	Super::OnRegistered(ModeName);
+
+	UMoverComponent* MoverComponent = GetMoverComponent();
+	CommonLegacySettings = MoverComponent->FindSharedSettings<UCommonLegacyMovementSettings>();
+	ensureMsgf(CommonLegacySettings.IsValid(), TEXT("Failed to find instance of CommonLegacyMovementSettings on %s. Movement may not function properly."), *GetPathNameSafe(this));
+
+	if (const AActor* MoverCompOwner = MoverComponent->GetOwner())
+	{
+		NavMoverComponent = MoverCompOwner->FindComponentByClass<UNavMoverComponent>();
+	}
+
+	if (!NavMoverComponent.IsValid())
+	{
+		UE_LOG(LogMover, Warning, TEXT("NavWalkingMode on %s could not find a valid NavMoverComponent and will not function properly."), *GetNameSafe(GetMoverComponent()->GetOwner()));
+	}
+}
+
+void UAsyncNavWalkingMode::OnUnregistered()
+{
+	CommonLegacySettings = nullptr;
+	NavMoverComponent = nullptr;
+
+	Super::OnUnregistered();
+}
+
+void UAsyncNavWalkingMode::CaptureFinalState(const FVector FinalLocation, const FRotator FinalRotation, bool bDidAttemptMovement, const FFloorCheckResult& FloorResult, const FMovementRecord& Record, FMoverDefaultSyncState& OutputSyncState) const
+{
+	FRelativeBaseInfo PriorBaseInfo;
+
+	UMoverBlackboard* SimBlackboard = GetMoverComponent()->GetSimBlackboard_Mutable();
+
+	const bool bHasPriorBaseInfo = SimBlackboard->TryGet(CommonBlackboard::LastFoundDynamicMovementBase, PriorBaseInfo);
+
+	FRelativeBaseInfo CurrentBaseInfo = UpdateFloorAndBaseInfo(FloorResult);
+
+	// If we're on a dynamic base and we're not trying to move, keep using the same relative actor location. This prevents slow relative 
+	//  drifting that can occur from repeated floor sampling as the base moves through the world.
+	if (CurrentBaseInfo.HasRelativeInfo()
+		&& bHasPriorBaseInfo && !bDidAttemptMovement
+		&& PriorBaseInfo.UsesSameBase(CurrentBaseInfo))
+	{
+		CurrentBaseInfo.ContactLocalPosition = PriorBaseInfo.ContactLocalPosition;
+	}
+
+	if (CurrentBaseInfo.HasRelativeInfo())
+	{
+		SimBlackboard->Set(CommonBlackboard::LastFoundDynamicMovementBase, CurrentBaseInfo);
+
+		OutputSyncState.SetTransforms_WorldSpace(FinalLocation,
+			FinalRotation,
+			Record.GetRelevantVelocity(),
+			CurrentBaseInfo.MovementBase.Get(), CurrentBaseInfo.BoneName);
+	}
+	else
+	{
+		SimBlackboard->Invalidate(CommonBlackboard::LastFoundDynamicMovementBase);
+
+		OutputSyncState.SetTransforms_WorldSpace(FinalLocation,
+			FinalRotation,
+			Record.GetRelevantVelocity(),
+			nullptr);	// no movement base
+	}
+}
+
+FRelativeBaseInfo UAsyncNavWalkingMode::UpdateFloorAndBaseInfo(const FFloorCheckResult& FloorResult) const
+{
+	FRelativeBaseInfo ReturnBaseInfo;
+
+	UMoverBlackboard* SimBlackboard = GetMoverComponent()->GetSimBlackboard_Mutable();
+
+	SimBlackboard->Set(CommonBlackboard::LastFloorResult, FloorResult);
+
+	if (FloorResult.IsWalkableFloor() && UBasedMovementUtils::IsADynamicBase(FloorResult.HitResult.GetComponent()))
+	{
+		ReturnBaseInfo.SetFromFloorResult(FloorResult);
+	}
+
+	return ReturnBaseInfo;
+}
+
+
diff --git a/Engine/Plugins/Experimental/Mover/Source/Mover/Private/DefaultMovementSet/Modes/AsyncWalkingMode.cpp b/Engine/Plugins/Experimental/Mover/Source/Mover/Private/DefaultMovementSet/Modes/AsyncWalkingMode.cpp
new file mode 100644
--- /dev/null
+++ b/Engine/Plugins/Experimental/Mover/Source/Mover/Private/DefaultMovementSet/Modes/AsyncWalkingMode.cpp
@@ -0,0 +1,422 @@
+// Copyright Epic Games, Inc. All Rights Reserved.
+
+#include "DefaultMovementSet/Modes/AsyncWalkingMode.h"
+#include "Components/SkeletalMeshComponent.h"
+#include "GameFramework/Pawn.h"
+#include "MoveLibrary/MovementUtils.h"
+#include "MoveLibrary/AsyncMovementUtils.h"
+#include "MoveLibrary/BasedMovementUtils.h"
+#include "MoveLibrary/ModularMovement.h"
+#include "MoveLibrary/FloorQueryUtils.h"
+#include "MoveLibrary/GroundMovementUtils.h"
+#include "MoverComponent.h"
+#include "DefaultMovementSet/Settings/CommonLegacyMovementSettings.h"
+#include "MoverLog.h"
+#include "DefaultMovementSet/InstantMovementEffects/BasicInstantMovementEffects.h"
+#include "DrawDebugHelpers.h"
+
+#include UE_INLINE_GENERATED_CPP_BY_NAME(AsyncWalkingMode)
+
+
+UAsyncWalkingMode::UAsyncWalkingMode(const FObjectInitializer& ObjectInitializer)
+	: Super(ObjectInitializer)
+{
+	SharedSettingsClasses.Add(UCommonLegacyMovementSettings::StaticClass());
+
+	GameplayTags.AddTag(Mover_IsOnGround);
+}
+
+void UAsyncWalkingMode::OnGenerateMove(const FMoverTickStartData& StartState, const FMoverTimeStep& TimeStep, FProposedMove& OutProposedMove) const
+{
+	const UMoverComponent* MoverComp = GetMoverComponent();
+	const FCharacterDefaultInputs* CharacterInputs = StartState.InputCmd.InputCollection.FindDataByType<FCharacterDefaultInputs>();
+	const FMoverDefaultSyncState* StartingSyncState = StartState.SyncState.SyncStateCollection.FindDataByType<FMoverDefaultSyncState>();
+	check(StartingSyncState);
+
+	if (!CommonLegacySettings.IsValid())
+	{
+		return;
+	}
+
+	const float DeltaSeconds = TimeStep.StepMs * 0.001f;
+	FFloorCheckResult LastFloorResult;
+	FVector MovementNormal;
+
+	UMoverBlackboard* SimBlackboard = MoverComp->GetSimBlackboard_Mutable();
+	FVector UpDirection = MoverComp->GetUpDirection();
+
+	// Try to use the floor as the basis for the intended move direction (i.e. try to walk along slopes, rather than into them)
+	if (SimBlackboard && SimBlackboard->TryGet(CommonBlackboard::LastFloorResult, LastFloorResult) && LastFloorResult.IsWalkableFloor())
+	{
+		MovementNormal = LastFloorResult.HitResult.ImpactNormal;
+	}
+	else
+	{
+		MovementNormal = UpDirection;
+	}
+
+	FRotator IntendedOrientation_WorldSpace;
+	// If there's no intent from input to change orientation, use the current orientation
+	if (!CharacterInputs || CharacterInputs->OrientationIntent.IsNearlyZero())
+	{
+		IntendedOrientation_WorldSpace = StartingSyncState->GetOrientation_WorldSpace();
+	}
+	else
+	{
+		IntendedOrientation_WorldSpace = CharacterInputs->GetOrientationIntentDir_WorldSpace().ToOrientationRotator();
+	}
+
+	IntendedOrientation_WorldSpace = UMovementUtils::ApplyGravityToOrientationIntent(IntendedOrientation_WorldSpace, MoverComp->GetWorldToGravityTransform(), CommonLegacySettings->bShouldRemainVertical);
+	
+	FGroundMoveParams Params;
+
+	if (CharacterInputs)
+	{
+		Params.MoveInputType = CharacterInputs->GetMoveInputType();
+
+		const bool bMaintainInputMagnitude = true;
+		Params.MoveInput = UPlanarConstraintUtils::ConstrainDirectionToPlane(MoverComp->GetPlanarConstraint(), CharacterInputs->GetMoveInput_WorldSpace(), bMaintainInputMagnitude);
+	}
+	else
+	{
+		Params.MoveInputType = EMoveInputType::Invalid;
+		Params.MoveInput = FVector::ZeroVector;
+	}
+
+	Params.OrientationIntent = IntendedOrientation_WorldSpace;
+	Params.PriorVelocity = FVector::VectorPlaneProject(StartingSyncState->GetVelocity_WorldSpace(), MovementNormal);
+	Params.PriorOrientation = StartingSyncState->GetOrientation_WorldSpace();
+	Params.GroundNormal = MovementNormal;
+	Params.TurningRate = CommonLegacySettings->TurningRate;
+	Params.TurningBoost = CommonLegacySettings->TurningBoost;
+	Params.MaxSpeed = CommonLegacySettings->MaxSpeed;
+	Params.Acceleration = CommonLegacySettings->Acceleration;
+	Params.Deceleration = CommonLegacySettings->Deceleration;
+	Params.DeltaSeconds = DeltaSeconds;
+	Params.WorldToGravityQuat = MoverComp->GetWorldToGravityTransform();
+	Params.UpDirection = UpDirection;
+	Params.bUseAccelerationForVelocityMove = CommonLegacySettings->bUseAccelerationForVelocityMove;
+	
+	if (Params.MoveInput.SizeSquared() > 0.f && !UMovementUtils::IsExceedingMaxSpeed(Params.PriorVelocity, CommonLegacySettings->MaxSpeed))
+	{
+		Params.Friction = CommonLegacySettings->GroundFriction;
+	}
+	else
+	{
+		Params.Friction = CommonLegacySettings->bUseSeparateBrakingFriction ? CommonLegacySettings->BrakingFriction : CommonLegacySettings->GroundFriction;
+		Params.Friction *= CommonLegacySettings->BrakingFrictionFactor;
+	}
+
+	OutProposedMove = UGroundMovementUtils::ComputeControlledGroundMove(Params);
+
+	if (TurnGenerator)
+	{
+		OutProposedMove.AngularVelocity = ITurnGeneratorInterface::Execute_GetTurn(TurnGenerator, IntendedOrientation_WorldSpace, StartState, *StartingSyncState, TimeStep, OutProposedMove, SimBlackboard);
+	}
+}
+
+void UAsyncWalkingMode::OnSimulationTick(const FSimulationTickParams& Params, FMoverTickEndData& OutputState)
+{
+	if (!CommonLegacySettings.IsValid())
+	{
+		return;
+	}
+
+	UMoverComponent* MoverComp = GetMoverComponent();
+	const FMoverTickStartData& StartState = Params.StartState;
+	FProposedMove ProposedMove = Params.ProposedMove;
+
+	const FCharacterDefaultInputs* CharacterInputs = StartState.InputCmd.InputCollection.FindDataByType<FCharacterDefaultInputs>();
+	const FMoverDefaultSyncState* StartingSyncState = StartState.SyncState.SyncStateCollection.FindDataByType<FMoverDefaultSyncState>();
+	check(StartingSyncState);
+
+	FMoverDefaultSyncState& OutputSyncState = OutputState.SyncState.SyncStateCollection.FindOrAddMutableDataByType<FMoverDefaultSyncState>();
+
+
+	const float DeltaSeconds = Params.TimeStep.StepMs * 0.001f;
+
+	const FVector OrigMoveDelta = ProposedMove.LinearVelocity * DeltaSeconds;
+
+	const FVector StartLocation = StartingSyncState->GetLocation_WorldSpace();
+	const FVector TargetLocation = StartLocation + OrigMoveDelta;
+
+
+	FMovementRecord MoveRecord;
+	MoveRecord.SetDeltaSeconds(DeltaSeconds);
+
+	FFloorCheckResult CurrentFloor;
+	UMoverBlackboard* SimBlackboard = MoverComp->GetSimBlackboard_Mutable();
+
+	FVector UpDirection = MoverComp->GetUpDirection();
+	
+	// If we don't have cached floor information, we need to search for it again
+	if (!SimBlackboard->TryGet(CommonBlackboard::LastFloorResult, CurrentFloor))
+	{
+		UFloorQueryUtils::FindFloor(Params.MovingComps, CommonLegacySettings->FloorSweepDistance, CommonLegacySettings->MaxWalkSlopeCosine, StartLocation, CurrentFloor);
+	}
+
+	OutputSyncState.MoveDirectionIntent = (ProposedMove.bHasDirIntent ? ProposedMove.DirectionIntent : FVector::ZeroVector);
+
+	const FRotator StartingOrient = StartingSyncState->GetOrientation_WorldSpace();
+	FRotator TargetOrient = StartingOrient;
+
+	bool bIsOrientationChanging = false;
+
+	// Apply orientation changes (if any)
+	if (!UMovementUtils::IsAngularVelocityZero(ProposedMove.AngularVelocity))
+	{ 
+		TargetOrient += (ProposedMove.AngularVelocity * DeltaSeconds);
+		bIsOrientationChanging = (TargetOrient != StartingOrient);
+	}
+
+	const FQuat StartRotation = StartingOrient.Quaternion();
+	const FQuat TargetRotation = TargetOrient.Quaternion();
+
+	FVector LocationInProgress = StartLocation;
+	FQuat   RotationInProgress = StartRotation;
+
+	FHitResult MoveHitResult(1.f);
+	
+	FVector CurMoveDelta = OrigMoveDelta;
+
+	bool bDidAttemptMovement = false;
+
+	float PercentTimeAppliedSoFar = MoveHitResult.Time;
+	bool bWasFirstMoveBlocked = false;
+
+	if (!CurMoveDelta.IsNearlyZero() || bIsOrientationChanging)
+	{
+		// Attempt to move the full amount first
+		bDidAttemptMovement = true;
+
+		const bool bWouldMove = UAsyncMovementUtils::TestDepenetratingMove(Params.MovingComps, StartLocation, TargetLocation, StartRotation, TargetRotation, /* bShouldSweep */ true, OUT MoveHitResult, MoveRecord);
+
+		float LastMoveSeconds = DeltaSeconds;
+
+		LocationInProgress = StartLocation + ((TargetLocation - StartLocation) * MoveHitResult.Time);
+		RotationInProgress = FQuat::Slerp(StartRotation, TargetRotation, MoveHitResult.Time);
+
+		if (MoveHitResult.bStartPenetrating)
+		{
+			// We started by being stuck in geometry and need to resolve it first
+			// TODO: try to resolve starting stuck in geometry
+		}
+		else if (MoveHitResult.IsValidBlockingHit())
+		{
+			// We impacted something (possibly a ramp, possibly a barrier)
+			PercentTimeAppliedSoFar = MoveHitResult.Time;		
+
+			// Check if the blockage is a walkable ramp rising in front of us
+			if ((MoveHitResult.Time > 0.f) && (MoveHitResult.Normal.Dot(UpDirection) > UE_KINDA_SMALL_NUMBER) && 
+			    UFloorQueryUtils::IsHitSurfaceWalkable(MoveHitResult, UpDirection, CommonLegacySettings->MaxWalkSlopeCosine))
+			{
+				// It's a walkable ramp, so cut up the move and attempt to move the remainder of it along the ramp's surface, possibly generating another hit
+				const float PercentTimeRemaining = 1.f - PercentTimeAppliedSoFar;
+				CurMoveDelta = UGroundMovementUtils::ComputeDeflectedMoveOntoRamp(CurMoveDelta * PercentTimeRemaining, UpDirection, MoveHitResult, CommonLegacySettings->MaxWalkSlopeCosine, CurrentFloor.bLineTrace);
+
+				UAsyncMovementUtils::TestDepenetratingMove(Params.MovingComps, LocationInProgress, LocationInProgress+CurMoveDelta, RotationInProgress, TargetRotation, /*bShouldSweep=*/ true, OUT MoveHitResult, MoveRecord);
+
+				LastMoveSeconds = PercentTimeRemaining * LastMoveSeconds;
+
+				LocationInProgress = LocationInProgress + ((MoveHitResult.TraceEnd - MoveHitResult.TraceStart) * MoveHitResult.Time);
+				RotationInProgress = FQuat::Slerp(RotationInProgress, TargetRotation, MoveHitResult.Time);
+
+				const float SecondHitPercent = MoveHitResult.Time * PercentTimeRemaining;
+				PercentTimeAppliedSoFar = FMath::Clamp(PercentTimeAppliedSoFar + SecondHitPercent, 0.f, 1.f);
+			}
+
+			if (MoveHitResult.IsValidBlockingHit())
+			{
+	
+				// If still blocked, try to step up onto the blocking object OR slide along it
+				// TODO: Take movement bases into account
+				if (UGroundMovementUtils::CanStepUpOnHitSurface(MoveHitResult)) // || (CharacterOwner->GetMovementBase() != nullptr && Hit.HitObjectHandle == CharacterOwner->GetMovementBase()->GetOwner()))
+				{
+					// hit a barrier or unwalkable surface, try to step up and onto it
+					const FVector PreStepUpLocation = LocationInProgress;
+					const FVector DownwardDir = -UpDirection;
+
+					FOptionalFloorCheckResult StepUpFloorResult;	// passed to sub-operations, so we can use their final floor results if they did a test
+					FVector PostStepUpLocation; // Valid if step-up succeeded
+
+					if (UGroundMovementUtils::TestMoveToStepOver(Params.MovingComps, DownwardDir, CommonLegacySettings->MaxStepHeight, CommonLegacySettings->MaxWalkSlopeCosine, CommonLegacySettings->FloorSweepDistance, OrigMoveDelta * (1.f - PercentTimeAppliedSoFar), RotationInProgress, MoveHitResult, CurrentFloor, false, &StepUpFloorResult, PostStepUpLocation, MoveRecord))
+					{
+						LocationInProgress = PostStepUpLocation;
+					}
+					else
+					{
+						FMoverOnImpactParams ImpactParams(DefaultModeNames::Walking, MoveHitResult, OrigMoveDelta);
+						MoverComp->HandleImpact(ImpactParams);
+						float PercentAvailableToSlide = 1.f - PercentTimeAppliedSoFar;
+
+						float SlideAmount = UGroundMovementUtils::TestGroundedMoveAlongHitSurface(Params.MovingComps, OrigMoveDelta, LocationInProgress, TargetRotation, /*bHandleImpact=*/true, CommonLegacySettings->MaxStepHeight, CommonLegacySettings->MaxWalkSlopeCosine, IN OUT MoveHitResult, IN OUT MoveRecord);
+						PercentTimeAppliedSoFar += PercentAvailableToSlide * SlideAmount;
+					}
+				}
+				else if (MoveHitResult.Component.IsValid() && !MoveHitResult.Component.Get()->CanCharacterStepUp(Cast<APawn>(MoveHitResult.GetActor())))
+				{
+					FMoverOnImpactParams ImpactParams(DefaultModeNames::Walking, MoveHitResult, OrigMoveDelta);
+					MoverComp->HandleImpact(ImpactParams);
+					float PercentAvailableToSlide = 1.f - PercentTimeAppliedSoFar;
+					
+					float SlideAmount = UGroundMovementUtils::TestGroundedMoveAlongHitSurface(Params.MovingComps, OrigMoveDelta, LocationInProgress, TargetRotation, /*bHandleImpact=*/true, CommonLegacySettings->MaxStepHeight, CommonLegacySettings->MaxWalkSlopeCosine, IN OUT MoveHitResult, IN OUT MoveRecord);
+
+					LocationInProgress = LocationInProgress + ((MoveHitResult.TraceEnd - MoveHitResult.TraceStart) * SlideAmount);
+					RotationInProgress = FQuat::Slerp(RotationInProgress, TargetRotation, SlideAmount);
+
+
+					PercentTimeAppliedSoFar += PercentAvailableToSlide * SlideAmount;
+				}
+			}
+		}
+
+		// Search for the floor we've ended up on
+		UFloorQueryUtils::FindFloor(Params.MovingComps, CommonLegacySettings->FloorSweepDistance, CommonLegacySettings->MaxWalkSlopeCosine, LocationInProgress, CurrentFloor);
+
+		if (CurrentFloor.IsWalkableFloor())
+		{
+			LocationInProgress = UGroundMovementUtils::TestMoveToAdjustToFloor(Params.MovingComps, LocationInProgress, RotationInProgress, CommonLegacySettings->MaxWalkSlopeCosine, IN OUT CurrentFloor, MoveRecord);
+		}
+    
+		if (!CurrentFloor.IsWalkableFloor() && !CurrentFloor.HitResult.bStartPenetrating)
+		{
+			// No floor or not walkable, so let's let the airborne movement mode deal with it
+			OutputState.MovementEndState.NextModeName = CommonLegacySettings->AirMovementModeName;
+			OutputState.MovementEndState.RemainingMs = Params.TimeStep.StepMs - (Params.TimeStep.StepMs * PercentTimeAppliedSoFar);
+			MoveRecord.SetDeltaSeconds((Params.TimeStep.StepMs - OutputState.MovementEndState.RemainingMs) * 0.001f);
+			CaptureFinalState(LocationInProgress, RotationInProgress.Rotator(), bDidAttemptMovement, CurrentFloor, MoveRecord, OutputSyncState);
+			return;
+		}
+	}
+	else
+	{
+		/* TODO: Support option to perform floor checks even when not moving
+		// If the actor isn't moving we still may need to check if they have a valid floor, such as if they're on an elevator platform moving up/down
+		if ((FloorCheckPolicy == EStaticFloorCheckPolicy::Always) || 
+		    (FloorCheckPolicy == EStaticFloorCheckPolicy::OnDynamicBaseOnly && StartingSyncState->GetMovementBase()))
+		{
+			UFloorQueryUtils::FindFloor(Params.MovingComps, CommonLegacySettings->FloorSweepDistance, CommonLegacySettings->MaxWalkSlopeCosine, LocationInProgress, CurrentFloor);
+		
+			FHitResult Hit(CurrentFloor.HitResult);
+			if (Hit.bStartPenetrating)
+			{
+				// The floor check failed because it started in penetration
+				// We do not want to try to move downward because the downward sweep failed, rather we'd like to try to pop out of the floor.
+				Hit.TraceEnd = Hit.TraceStart + UpDirection * 2.4;
+				FVector RequestedAdjustment = UMovementUtils::ComputePenetrationAdjustment(Hit);
+			
+				const EMoveComponentFlags IncludeBlockingOverlapsWithoutEvents = (MOVECOMP_NeverIgnoreBlockingOverlaps | MOVECOMP_DisableBlockingOverlapDispatch);
+				EMoveComponentFlags MoveComponentFlags = MOVECOMP_NoFlags;
+				MoveComponentFlags = (MoveComponentFlags | IncludeBlockingOverlapsWithoutEvents);
+				UMovementUtils::TryMoveToResolvePenetration(Params.MovingComps, MoveComponentFlags, RequestedAdjustment, Hit, RotationInProgress, MoveRecord);
+
+				//TODO: Update Location/RotationInProgress
+			}
+		
+			if (!CurrentFloor.IsWalkableFloor() && !Hit.bStartPenetrating)
+			{
+				// No floor or not walkable, so let's let the airborne movement mode deal with it
+				OutputState.MovementEndState.NextModeName = CommonLegacySettings->AirMovementModeName;
+				OutputState.MovementEndState.RemainingMs = Params.TimeStep.StepMs;
+				MoveRecord.SetDeltaSeconds((Params.TimeStep.StepMs - OutputState.MovementEndState.RemainingMs) * 0.001f);
+				CaptureFinalState(LocationInProgress, RotationInProgress.Rotator(), bDidAttemptMovement, CurrentFloor, MoveRecord, OutputSyncState);
+				return;
+			}
+		}
+		*/
+	}
+
+	CaptureFinalState(LocationInProgress, RotationInProgress.Rotator(), bDidAttemptMovement, CurrentFloor, MoveRecord, OutputSyncState);
+
+}
+
+UObject* UAsyncWalkingMode::GetTurnGenerator()
+{
+	return TurnGenerator;
+}
+
+void UAsyncWalkingMode::SetTurnGeneratorClass(TSubclassOf<UObject> TurnGeneratorClass)
+{
+	if (TurnGeneratorClass)
+	{
+		TurnGenerator = NewObject<UObject>(this, TurnGeneratorClass);
+	}
+	else
+	{
+		TurnGenerator = nullptr; // Clearing the turn generator is valid - will go back to the default turn generation
+	}
+}
+
+
+void UAsyncWalkingMode::OnRegistered(const FName ModeName)
+{
+	Super::OnRegistered(ModeName);
+
+	CommonLegacySettings = GetMoverComponent()->FindSharedSettings<UCommonLegacyMovementSettings>();
+	ensureMsgf(CommonLegacySettings.IsValid(), TEXT("Failed to find instance of CommonLegacyMovementSettings on %s. Movement may not function properly."), *GetPathNameSafe(this));
+}
+
+void UAsyncWalkingMode::OnUnregistered()
+{
+	CommonLegacySettings = nullptr;
+
+	Super::OnUnregistered();
+}
+
+void UAsyncWalkingMode::CaptureFinalState(const FVector FinalLocation, const FRotator FinalRotation, bool bDidAttemptMovement, const FFloorCheckResult& FloorResult, const FMovementRecord& Record, FMoverDefaultSyncState& OutputSyncState) const
+{
+	FRelativeBaseInfo PriorBaseInfo;
+
+	const UMoverComponent* MoverComp = GetMoverComponent();
+	UMoverBlackboard* SimBlackboard = MoverComp->GetSimBlackboard_Mutable();
+
+	const bool bHasPriorBaseInfo = SimBlackboard->TryGet(CommonBlackboard::LastFoundDynamicMovementBase, PriorBaseInfo);
+
+	FRelativeBaseInfo CurrentBaseInfo = UpdateFloorAndBaseInfo(FloorResult);
+
+	// If we're on a dynamic base and we're not trying to move, keep using the same relative actor location. This prevents slow relative 
+	//  drifting that can occur from repeated floor sampling as the base moves through the world.
+	if (CurrentBaseInfo.HasRelativeInfo() 
+		&& bHasPriorBaseInfo && !bDidAttemptMovement 
+		&& PriorBaseInfo.UsesSameBase(CurrentBaseInfo))
+	{
+		CurrentBaseInfo.ContactLocalPosition = PriorBaseInfo.ContactLocalPosition;
+	}
+	
+	if (CurrentBaseInfo.HasRelativeInfo())
+	{
+		SimBlackboard->Set(CommonBlackboard::LastFoundDynamicMovementBase, CurrentBaseInfo);
+
+		OutputSyncState.SetTransforms_WorldSpace( FinalLocation,
+                                                  FinalRotation,
+                                                  Record.GetRelevantVelocity(),
+                                                  CurrentBaseInfo.MovementBase.Get(), CurrentBaseInfo.BoneName);
+	}
+	else
+	{
+		SimBlackboard->Invalidate(CommonBlackboard::LastFoundDynamicMovementBase);
+
+		OutputSyncState.SetTransforms_WorldSpace( FinalLocation,
+                                                  FinalRotation,
+                                                  Record.GetRelevantVelocity(),
+                                                  nullptr);  // no movement base
+	}
+}
+
+
+FRelativeBaseInfo UAsyncWalkingMode::UpdateFloorAndBaseInfo(const FFloorCheckResult& FloorResult) const
+{
+	FRelativeBaseInfo ReturnBaseInfo;
+
+	const UMoverComponent* MoverComp = GetMoverComponent();
+	UMoverBlackboard* SimBlackboard = MoverComp->GetSimBlackboard_Mutable();
+
+	SimBlackboard->Set(CommonBlackboard::LastFloorResult, FloorResult);
+
+	if (FloorResult.IsWalkableFloor() && UBasedMovementUtils::IsADynamicBase(FloorResult.HitResult.GetComponent()))
+	{
+		ReturnBaseInfo.SetFromFloorResult(FloorResult);
+	}
+
+	return ReturnBaseInfo;
+}
diff --git a/Engine/Plugins/Experimental/Mover/Source/Mover/Private/DefaultMovementSet/NavMoverComponent.cpp b/Engine/Plugins/Experimental/Mover/Source/Mover/Private/DefaultMovementSet/NavMoverComponent.cpp
--- a/Engine/Plugins/Experimental/Mover/Source/Mover/Private/DefaultMovementSet/NavMoverComponent.cpp
+++ b/Engine/Plugins/Experimental/Mover/Source/Mover/Private/DefaultMovementSet/NavMoverComponent.cpp
@@ -108,6 +108,12 @@ FVector UNavMoverComponent::GetFeetLocation() const
 	return UpdatedComponent ? (UpdatedComponent->GetComponentLocation() - FVector(0,0,UpdatedComponent->Bounds.BoxExtent.Z)) : FNavigationSystem::InvalidLocation;
 }
 
+FVector UNavMoverComponent::GetFeetLocationAt(FVector ComponentLocation) const
+{
+	const USceneComponent* UpdatedComponent = Cast<USceneComponent>(GetUpdatedObject());
+	return UpdatedComponent ? (ComponentLocation - FVector(0, 0, UpdatedComponent->Bounds.BoxExtent.Z)) : FNavigationSystem::InvalidLocation;
+}
+
 FBasedPosition UNavMoverComponent::GetFeetLocationBased() const
 {
 	FBasedPosition BasedPosition(NULL, GetFeetLocation());
diff --git a/Engine/Plugins/Experimental/Mover/Source/Mover/Private/MoveLibrary/AirMovementUtils.cpp b/Engine/Plugins/Experimental/Mover/Source/Mover/Private/MoveLibrary/AirMovementUtils.cpp
--- a/Engine/Plugins/Experimental/Mover/Source/Mover/Private/MoveLibrary/AirMovementUtils.cpp
+++ b/Engine/Plugins/Experimental/Mover/Source/Mover/Private/MoveLibrary/AirMovementUtils.cpp
@@ -3,6 +3,7 @@
 #include "MoveLibrary/AirMovementUtils.h"
 
 #include "MoverComponent.h"
+#include "MoveLibrary/AsyncMovementUtils.h"
 #include "MoveLibrary/FloorQueryUtils.h"
 #include "MoveLibrary/MovementUtils.h"
 
@@ -131,3 +132,81 @@ float UAirMovementUtils::TryMoveToFallAlongSurface(const FMovingComponentSet& Mo
 
 	return 0.f;
 }
+
+
+/* static */
+float UAirMovementUtils::TestFallingMoveAlongHitSurface(const FMovingComponentSet& MovingComps, const FVector& OriginalMoveDelta, const FVector& LocationAtHit, const FQuat& TargetRotation, bool bHandleImpact, float FloorSweepDistance, float MaxWalkSlopeCosine, FHitResult& InOutHit, FFloorCheckResult& OutFloorResult, FMovementRecord& InOutMoveRecord)
+{
+	OutFloorResult.Clear();
+
+	if (!InOutHit.bBlockingHit)
+	{
+		return 0.f;
+	}
+
+	float PctOfTimeUsed = 0.f;
+	float PctOfOrigDeltaToSlide = 1.f - InOutHit.Time;
+	const FVector OrigHitNormal = InOutHit.Normal;
+
+	FVector SlideDelta = UMovementUtils::ComputeSlideDelta(MovingComps, OriginalMoveDelta, PctOfOrigDeltaToSlide, OrigHitNormal, InOutHit);
+
+	if ((SlideDelta | OriginalMoveDelta) > 0.f)
+	{
+		// First sliding attempt along surface
+		UAsyncMovementUtils::TestDepenetratingMove(MovingComps, LocationAtHit, LocationAtHit + SlideDelta, TargetRotation, TargetRotation, /*bShouldSweep*/ true, InOutHit, InOutMoveRecord);
+
+		PctOfTimeUsed = InOutHit.Time;
+
+		if (InOutHit.IsValidBlockingHit())
+		{
+			UMoverComponent* MoverComponent = MovingComps.MoverComponent.Get();
+			const UPrimitiveComponent* UpdatedPrimitive = MovingComps.UpdatedPrimitive.Get();
+
+			// Notify first impact
+			if (MoverComponent && bHandleImpact)
+			{
+				FMoverOnImpactParams ImpactParams(NAME_None, InOutHit, SlideDelta);
+				MoverComponent->HandleImpact(ImpactParams);
+			}
+
+			const FVector LocationAfter1stSlide = InOutHit.TraceStart + ((InOutHit.TraceEnd - InOutHit.TraceStart) * InOutHit.Time);
+
+			// Check if we landed
+			if (!UAirMovementUtils::IsValidLandingSpot(MovingComps, LocationAfter1stSlide,
+				InOutHit, FloorSweepDistance, MaxWalkSlopeCosine, OutFloorResult))
+			{
+				// We've hit another surface during our first move, so let's try to slide along both of them together
+
+				// Compute new slide normal when hitting multiple surfaces.
+				SlideDelta = UMovementUtils::ComputeTwoWallAdjustedDelta(MovingComps, SlideDelta, InOutHit, OrigHitNormal);
+
+				// Only proceed if the new direction is of significant length and not in reverse of original attempted move.
+				if (!SlideDelta.IsNearlyZero(UE::MoverUtils::SMALL_MOVE_DISTANCE) && (SlideDelta | OriginalMoveDelta) > 0.f)
+				{
+					// Perform second move, taking 2 walls into account
+					UAsyncMovementUtils::TestDepenetratingMove(MovingComps, LocationAfter1stSlide, LocationAfter1stSlide + SlideDelta, TargetRotation, TargetRotation, /*bShouldSweep*/ true, InOutHit, InOutMoveRecord);
+					PctOfTimeUsed += (InOutHit.Time * (1.f - PctOfTimeUsed));
+
+					// Notify second impact
+					if (MoverComponent && bHandleImpact && InOutHit.bBlockingHit)
+					{
+						FMoverOnImpactParams ImpactParams(NAME_None, InOutHit, SlideDelta);
+						MoverComponent->HandleImpact(ImpactParams);
+					}
+
+					const FVector LocationAfter2ndSlide = InOutHit.TraceStart + ((InOutHit.TraceEnd - InOutHit.TraceStart) * InOutHit.Time);
+
+					// Check if we've landed, to acquire floor result
+					UAirMovementUtils::IsValidLandingSpot(MovingComps, LocationAfter2ndSlide,
+						InOutHit, FloorSweepDistance, MaxWalkSlopeCosine, OutFloorResult);
+				}
+			}
+		}
+
+		return FMath::Clamp(PctOfTimeUsed, 0.f, 1.f);
+	}
+
+	return 0.f;
+
+}
+
diff --git a/Engine/Plugins/Experimental/Mover/Source/Mover/Private/MoveLibrary/GroundMovementUtils.cpp b/Engine/Plugins/Experimental/Mover/Source/Mover/Private/MoveLibrary/GroundMovementUtils.cpp
--- a/Engine/Plugins/Experimental/Mover/Source/Mover/Private/MoveLibrary/GroundMovementUtils.cpp
+++ b/Engine/Plugins/Experimental/Mover/Source/Mover/Private/MoveLibrary/GroundMovementUtils.cpp
@@ -6,6 +6,7 @@
 #include "MoverLog.h"
 #include "Components/CapsuleComponent.h"
 #include "GameFramework/Pawn.h"
+#include "MoveLibrary/AsyncMovementUtils.h"
 #include "MoveLibrary/FloorQueryUtils.h"
 #include "MoveLibrary/MovementRecord.h"
 #include "MoveLibrary/MovementUtils.h"
@@ -557,3 +558,459 @@ bool UGroundMovementUtils::CanStepUpOnHitSurface(const FHitResult& Hit)
 
 	return true;
 }
+
+float UGroundMovementUtils::TestGroundedMoveAlongHitSurface(const FMovingComponentSet& MovingComps, const FVector& OriginalMoveDelta, const FVector& LocationAtHit, const FQuat& TargetRotation, bool bHandleImpact, float MaxStepHeight, float MaxWalkSlopeCosine, FHitResult& InOutHit, FMovementRecord& InOutMoveRecord)
+{
+	if (!InOutHit.bBlockingHit)
+	{
+		return 0.f;
+	}
+
+	float PctOfOrigDeltaToSlide = 1.f - InOutHit.Time;
+
+	if (UMoverComponent* MoverComponent = MovingComps.MoverComponent.Get())
+	{
+		const FVector UpDirection = MoverComponent->GetUpDirection();
+		FVector SafeWalkNormal(InOutHit.Normal);
+		const FVector::FReal NormalDot = InOutHit.Normal.Dot(UpDirection);
+
+		// We don't want to be pushed up an unwalkable surface.
+		if (NormalDot > 0.f && !UFloorQueryUtils::IsHitSurfaceWalkable(InOutHit, UpDirection, MaxWalkSlopeCosine))
+		{
+			SafeWalkNormal = FVector::VectorPlaneProject(SafeWalkNormal, -UpDirection).GetSafeNormal();
+		}
+
+		float PctOfTimeUsed = 0.f;
+		const FVector OldSafeHitNormal = SafeWalkNormal;
+
+		FVector SlideDelta = UMovementUtils::ComputeSlideDelta(MovingComps, OriginalMoveDelta, PctOfOrigDeltaToSlide, SafeWalkNormal, InOutHit);
+		const FVector OriginalSlideDelta = SlideDelta;
+
+		if (SlideDelta.Dot(OriginalMoveDelta) > 0.f)
+		{
+			UAsyncMovementUtils::TestDepenetratingMove(MovingComps, LocationAtHit, LocationAtHit + SlideDelta, TargetRotation, TargetRotation, /*bShouldSweep*/ true, InOutHit, InOutMoveRecord);
+
+			FVector LocationInProgress = InOutHit.TraceStart + ((InOutHit.TraceEnd - InOutHit.TraceStart) * InOutHit.Time);
+
+			PctOfTimeUsed = InOutHit.Time;
+
+			if (InOutHit.IsValidBlockingHit())
+			{
+				// Notify first impact
+				if (MoverComponent && bHandleImpact)
+				{
+					FMoverOnImpactParams ImpactParams(NAME_None, InOutHit, SlideDelta);
+					MoverComponent->HandleImpact(ImpactParams);
+				}
+
+				// Compute new slide normal when hitting multiple surfaces.
+				SlideDelta = UMovementUtils::ComputeTwoWallAdjustedDelta(MovingComps, SlideDelta, InOutHit, OldSafeHitNormal);
+				const float SlideDeltaDot = SlideDelta.Dot(UpDirection);
+				const float HitNormalDot = InOutHit.Normal.Dot(UpDirection);
+				if (SlideDeltaDot > 0.f && UFloorQueryUtils::IsHitSurfaceWalkable(InOutHit, UpDirection, MaxWalkSlopeCosine) && HitNormalDot > UE_KINDA_SMALL_NUMBER)
+				{
+					// Maintain horizontal velocity
+					const float Time = (1.f - InOutHit.Time);
+					const FVector ScaledDelta = SlideDelta.GetSafeNormal() * SlideDelta.Size();
+					const FVector::FReal NewDeltaDot = (ScaledDelta.Dot(UpDirection) / HitNormalDot);
+					FVector SlideProjectedToGrav = FVector::VectorPlaneProject(OriginalSlideDelta, -UpDirection);
+					SlideDelta = (SlideProjectedToGrav + NewDeltaDot * UpDirection) * Time;
+					// Should never exceed MaxStepHeight in vertical component, so rescale if necessary.
+					// This should be rare (HitNormalDot above would have been very small) but we'd rather lose horizontal velocity than go too high.
+					if (SlideDeltaDot > MaxStepHeight)
+					{
+						const float Rescale = MaxStepHeight / SlideDeltaDot;
+						SlideDelta *= Rescale;
+					}
+				}
+				else
+				{
+					SlideDelta = FVector::VectorPlaneProject(SlideDelta, -UpDirection);
+				}
+
+				// Only proceed if the new direction is of significant length and not in reverse of original attempted move.
+				if (!SlideDelta.IsNearlyZero(UE::MoverUtils::SMALL_MOVE_DISTANCE) && (SlideDelta | OriginalMoveDelta) > 0.f)
+				{
+					// Perform second move
+					UAsyncMovementUtils::TestDepenetratingMove(MovingComps, LocationInProgress, LocationInProgress + SlideDelta, TargetRotation, TargetRotation, /*bShouldSweep*/ true, InOutHit, InOutMoveRecord);
+					PctOfTimeUsed += (InOutHit.Time * (1.f - PctOfTimeUsed));
+
+					// Notify second impact
+					if (MoverComponent && bHandleImpact && InOutHit.bBlockingHit)
+					{
+						FMoverOnImpactParams ImpactParams(NAME_None, InOutHit, SlideDelta);
+						MoverComponent->HandleImpact(ImpactParams);
+					}
+				}
+			}
+
+			return FMath::Clamp(PctOfTimeUsed, 0.f, 1.f);
+		}
+	}
+
+	return 0.f;
+}
+
+
+/* static */
+FVector UGroundMovementUtils::TestMoveToAdjustToFloor(const FMovingComponentSet& MovingComps, const FVector& Location, const FQuat& Rotation, float MaxWalkSlopeCosine, FFloorCheckResult& InOutCurrentFloor, FMovementRecord& InOutMoveRecord)
+{
+	// If we have a floor check that hasn't hit anything, don't adjust height.
+	if (!InOutCurrentFloor.IsWalkableFloor())
+	{
+		return Location;	// no adjustment
+	}
+
+	const FVector UpDirection = MovingComps.MoverComponent->GetUpDirection();
+
+	float OldFloorDist = InOutCurrentFloor.FloorDist;
+	if (InOutCurrentFloor.bLineTrace)
+	{
+		if (OldFloorDist < UE::FloorQueryUtility::MIN_FLOOR_DIST && InOutCurrentFloor.LineDist >= UE::FloorQueryUtility::MIN_FLOOR_DIST)
+		{
+			// This would cause us to scale unwalkable walls
+			return Location;	// no adjustment
+		}
+		// Falling back to a line trace means the sweep was unwalkable (or in penetration). Use the line distance for the vertical adjustment.
+		OldFloorDist = InOutCurrentFloor.LineDist;
+	}
+
+	// Move up or down to maintain floor height.
+	if (OldFloorDist < UE::FloorQueryUtility::MIN_FLOOR_DIST || OldFloorDist > UE::FloorQueryUtility::MAX_FLOOR_DIST)
+	{
+		FHitResult AdjustHit(1.f);
+		const float InitialUp = Location.Dot(UpDirection);
+		const float AvgFloorDist = (UE::FloorQueryUtility::MIN_FLOOR_DIST + UE::FloorQueryUtility::MAX_FLOOR_DIST) * 0.5f;
+		const float MoveDist = AvgFloorDist - OldFloorDist;
+		const FVector TargetAdjustedLocation = Location + (UpDirection * MoveDist);
+
+
+		InOutMoveRecord.LockRelevancy(false);	// Makes it so any height adjustment will not count towards velocity changes
+
+		UAsyncMovementUtils::TestDepenetratingMove(MovingComps, Location, TargetAdjustedLocation, Rotation, Rotation, /*bShouldSweep*/ true, AdjustHit, InOutMoveRecord);
+
+		InOutMoveRecord.UnlockRelevancy();
+
+		const FVector AdjustedLocation = AdjustHit.TraceStart + ((AdjustHit.TraceEnd - AdjustHit.TraceStart) * AdjustHit.Time);
+
+		if (!AdjustHit.IsValidBlockingHit())
+		{
+			InOutCurrentFloor.FloorDist += MoveDist;
+		}
+		else if (MoveDist > 0.f)
+		{
+			const float CurrentUp = AdjustedLocation.Dot(UpDirection);
+			InOutCurrentFloor.FloorDist += CurrentUp - InitialUp;
+		}
+		else
+		{
+			checkSlow(MoveDist < 0.f);
+			const float CurrentUp = AdjustedLocation.Dot(UpDirection);
+			InOutCurrentFloor.FloorDist = CurrentUp - AdjustHit.Location.Dot(UpDirection);
+			if (UFloorQueryUtils::IsHitSurfaceWalkable(AdjustHit, UpDirection, MaxWalkSlopeCosine))
+			{
+				InOutCurrentFloor.SetFromSweep(AdjustHit, InOutCurrentFloor.FloorDist, true);
+			}
+		}
+
+		return AdjustedLocation;
+	}
+
+	return Location; // no adjustment
+}
+
+/* static */
+bool UGroundMovementUtils::TestMoveToStepOver(const FMovingComponentSet& MovingComps, const FVector& GravDir, float MaxStepHeight, float MaxWalkSlopeCosine, float FloorSweepDistance, const FVector& MoveDelta, const FQuat& Rotation, const FHitResult& MoveHitResult, const FFloorCheckResult& CurrentFloor, bool bIsFalling, FOptionalFloorCheckResult* OutFloorTestResult, FVector& OutFinalLocation, FMovementRecord& InOutMoveRecord)
+{
+	UCapsuleComponent* CapsuleComponent = Cast<UCapsuleComponent>(MovingComps.UpdatedPrimitive.Get());
+
+	if (!UGroundMovementUtils::CanStepUpOnHitSurface(MoveHitResult) || MaxStepHeight <= 0.f)
+	{
+		return false;
+	}
+
+	TArray<FMovementSubstep> QueuedSubsteps;	// keeping track of substeps before committing, because some moves can be backed out
+
+	FVector UpDirection = MovingComps.MoverComponent->GetUpDirection();
+
+	const FVector OldLocation = MoveHitResult.TraceStart + ((MoveHitResult.TraceEnd - MoveHitResult.TraceStart) * MoveHitResult.Time);
+	FVector LocationInProgress = OldLocation;
+
+	float PawnRadius, PawnHalfHeight;
+
+	CapsuleComponent->GetScaledCapsuleSize(PawnRadius, PawnHalfHeight);
+
+	// Don't bother stepping up if top of capsule is hitting something.
+	const float InitialImpactDot = MoveHitResult.ImpactPoint.Dot(UpDirection);
+	const float OldLocationDot = OldLocation.Dot(UpDirection);
+	if (InitialImpactDot > OldLocationDot + (PawnHalfHeight - PawnRadius))
+	{
+		UE_LOG(LogMover, VeryVerbose, TEXT("Not stepping up due to top of capsule hitting something"));
+		return false;
+	}
+
+	// TODO: We should rely on movement plane normal, rather than gravity direction
+	if (GravDir.IsZero())
+	{
+		UE_LOG(LogMover, VeryVerbose, TEXT("Not stepping up because there's no gravity"));
+		return false;
+	}
+
+	// Gravity should be a normalized direction
+	ensure(GravDir.IsNormalized());
+
+	float StepTravelUpHeight = MaxStepHeight;
+	float StepTravelDownHeight = StepTravelUpHeight;
+	const float StepSideDot = -1.f * FVector::DotProduct(MoveHitResult.ImpactNormal, GravDir);
+	float PawnInitialFloorBaseDot = OldLocationDot - PawnHalfHeight;
+	float PawnFloorPointDot = PawnInitialFloorBaseDot;
+
+	if (CurrentFloor.IsWalkableFloor())
+	{
+		// Since we float a variable amount off the floor, we need to enforce max step height off the actual point of impact with the floor.
+		const float FloorDist = FMath::Max(0.f, CurrentFloor.GetDistanceToFloor());
+		PawnInitialFloorBaseDot -= FloorDist;
+		StepTravelUpHeight = FMath::Max(StepTravelUpHeight - FloorDist, 0.f);
+		StepTravelDownHeight = (MaxStepHeight + UE::FloorQueryUtility::MAX_FLOOR_DIST * 2.f);
+
+		const bool bHitVerticalFace = !UFloorQueryUtils::IsWithinEdgeTolerance(MoveHitResult.Location, MoveHitResult.ImpactPoint, PawnRadius, UpDirection);
+		if (!CurrentFloor.bLineTrace && !bHitVerticalFace)
+		{
+			PawnFloorPointDot = CurrentFloor.HitResult.ImpactPoint.Dot(UpDirection);
+		}
+		else
+		{
+			// Base floor point is the base of the capsule moved down by how far we are hovering over the surface we are hitting.
+			PawnFloorPointDot -= CurrentFloor.FloorDist;
+		}
+	}
+
+	// Don't step up if the impact is below us, accounting for distance from floor.
+	if (InitialImpactDot <= PawnInitialFloorBaseDot)
+	{
+		UE_LOG(LogMover, VeryVerbose, TEXT("Not stepping up because the impact is below us"));
+		return false;
+	}
+
+	// step up - treat as vertical wall
+	FHitResult SweepUpHit(1.f);
+
+	const FVector UpAdjustment = -GravDir * StepTravelUpHeight;
+
+	const bool bDidStepUp = UAsyncMovementUtils::TestMoveComponent_Internal(MovingComps, LocationInProgress, LocationInProgress + UpAdjustment, Rotation, Rotation, /*bShouldSweep=*/ true, MOVECOMP_NoFlags, SweepUpHit);
+
+	UE_LOG(LogMover, VeryVerbose, TEXT("TestMoveToStepOver Up: %s (role %i) UpAdjustment=%s DidMove=%i"),
+		*GetNameSafe(CapsuleComponent->GetOwner()), CapsuleComponent->GetOwnerRole(), *UpAdjustment.ToCompactString(), bDidStepUp);
+
+	if (SweepUpHit.bStartPenetrating)
+	{
+		// Undo movement
+		UE_LOG(LogMover, VeryVerbose, TEXT("Reverting step-up attempt because we started in a penetrating state"));
+		return false;
+	}
+
+	FVector UpStepDelta = (SweepUpHit.TraceStart + ((SweepUpHit.TraceEnd - SweepUpHit.TraceStart) * SweepUpHit.Time)) - LocationInProgress;
+
+	// Cache upwards substep
+	QueuedSubsteps.Add(FMovementSubstep(StepUpSubstepName, UpStepDelta, false));
+	LocationInProgress += UpStepDelta;
+
+	// step fwd
+	FHitResult StepFwdHit(1.f);
+	const bool bDidStepFwd = UAsyncMovementUtils::TestMoveComponent_Internal(MovingComps, LocationInProgress, LocationInProgress + MoveDelta, Rotation, Rotation, /*bShouldSweep=*/ true, MOVECOMP_NoFlags, StepFwdHit);
+
+	FVector FwdStepDelta = (StepFwdHit.TraceStart + ((StepFwdHit.TraceEnd - StepFwdHit.TraceStart) * StepFwdHit.Time)) - LocationInProgress;
+
+	UE_LOG(LogMover, VeryVerbose, TEXT("TestMoveToStepOver Fwd: %s (role %i) MoveDelta=%s DidMove=%i"),
+		*GetNameSafe(CapsuleComponent->GetOwner()), CapsuleComponent->GetOwnerRole(), *MoveDelta.ToCompactString(), bDidStepFwd);
+
+	// Check result of forward movement
+	if (StepFwdHit.bBlockingHit)
+	{
+		if (StepFwdHit.bStartPenetrating)
+		{
+			// Undo movement
+			UE_LOG(LogMover, VeryVerbose, TEXT("Reverting step-fwd attempt during step-up, because we started in a penetrating state"));
+			return false;
+		}
+
+		UMoverComponent* MoverComp = MovingComps.MoverComponent.Get();
+		// If we hit something above us and also something ahead of us, we should notify about the upward hit as well.
+		// The forward hit will be handled later (in the bSteppedOver case below).
+		// In the case of hitting something above but not forward, we are not blocked from moving so we don't need the notification.
+		if (MoverComp && SweepUpHit.bBlockingHit && StepFwdHit.bBlockingHit)
+		{
+			FMoverOnImpactParams ImpactParams(NAME_None, SweepUpHit, MoveDelta);
+			MoverComp->HandleImpact(ImpactParams);
+		}
+
+		// pawn ran into a wall
+		if (MoverComp)
+		{
+			FMoverOnImpactParams ImpactParams(NAME_None, StepFwdHit, MoveDelta);
+			MoverComp->HandleImpact(ImpactParams);
+		}
+
+		// Cache forwards substep before the slide attempt
+		QueuedSubsteps.Add(FMovementSubstep(StepFwdSubstepName, FwdStepDelta, true));
+		LocationInProgress += FwdStepDelta;
+
+
+		// If falling, we only need to try to reach up and fwd. We do not need to continue searching downward.
+		if (bIsFalling)
+		{
+			// Commit queued substeps to movement record
+			for (FMovementSubstep Substep : QueuedSubsteps)
+			{
+				InOutMoveRecord.Append(Substep);
+			}
+
+			OutFinalLocation = LocationInProgress;
+			return true;
+		}
+
+
+		// adjust and try again
+		const float ForwardHitTime = StepFwdHit.Time;
+
+		// locking relevancy so velocity isn't added until it is needed to (adding it to the QueuedSubsteps so it can get added later)
+		InOutMoveRecord.LockRelevancy(false);
+
+		const float ForwardSlideAmount = TestGroundedMoveAlongHitSurface(MovingComps, MoveDelta, LocationInProgress, Rotation, /*bHandleImpact=*/true, MaxStepHeight, MaxWalkSlopeCosine, StepFwdHit, IN OUT InOutMoveRecord);
+
+		const FVector FwdSlideDelta = (StepFwdHit.TraceStart + ((StepFwdHit.TraceEnd - StepFwdHit.TraceStart) * StepFwdHit.Time)) - LocationInProgress;
+
+		QueuedSubsteps.Add(FMovementSubstep(SlideSubstepName, FwdSlideDelta, true));
+		LocationInProgress += FwdSlideDelta;
+
+		InOutMoveRecord.UnlockRelevancy();
+
+		if (bIsFalling)
+		{
+			UE_LOG(LogMover, VeryVerbose, TEXT("Reverting step-fwd attempt during step-up, because we could not adjust without falling"));
+			return false;
+		}
+
+		// If both the forward hit and the deflection got us nowhere, there is no point in this step up.
+		if (ForwardHitTime == 0.f && ForwardSlideAmount == 0.f)
+		{
+			UE_LOG(LogMover, VeryVerbose, TEXT("Reverting step-fwd attempt during step-up, because no movement differences occurred"));
+			return false;
+		}
+	}
+	else
+	{
+		// Our forward move attempt was unobstructed - cache it
+		QueuedSubsteps.Add(FMovementSubstep(StepFwdSubstepName, FwdStepDelta, true));
+		LocationInProgress += FwdStepDelta;
+	}
+
+
+	// Step down
+	const FVector StepDownAdjustment = GravDir * StepTravelDownHeight;
+
+	const bool bDidStepDown = UAsyncMovementUtils::TestMoveComponent_Internal(MovingComps, LocationInProgress, LocationInProgress + StepDownAdjustment, Rotation, Rotation, /*bShouldSweep=*/ true, MOVECOMP_NoFlags, StepFwdHit);
+
+
+	UE_LOG(LogMover, VeryVerbose, TEXT("TestMoveToStepOver Down: %s (role %i) StepDownAdjustment=%s DidMove=%i"),
+		*GetNameSafe(CapsuleComponent->GetOwner()), CapsuleComponent->GetOwnerRole(), *StepDownAdjustment.ToCompactString(), bDidStepDown);
+
+
+	// If step down was initially penetrating abort the step up
+	if (StepFwdHit.bStartPenetrating)
+	{
+		UE_LOG(LogMover, VeryVerbose, TEXT("Reverting step-down attempt during step-up/step-fwd, because we started in a penetrating state"));
+		return false;
+	}
+
+	FOptionalFloorCheckResult StepDownResult;
+	if (StepFwdHit.IsValidBlockingHit())
+	{
+		// See if this step sequence would have allowed us to travel higher than our max step height allows.
+		const float DeltaDot = (StepFwdHit.ImpactPoint.Dot(UpDirection)) - PawnFloorPointDot;
+		if (DeltaDot > MaxStepHeight)
+		{
+			UE_LOG(LogMover, VeryVerbose, TEXT("Reject step-down attempt during step-up/step-fwd, because it made us travel too high (too high Height %.3f) up from floor base %f to %f"), DeltaDot, PawnInitialFloorBaseDot, StepFwdHit.ImpactPoint.Z);
+			return false;
+		}
+
+		// Reject unwalkable surface normals here.
+		if (!UFloorQueryUtils::IsHitSurfaceWalkable(StepFwdHit, UpDirection, MaxWalkSlopeCosine))
+		{
+			// Reject if normal opposes movement direction
+			const bool bNormalTowardsMe = (MoveDelta | StepFwdHit.ImpactNormal) < 0.f;
+			if (bNormalTowardsMe)
+			{
+				UE_LOG(LogMover, VeryVerbose, TEXT("Reject step-down attempt during step-up/step-fwd, due to unwalkable normal %s opposed to movement"), *StepFwdHit.ImpactNormal.ToString());
+				return false;
+			}
+
+			// Also reject if we would end up being higher than our starting location by stepping down.
+			// It's fine to step down onto an unwalkable normal below us, we will just slide off. Rejecting those moves would prevent us from being able to walk off the edge.
+			if (StepFwdHit.Location.Dot(UpDirection) > OldLocationDot)
+			{
+				UE_LOG(LogMover, VeryVerbose, TEXT("Reject step-down attempt during step-up/step-fwd, due to unwalkable normal %s above old position)"), *StepFwdHit.ImpactNormal.ToString());
+				return false;
+			}
+		}
+
+		// Reject moves where the downward sweep hit something very close to the edge of the capsule. This maintains consistency with FindFloor as well.
+		if (!UFloorQueryUtils::IsWithinEdgeTolerance(StepFwdHit.Location, StepFwdHit.ImpactPoint, PawnRadius, UpDirection))
+		{
+			UE_LOG(LogMover, VeryVerbose, TEXT("Reject step-down attempt during step-up/step-fwd, due to being outside edge tolerance"));
+			return false;
+		}
+
+		// Don't step up onto invalid surfaces if traveling higher.
+		if (DeltaDot > 0.f && !UGroundMovementUtils::CanStepUpOnHitSurface(StepFwdHit))
+		{
+			UE_LOG(LogMover, VeryVerbose, TEXT("Reject step-down attempt during step-up/step-fwd, due to being up onto surface with !CanStepUpOnHitSurface"));
+			return false;
+		}
+
+		// See if we can validate the floor as a result of this step down. In almost all cases this should succeed, and we can avoid computing the floor outside this method.
+		if (OutFloorTestResult != NULL)
+		{
+
+			UFloorQueryUtils::FindFloor(MovingComps, FloorSweepDistance, MaxWalkSlopeCosine, LocationInProgress, StepDownResult.FloorTestResult);
+
+			// Reject unwalkable normals if we end up higher than our initial height.
+			// It's fine to walk down onto an unwalkable surface, don't reject those moves.
+			if (StepFwdHit.Location.Dot(UpDirection) > OldLocationDot)
+			{
+				// We should reject the floor result if we are trying to step up an actual step where we are not able to perch (this is rare).
+				// In those cases we should instead abort the step up and try to slide along the stair.
+				const float MAX_STEP_SIDE_DOT = 0.08f; // TODO: Move magic numbers elsewhere
+				if (!StepDownResult.FloorTestResult.bBlockingHit && StepSideDot < MAX_STEP_SIDE_DOT)
+				{
+					UE_LOG(LogMover, VeryVerbose, TEXT("Reject step-down attempt during step-up/step-fwd, due to it being an unperchable step"));
+					return false;
+				}
+			}
+
+			StepDownResult.bHasFloorResult = true;
+		}
+	}
+
+	const FVector StepDownDelta = (StepFwdHit.TraceStart + ((StepFwdHit.TraceEnd - StepFwdHit.TraceStart) * StepFwdHit.Time)) - LocationInProgress;
+
+
+	// Cache downwards substep
+	QueuedSubsteps.Add(FMovementSubstep(StepDownSubstepName, StepDownDelta, false));
+	LocationInProgress += StepDownDelta;
+
+	// Copy step down result.
+	if (OutFloorTestResult != NULL)
+	{
+		*OutFloorTestResult = StepDownResult;
+	}
+
+	// Commit queued substeps to movement record
+	for (FMovementSubstep Substep : QueuedSubsteps)
+	{
+		InOutMoveRecord.Append(Substep);
+	}
+
+	OutFinalLocation = LocationInProgress;
+	return true;
+}
diff --git a/Engine/Plugins/Experimental/Mover/Source/Mover/Private/MoverComponent.cpp b/Engine/Plugins/Experimental/Mover/Source/Mover/Private/MoverComponent.cpp
--- a/Engine/Plugins/Experimental/Mover/Source/Mover/Private/MoverComponent.cpp
+++ b/Engine/Plugins/Experimental/Mover/Source/Mover/Private/MoverComponent.cpp
@@ -38,7 +38,7 @@
 namespace MoverComponentCVars
 {
 	static int32 WarnOnPostSimDifference = 0;
-	static FAutoConsoleVariableRef CVarMoverWarnOnPostSimDifference(
+	FAutoConsoleVariableRef CVarMoverWarnOnPostSimDifference(
 		TEXT("mover.debug.WarnOnPostSimDifference"),
 		WarnOnPostSimDifference,
 		TEXT("If != 0, then any differences between the sim sync state and the component locations just after movement simulation will emit warnings.\n")
diff --git a/Engine/Plugins/Experimental/Mover/Source/Mover/Public/DefaultMovementSet/Modes/AsyncFallingMode.h b/Engine/Plugins/Experimental/Mover/Source/Mover/Public/DefaultMovementSet/Modes/AsyncFallingMode.h
new file mode 100644
--- /dev/null
+++ b/Engine/Plugins/Experimental/Mover/Source/Mover/Public/DefaultMovementSet/Modes/AsyncFallingMode.h
@@ -0,0 +1,99 @@
+// Copyright Epic Games, Inc. All Rights Reserved.
+
+#pragma once
+
+#include "MovementMode.h"
+#include "FallingMode.h"
+#include "AsyncFallingMode.generated.h"
+
+class UCommonLegacyMovementSettings;
+struct FFloorCheckResult;
+
+/**
+ * AsyncFallingMode: a default movement mode for moving through the air and jumping, typically influenced by gravity and air control
+ * This mode simulates movement without actually modifying any scene component(s).
+ */
+UCLASS(Blueprintable, BlueprintType, Experimental)
+class MOVER_API UAsyncFallingMode : public UBaseMovementMode
+{
+	GENERATED_BODY()
+
+
+public:
+	UAsyncFallingMode(const FObjectInitializer& ObjectInitializer);
+
+	UFUNCTION(BlueprintCallable, Category=Mover)
+	virtual void OnGenerateMove(const FMoverTickStartData& StartState, const FMoverTimeStep& TimeStep, FProposedMove& OutProposedMove) const override;
+
+	UFUNCTION(BlueprintCallable, Category=Mover)
+	virtual void OnSimulationTick(const FSimulationTickParams& Params, FMoverTickEndData& OutputState) override;
+
+	// Broadcast when this actor lands on a valid surface.
+	UPROPERTY(BlueprintAssignable, Category = Mover)
+	FMover_OnLanded OnLanded;
+	
+	/**
+	 * If true, actor will land and lose all speed in the vertical direction upon landing. If false, actor's vertical speed will be redirected based on the surface normal it hit.
+	 * Note: Actor's horizontal speed will not be affected if true. If false, horizontal speed may be increased on landing.
+	 */
+	UPROPERTY(Category = Mover, EditAnywhere, BlueprintReadWrite)
+	bool bCancelVerticalSpeedOnLanding;
+	
+	/**
+	 * When falling, amount of movement control available to the actor.
+	 * 0 = no control, 1 = full control
+	 */
+	UPROPERTY(Category=Mover, EditAnywhere, BlueprintReadWrite, meta=(ClampMin="0", ClampMax="1.0"))
+	float AirControlPercentage;
+	
+	/**
+ 	 * Deceleration to apply to air movement when falling slower than terminal velocity.
+ 	 * Note: This is NOT applied to vertical velocity, only movement plane velocity
+ 	 */
+	UPROPERTY(Category=Mover, EditAnywhere, BlueprintReadWrite, meta=(ClampMin="0", ForceUnits = "cm/s^2"))
+	float FallingDeceleration;
+
+	/**
+     * Deceleration to apply to air movement when falling faster than terminal velocity
+	 * Note: This is NOT applied to vertical velocity, only movement plane velocity
+     */
+    UPROPERTY(Category=Mover, EditAnywhere, BlueprintReadWrite, meta=(ClampMin="0", ForceUnits = "cm/s^2"))
+    float OverTerminalSpeedFallingDeceleration;
+	
+	/**
+	 * If the actor's movement plane velocity is greater than this speed falling will start applying OverTerminalSpeedFallingDeceleration instead of FallingDeceleration
+	 * The expected behavior is to set OverTerminalSpeedFallingDeceleration higher than FallingDeceleration so the actor will slow down faster
+	 * when going over TerminalMovementPlaneSpeed.
+	 */
+	UPROPERTY(Category=Mover, EditAnywhere, BlueprintReadWrite, meta=(ClampMin="0", ForceUnits = "cm/s"))
+	float TerminalMovementPlaneSpeed;
+
+	/** When exceeding maximum vertical speed, should it be enforced via a hard clamp? If false, VerticalFallingDeceleration will be used for a smoother transition to the terminal speed limit. */
+	UPROPERTY(Category = Mover, EditAnywhere, BlueprintReadWrite)
+	bool bShouldClampTerminalVerticalSpeed;
+
+	/** Deceleration to apply to vertical velocity when it's greater than TerminalVerticalSpeed. Only used if bShouldClampTerminalVerticalSpeed is false. */
+	UPROPERTY(Category=Mover, EditAnywhere, BlueprintReadWrite, meta=(EditCondition="!bShouldClampTerminalVerticalSpeed", ClampMin="0", ForceUnits = "cm/s^2"))
+	float VerticalFallingDeceleration;
+	
+	/**
+	 * If the actors vertical velocity is greater than this speed VerticalFallingDeceleration will be applied to vertical velocity
+	 */
+	UPROPERTY(Category=Mover, EditAnywhere, BlueprintReadWrite, meta=(ClampMin="0", ForceUnits = "cm/s"))
+	float TerminalVerticalSpeed;
+	
+protected:
+	/**
+	 * Is called at the end of the tick in falling mode. Handles checking any landings that should occur and switching to specific modes
+	 * (i.e. landing on a walkable surface would switch to the walking movement mode) 
+	 */
+	UFUNCTION(BlueprintCallable, Category=Mover)
+	virtual void ProcessLanded(const FFloorCheckResult& FloorResult, FVector& Velocity, FRelativeBaseInfo& BaseInfo, FMoverTickEndData& TickEndData) const;
+	
+	virtual void OnRegistered(const FName ModeName) override;
+	virtual void OnUnregistered() override;
+
+	void CaptureFinalState(const FVector FinalLocation, const FRotator FinalRotation, const FFloorCheckResult& FloorResult, float DeltaSeconds, float DeltaSecondsUsed, FMoverDefaultSyncState& OutputSyncState, FMoverTickEndData& TickEndData, FMovementRecord& Record) const;
+
+	TWeakObjectPtr<const UCommonLegacyMovementSettings> CommonLegacySettings;
+};
diff --git a/Engine/Plugins/Experimental/Mover/Source/Mover/Public/DefaultMovementSet/Modes/AsyncNavWalkingMode.h b/Engine/Plugins/Experimental/Mover/Source/Mover/Public/DefaultMovementSet/Modes/AsyncNavWalkingMode.h
new file mode 100644
--- /dev/null
+++ b/Engine/Plugins/Experimental/Mover/Source/Mover/Public/DefaultMovementSet/Modes/AsyncNavWalkingMode.h
@@ -0,0 +1,135 @@
+// Copyright Epic Games, Inc. All Rights Reserved.
+
+#pragma once
+
+#include "MovementMode.h"
+#include "MoveLibrary/BasedMovementUtils.h"
+#include "AI/Navigation/NavigationTypes.h"
+#include "NavWalkingMode.h"
+#include "AsyncNavWalkingMode.generated.h"
+
+class UNavMoverComponent;
+class INavigationDataInterface;
+class UCommonLegacyMovementSettings;
+
+
+/**
+ * AsyncNavWalkingMode: a default movement mode for traversing surfaces and movement bases by using an active navmesh when moving the actor rather than collision checks.
+ * Note: This movement mode requires a NavMoverComponent be on the actor to function properly. This mode also contains some randomization to avoid navmesh look ups
+ *	happening at the same time (which is fine for AI characters running on the server) but may cause issues if used on autonomous proxies.
+ * This mode simulates movement without actually modifying any scene component(s).
+ */
+UCLASS(Blueprintable, BlueprintType, Experimental)
+class MOVER_API UAsyncNavWalkingMode : public UBaseMovementMode
+{
+	GENERATED_BODY()
+
+public:
+	UAsyncNavWalkingMode();
+
+	UFUNCTION(BlueprintCallable, Category = Mover)
+	virtual void OnGenerateMove(const FMoverTickStartData& StartState, const FMoverTimeStep& TimeStep, FProposedMove& OutProposedMove) const override;
+
+	UFUNCTION(BlueprintCallable, Category = Mover)
+	virtual void OnSimulationTick(const FSimulationTickParams& Params, FMoverTickEndData& OutputState) override;
+
+	/**
+	 * Whether or not the actor should sweep for collision geometry while walking.
+	 */
+	UPROPERTY(Category = "NavMesh Movement", EditAnywhere, BlueprintReadWrite)
+	bool bSweepWhileNavWalking;
+
+	/** Whether to raycast to underlying geometry to better conform navmesh-walking actors */
+	UPROPERTY(Category = "NavMesh Movement", EditAnywhere, BlueprintReadOnly)
+	bool bProjectNavMeshWalking;
+
+	/**
+	 * Scale of the total capsule height to use for projection from navmesh to underlying geometry in the upward direction.
+	 * In other words, start the trace at [CapsuleHeight * NavMeshProjectionHeightScaleUp] above nav mesh.
+	 */
+	UPROPERTY(Category = "NavMesh Movement", EditAnywhere, BlueprintReadWrite, meta = (editcondition = "bProjectNavMeshWalking", ClampMin = "0", UIMin = "0"))
+	float NavMeshProjectionHeightScaleUp;
+
+	/**
+	 * Scale of the total capsule height to use for projection from navmesh to underlying geometry in the downward direction.
+	 * In other words, trace down to [CapsuleHeight * NavMeshProjectionHeightScaleDown] below nav mesh.
+	 */
+	UPROPERTY(Category = "NavMesh Movement", EditAnywhere, BlueprintReadWrite, meta = (editcondition = "bProjectNavMeshWalking", ClampMin = "0", UIMin = "0"))
+	float NavMeshProjectionHeightScaleDown;
+
+	/** How often we should raycast to project from navmesh to underlying geometry */
+	UPROPERTY(Category = "NavMesh Movement", EditAnywhere, BlueprintReadWrite, meta = (editcondition = "bProjectNavMeshWalking", ForceUnits = "s"))
+	float NavMeshProjectionInterval;
+
+	/** Speed at which to interpolate agent navmesh offset between traces. 0: Instant (no interp) > 0: Interp speed") */
+	UPROPERTY(Category = "NavMesh Movement", EditAnywhere, BlueprintReadWrite, meta = (editcondition = "bProjectNavMeshWalking", ClampMin = "0", UIMin = "0"))
+	float NavMeshProjectionInterpSpeed;
+
+	UPROPERTY(Transient)
+	float NavMeshProjectionTimer;
+
+	/** last known location projected on navmesh */
+	FNavLocation CachedNavLocation;
+
+	/**
+	 * Project a location to navmesh to find adjusted height.
+	 * @param TestLocation		Location to project
+	 * @param NavFloorLocation	Location on navmesh
+	 * @return True if projection was performed (successfully or not)
+	 */
+	virtual bool FindNavFloor(const FVector& TestLocation, FNavLocation& NavFloorLocation) const;
+
+	// Returns the active turn generator. Note: you will need to cast the return value to the generator you expect to get, it can also be none
+	UFUNCTION(BlueprintPure, Category = Mover)
+	UObject* GetTurnGenerator();
+
+	// Sets the active turn generator to use the class provided. Note: To set it back to the default implementation pass in none
+	UFUNCTION(BlueprintCallable, Category = Mover)
+	void SetTurnGeneratorClass(UPARAM(meta = (MustImplement = "/Script/Mover.TurnGeneratorInterface", AllowAbstract = "false")) TSubclassOf<UObject> TurnGeneratorClass);
+
+protected:
+	/** Associated Movement component that will actually move the actor */
+	TWeakObjectPtr<UNavMoverComponent> NavMoverComponent;
+
+	/** Use both WorldStatic and WorldDynamic channels for NavWalking geometry conforming */
+	UPROPERTY(Category = "NavMesh Movement", EditAnywhere, BlueprintReadOnly, AdvancedDisplay)
+	uint8 bProjectNavMeshOnBothWorldChannels : 1;
+
+	/** Optional modular object for generating rotation towards desired orientation. If not specified, linear interpolation will be used. */
+	UPROPERTY(EditAnywhere, Instanced, Category = Mover, meta = (MustImplement = "/Script/Mover.TurnGeneratorInterface"))
+	TObjectPtr<UObject> TurnGenerator;
+
+	/** Switch collision settings for NavWalking mode (ignore world collisions) */
+	virtual void SetCollisionForNavWalking(bool bEnable);
+
+	virtual void OnActivate() override;
+
+	virtual void OnDeactivate() override;
+
+	/** Get Navigation data for the actor. Returns null if there is no associated nav data. */
+	const INavigationDataInterface* GetNavData() const;
+
+	/** Performs trace for ProjectLocationFromNavMesh */
+	virtual void FindBestNavMeshLocation(const FVector& TraceStart, const FVector& TraceEnd, const FVector& CurrentFeetLocation, const FVector& TargetNavLocation, FHitResult& OutHitResult) const;
+
+	/**
+	 * Attempts to better align navmesh walking actors with underlying geometry (sometimes
+	 * navmesh can differ quite significantly from geometry).
+	 * Updates CachedProjectedNavMeshHitResult, access this for more info about hits.
+	 */
+	virtual FVector ProjectLocationFromNavMesh(float DeltaSeconds, const FVector& CurrentFeetLocation, const FVector& TargetNavLocation, float UpOffset, float DownOffset);
+
+
+	virtual void OnRegistered(const FName ModeName) override;
+	virtual void OnUnregistered() override;
+
+	void CaptureFinalState(const FVector FinalLocation, const FRotator FinalRotation, bool bDidAttemptMovement, const FFloorCheckResult& FloorResult, const FMovementRecord& Record, FMoverDefaultSyncState& OutputSyncState) const;
+
+	FRelativeBaseInfo UpdateFloorAndBaseInfo(const FFloorCheckResult& FloorResult) const;
+
+	TWeakObjectPtr<const UCommonLegacyMovementSettings> CommonLegacySettings;
+ 
+ 	// Collision responses saved at the activation of this mode, to be restored as upon deactivation if the CDO can't be used
+ 	ECollisionResponse CollideVsWorldStatic = ECollisionResponse::ECR_Block;
+ 	ECollisionResponse CollideVsWorldDynamic = ECollisionResponse::ECR_Block;
+};
diff --git a/Engine/Plugins/Experimental/Mover/Source/Mover/Public/DefaultMovementSet/Modes/AsyncWalkingMode.h b/Engine/Plugins/Experimental/Mover/Source/Mover/Public/DefaultMovementSet/Modes/AsyncWalkingMode.h
new file mode 100644
--- /dev/null
+++ b/Engine/Plugins/Experimental/Mover/Source/Mover/Public/DefaultMovementSet/Modes/AsyncWalkingMode.h
@@ -0,0 +1,60 @@
+// Copyright Epic Games, Inc. All Rights Reserved.
+
+#pragma once
+
+#include "MovementMode.h"
+#include "MoveLibrary/ModularMovement.h"
+#include "WalkingMode.h"
+#include "AsyncWalkingMode.generated.h"
+
+class UCommonLegacyMovementSettings;
+struct FFloorCheckResult;
+struct FRelativeBaseInfo;
+struct FMovementRecord;
+
+
+/**
+ * AsyncWalkingMode: a default movement mode for traversing surfaces and movement bases (walking, running, sneaking, etc.)
+ * This mode simulates movement without actually modifying any scene component(s).
+ */
+UCLASS(Blueprintable, BlueprintType, Experimental)
+class MOVER_API UAsyncWalkingMode : public UBaseMovementMode
+{
+	GENERATED_BODY()
+
+public:
+	UAsyncWalkingMode(const FObjectInitializer& ObjectInitializer);
+	
+	UFUNCTION(BlueprintCallable, Category=Mover)
+	virtual void OnGenerateMove(const FMoverTickStartData& StartState, const FMoverTimeStep& TimeStep, FProposedMove& OutProposedMove) const override;
+
+	UFUNCTION(BlueprintCallable, Category=Mover)
+	virtual void OnSimulationTick(const FSimulationTickParams& Params, FMoverTickEndData& OutputState) override;
+
+	// Returns the active turn generator. Note: you will need to cast the return value to the generator you expect to get, it can also be none
+	UFUNCTION(BlueprintPure, Category=Mover)
+	UObject* GetTurnGenerator();
+
+	// Sets the active turn generator to use the class provided. Note: To set it back to the default implementation pass in none
+	UFUNCTION(BlueprintCallable, Category=Mover)
+	void SetTurnGeneratorClass(UPARAM(meta=(MustImplement="/Script/Mover.TurnGeneratorInterface", AllowAbstract="false")) TSubclassOf<UObject> TurnGeneratorClass);
+
+protected:
+
+	/** Choice of behavior for floor checks while not moving.  */
+	UPROPERTY(EditAnywhere, BlueprintReadWrite, Category = Mover)
+	EStaticFloorCheckPolicy FloorCheckPolicy = EStaticFloorCheckPolicy::OnDynamicBaseOnly;
+
+	/** Optional modular object for generating rotation towards desired orientation. If not specified, linear interpolation will be used. */
+	UPROPERTY(EditAnywhere, Instanced, Category=Mover, meta=(MustImplement="/Script/Mover.TurnGeneratorInterface"))
+	TObjectPtr<UObject> TurnGenerator;
+
+	virtual void OnRegistered(const FName ModeName) override; 
+	virtual void OnUnregistered() override;
+
+	void CaptureFinalState(const FVector FinalLocation, const FRotator FinalRotation, bool bDidAttemptMovement, const FFloorCheckResult& FloorResult, const FMovementRecord& Record, FMoverDefaultSyncState& OutputSyncState) const;
+
+	FRelativeBaseInfo UpdateFloorAndBaseInfo(const FFloorCheckResult& FloorResult) const;
+
+	TWeakObjectPtr<const UCommonLegacyMovementSettings> CommonLegacySettings;
+};
diff --git a/Engine/Plugins/Experimental/Mover/Source/Mover/Public/DefaultMovementSet/NavMoverComponent.h b/Engine/Plugins/Experimental/Mover/Source/Mover/Public/DefaultMovementSet/NavMoverComponent.h
--- a/Engine/Plugins/Experimental/Mover/Source/Mover/Public/DefaultMovementSet/NavMoverComponent.h
+++ b/Engine/Plugins/Experimental/Mover/Source/Mover/Public/DefaultMovementSet/NavMoverComponent.h
@@ -99,8 +99,10 @@ public:
 	
 	/** Returns location of controlled actor - meaning center of collision bounding box */
 	virtual FVector GetLocation() const override;
-	/** Returns location of controlled actor's "feet" meaning center of bottom of collision bounding box */
+	/** Returns location of controlled actor's "feet": the center bottom of its collision bounding box at its current location */
 	virtual FVector GetFeetLocation() const override;
+	/** Returns location of controlled actor's "feet": the center bottom of its collision bounding box, as if it was at a specific location */
+	virtual FVector GetFeetLocationAt(FVector ComponentLocation) const;
 	/** Returns based location of controlled actor */
 	virtual FBasedPosition GetFeetLocationBased() const override;
 
diff --git a/Engine/Plugins/Experimental/Mover/Source/Mover/Public/MoveLibrary/AirMovementUtils.h b/Engine/Plugins/Experimental/Mover/Source/Mover/Public/MoveLibrary/AirMovementUtils.h
--- a/Engine/Plugins/Experimental/Mover/Source/Mover/Public/MoveLibrary/AirMovementUtils.h
+++ b/Engine/Plugins/Experimental/Mover/Source/Mover/Public/MoveLibrary/AirMovementUtils.h
@@ -76,4 +76,17 @@ public:
     UFUNCTION(BlueprintCallable, Category=Mover)
     static float TryMoveToFallAlongSurface(const FMovingComponentSet& MovingComps, const FVector& Delta, float PctOfDeltaToMove, const FQuat Rotation, const FVector& Normal, FHitResult& Hit, bool bHandleImpact, float FloorSweepDistance, float MaxWalkSlopeCosine, FFloorCheckResult& OutFloorResult, FMovementRecord& MoveRecord);
 
+public:	
+	// Threadsafe movement queries
+
+	/** Tests potential movement of a component falling/sliding along a surface, while checking for landing on a walkable surface. Intended for use while falling. 
+	 * Modifies InOutHit with final movement hit data
+	 * Sets OutFloorResult with any found walkable surface info
+     * Appends to InOutMoveRecord with any movement substeps
+	 * Returns the percent of time applied, with 0.0 meaning no movement would occur. 
+	 */
+	UFUNCTION(BlueprintCallable, Category = Mover)
+	static float TestFallingMoveAlongHitSurface(const FMovingComponentSet& MovingComps, const FVector& OriginalMoveDelta, const FVector& LocationAtHit, const FQuat& TargetRotation, bool bHandleImpact, float FloorSweepDistance, float MaxWalkSlopeCosine, FHitResult& InOutHit, FFloorCheckResult& OutFloorResult, FMovementRecord& InOutMoveRecord);
+
+
 };
diff --git a/Engine/Plugins/Experimental/Mover/Source/Mover/Public/MoveLibrary/GroundMovementUtils.h b/Engine/Plugins/Experimental/Mover/Source/Mover/Public/MoveLibrary/GroundMovementUtils.h
--- a/Engine/Plugins/Experimental/Mover/Source/Mover/Public/MoveLibrary/GroundMovementUtils.h
+++ b/Engine/Plugins/Experimental/Mover/Source/Mover/Public/MoveLibrary/GroundMovementUtils.h
@@ -101,4 +101,33 @@ public:
 
 	UFUNCTION(BlueprintCallable, Category = Mover)
 	static bool CanStepUpOnHitSurface(const FHitResult& Hit);
+
+public:	
+	// Threadsafe movement queries
+
+	/** Tests a potential movement along a walkable surface. 
+	 * Modifies InOutHit with final movement hit data
+     * Appends to InOutMoveRecord with any movement substeps
+	 * Returns the percent of time applied, with 0.0 meaning no movement would occur. 
+	 */
+	UFUNCTION(BlueprintCallable, Category = Mover)
+	static float TestGroundedMoveAlongHitSurface(const FMovingComponentSet& MovingComps, const FVector& OriginalMoveDelta, const FVector& LocationAtHit, const FQuat& TargetRotation, bool bHandleImpact, float MaxStepHeight, float MaxWalkSlopeCosine, FHitResult& InOutHit, FMovementRecord& InOutMoveRecord);
+
+	/** Tests potential movement of a component up/down to adjust to a walkable floor. Intended for use while performing ground movement. 
+	 * Modifies InOutCurrentFloor to account for adjustments, if any was made
+     * Appends to InOutMoveRecord with any movement substeps
+	 * Returns new location, including any adjustment if it was made
+	 */
+	UFUNCTION(BlueprintCallable, Category = Mover)
+	static FVector TestMoveToAdjustToFloor(const FMovingComponentSet& MovingComps, const FVector& Location, const FQuat& Rotation, float MaxWalkSlopeCosine, FFloorCheckResult& InOutCurrentFloor, FMovementRecord& InOutMoveRecord);
+
+	/** Tests potential up-and-over movement of a component to overcome a blocking obstacle hit, such as moving up steps or slope.
+	 * Sets OutFloorTestResult with final floor info (optional)
+	 * Sets OutFinalLocation with final component location
+     * Appends to InOutMoveRecord with any movement substeps
+	 * Returns whether any movement was possible.
+	 */
+	static bool TestMoveToStepOver(const FMovingComponentSet& MovingComps, const FVector& GravDir, float MaxStepHeight, float MaxWalkSlopeCosine, float FloorSweepDistance, const FVector& MoveDelta, const FQuat& Rotation, const FHitResult& MoveHitResult, const FFloorCheckResult& CurrentFloor, bool bIsFalling, FOptionalFloorCheckResult* OutFloorTestResult, FVector& OutFinalLocation, FMovementRecord& InOutMoveRecord);
+
+
 };
