Mover: changed duplicately-defined constant to a single shared definition. Fixes redefinition compiler error and static analysis warning.

#rnx

diff --git a/Engine/Plugins/Experimental/Mover/Source/Mover/Private/DefaultMovementSet/Modes/AsyncFallingMode.cpp b/Engine/Plugins/Experimental/Mover/Source/Mover/Private/DefaultMovementSet/Modes/AsyncFallingMode.cpp
--- a/Engine/Plugins/Experimental/Mover/Source/Mover/Private/DefaultMovementSet/Modes/AsyncFallingMode.cpp
+++ b/Engine/Plugins/Experimental/Mover/Source/Mover/Private/DefaultMovementSet/Modes/AsyncFallingMode.cpp
@@ -32,7 +32,6 @@ UAsyncFallingMode::UAsyncFallingMode(const FObjectInitializer& ObjectInitializer
 	GameplayTags.AddTag(Mover_IsFalling);
 }
 
-constexpr float VERTICAL_SLOPE_NORMAL_MAX_DOT = 0.001f; // Slope is vertical if Abs(UpDirection) <= this threshold. Accounts for precision problems that sometimes angle normals slightly off horizontal for vertical surface.
 
 void UAsyncFallingMode::OnGenerateMove(const FMoverTickStartData& StartState, const FMoverTimeStep& TimeStep, FProposedMove& OutProposedMove) const
 {
@@ -108,7 +107,7 @@ void UAsyncFallingMode::OnGenerateMove(const FMoverTickStartData& StartState, co
 	// limit our moveinput based on the floor we're on
 	if (SimBlackboard && SimBlackboard->TryGet(CommonBlackboard::LastFloorResult, LastFloorResult))
 	{
-		if (LastFloorResult.HitResult.IsValidBlockingHit() && LastFloorResult.HitResult.Normal.Dot(UpDirection) > VERTICAL_SLOPE_NORMAL_MAX_DOT && !LastFloorResult.IsWalkableFloor())
+		if (LastFloorResult.HitResult.IsValidBlockingHit() && LastFloorResult.HitResult.Normal.Dot(UpDirection) > UE::MoverUtils::VERTICAL_SLOPE_NORMAL_MAX_DOT && !LastFloorResult.IsWalkableFloor())
 		{
 			// If acceleration is into the wall, limit contribution.
 			if (FVector::DotProduct(Params.MoveInput, LastFloorResult.HitResult.Normal) < 0.f)
diff --git a/Engine/Plugins/Experimental/Mover/Source/Mover/Private/DefaultMovementSet/Modes/FallingMode.cpp b/Engine/Plugins/Experimental/Mover/Source/Mover/Private/DefaultMovementSet/Modes/FallingMode.cpp
--- a/Engine/Plugins/Experimental/Mover/Source/Mover/Private/DefaultMovementSet/Modes/FallingMode.cpp
+++ b/Engine/Plugins/Experimental/Mover/Source/Mover/Private/DefaultMovementSet/Modes/FallingMode.cpp
@@ -29,7 +29,6 @@ UFallingMode::UFallingMode(const FObjectInitializer& ObjectInitializer)
 	GameplayTags.AddTag(Mover_IsFalling);
 }
 
-constexpr float VERTICAL_SLOPE_NORMAL_MAX_DOT = 0.001f; // Slope is vertical if Abs(UpDirection) <= this threshold. Accounts for precision problems that sometimes angle normals slightly off horizontal for vertical surface.
 
 void UFallingMode::OnGenerateMove(const FMoverTickStartData& StartState, const FMoverTimeStep& TimeStep, FProposedMove& OutProposedMove) const
 {
@@ -100,7 +99,7 @@ void UFallingMode::OnGenerateMove(const FMoverTickStartData& StartState, const F
 	// limit our moveinput based on the floor we're on
 	if (SimBlackboard && SimBlackboard->TryGet(CommonBlackboard::LastFloorResult, LastFloorResult))
 	{
-		if (LastFloorResult.HitResult.IsValidBlockingHit() && LastFloorResult.HitResult.Normal.Dot(UpDirection) > VERTICAL_SLOPE_NORMAL_MAX_DOT && !LastFloorResult.IsWalkableFloor())
+		if (LastFloorResult.HitResult.IsValidBlockingHit() && LastFloorResult.HitResult.Normal.Dot(UpDirection) > UE::MoverUtils::VERTICAL_SLOPE_NORMAL_MAX_DOT && !LastFloorResult.IsWalkableFloor())
 		{
 			// If acceleration is into the wall, limit contribution.
 			if (FVector::DotProduct(Params.MoveInput, LastFloorResult.HitResult.Normal) < 0.f)
diff --git a/Engine/Plugins/Experimental/Mover/Source/Mover/Private/MoveLibrary/MovementUtils.cpp b/Engine/Plugins/Experimental/Mover/Source/Mover/Private/MoveLibrary/MovementUtils.cpp
--- a/Engine/Plugins/Experimental/Mover/Source/Mover/Private/MoveLibrary/MovementUtils.cpp
+++ b/Engine/Plugins/Experimental/Mover/Source/Mover/Private/MoveLibrary/MovementUtils.cpp
@@ -14,6 +14,8 @@
 namespace UE::MoverUtils
 {
 	const double SMALL_MOVE_DISTANCE = 1e-3;
+	const double VERTICAL_SLOPE_NORMAL_MAX_DOT = 0.001f; // Slope is vertical if Abs(UpDirection) <= this threshold. Accounts for precision problems that sometimes angle normals slightly off horizontal for vertical surface.
+
 	const float  VELOCITY_INPUT_NO_ACCELERATION_DIFFERENCE = 1.01f;
 }
 
diff --git a/Engine/Plugins/Experimental/Mover/Source/Mover/Public/MoveLibrary/MovementUtils.h b/Engine/Plugins/Experimental/Mover/Source/Mover/Public/MoveLibrary/MovementUtils.h
--- a/Engine/Plugins/Experimental/Mover/Source/Mover/Public/MoveLibrary/MovementUtils.h
+++ b/Engine/Plugins/Experimental/Mover/Source/Mover/Public/MoveLibrary/MovementUtils.h
@@ -18,6 +18,7 @@ class UMoverComponent;
 namespace UE::MoverUtils
 {
 	extern const double SMALL_MOVE_DISTANCE;
+	extern const double VERTICAL_SLOPE_NORMAL_MAX_DOT;
 }
 
 /** Encapsulates detailed trajectory sample info, from a move that has already occurred or one projected into the future */
