Mover: added game thread-only finalize frame event, and used it as a hook for async navwalking mode to manipulate collision settings safely. This fixes cases where async navwalking actors could get hung up on world geometry.

diff --git a/Engine/Plugins/Experimental/Mover/Source/Mover/Private/DefaultMovementSet/Modes/AsyncNavWalkingMode.cpp b/Engine/Plugins/Experimental/Mover/Source/Mover/Private/DefaultMovementSet/Modes/AsyncNavWalkingMode.cpp
--- a/Engine/Plugins/Experimental/Mover/Source/Mover/Private/DefaultMovementSet/Modes/AsyncNavWalkingMode.cpp
+++ b/Engine/Plugins/Experimental/Mover/Source/Mover/Private/DefaultMovementSet/Modes/AsyncNavWalkingMode.cpp
@@ -337,12 +337,15 @@ void UAsyncNavWalkingMode::SetCollisionForNavWalking(bool bEnable)
 
 void UAsyncNavWalkingMode::OnActivate()
 {
-	// TODO: Need alternative to modifying collision settings
-	//SetCollisionForNavWalking(true);
-
-	if (UMoverBlackboard* SimBlackboard = GetMoverComponent()->GetSimBlackboard_Mutable())
+	if (UMoverComponent* MoverComp = GetMoverComponent())
 	{
-		SimBlackboard->Invalidate(NavWalkingBlackBoard::ProjectedNavMeshHitResult);
+		// TEMP solution for our inability to change collision settings outside of the game thread
+		MoverComp->OnPostFinalize.AddDynamic(this, &UAsyncNavWalkingMode::PostFinalizeSetCollisionForNavWalking);
+
+		if (UMoverBlackboard* SimBlackboard = MoverComp->GetSimBlackboard_Mutable())
+		{
+			SimBlackboard->Invalidate(NavWalkingBlackBoard::ProjectedNavMeshHitResult);
+		}
 	}
 
 	NavMeshProjectionTimer = (NavMeshProjectionInterval > 0.f) ? FMath::FRandRange(-NavMeshProjectionInterval, 0.f) : 0.f;
@@ -350,7 +353,10 @@ void UAsyncNavWalkingMode::OnActivate()
 
 void UAsyncNavWalkingMode::OnDeactivate()
 {
-	// TODO: Need alternative to modifying collision settings
+	if (UMoverComponent* MoverComp = GetMoverComponent())
+	{
+		MoverComp->OnPostFinalize.AddDynamic(this, &UAsyncNavWalkingMode::PostFinalizeRestoreOriginalCollision);
+	}
 }
 
 const INavigationDataInterface* UAsyncNavWalkingMode::GetNavData() const
@@ -617,3 +623,22 @@ FRelativeBaseInfo UAsyncNavWalkingMode::UpdateFloorAndBaseInfo(const FFloorCheck
 }
 
 
+void UAsyncNavWalkingMode::PostFinalizeSetCollisionForNavWalking(const FMoverSyncState& SyncState, const FMoverAuxStateContext& AuxState)
+{
+	SetCollisionForNavWalking(true);
+
+	if (UMoverComponent* MoverComp = GetMoverComponent())
+	{
+		MoverComp->OnPostFinalize.RemoveDynamic(this, &UAsyncNavWalkingMode::PostFinalizeSetCollisionForNavWalking);
+	}
+}
+
+void UAsyncNavWalkingMode::PostFinalizeRestoreOriginalCollision(const FMoverSyncState& SyncState, const FMoverAuxStateContext& AuxState)
+{
+	SetCollisionForNavWalking(false);
+
+	if (UMoverComponent* MoverComp = GetMoverComponent())
+	{
+		MoverComp->OnPostFinalize.RemoveDynamic(this, &UAsyncNavWalkingMode::PostFinalizeRestoreOriginalCollision);
+	}
+}
diff --git a/Engine/Plugins/Experimental/Mover/Source/Mover/Private/MoverComponent.cpp b/Engine/Plugins/Experimental/Mover/Source/Mover/Private/MoverComponent.cpp
--- a/Engine/Plugins/Experimental/Mover/Source/Mover/Private/MoverComponent.cpp
+++ b/Engine/Plugins/Experimental/Mover/Source/Mover/Private/MoverComponent.cpp
@@ -298,6 +298,8 @@ void UMoverComponent::FinalizeFrame(const FMoverSyncState* SyncState, const FMov
 		CachedLastSimTickTimeStep.ServerFrame = BackendLiaisonComp->GetCurrentSimFrame();
 		bHasValidCachedState = true;
 	}
+
+	OnPostFinalize.Broadcast(*SyncState, *AuxState);
 }
 
 void UMoverComponent::FinalizeSmoothingFrame(const FMoverSyncState* SyncState, const FMoverAuxStateContext* AuxState)
diff --git a/Engine/Plugins/Experimental/Mover/Source/Mover/Public/DefaultMovementSet/Modes/AsyncNavWalkingMode.h b/Engine/Plugins/Experimental/Mover/Source/Mover/Public/DefaultMovementSet/Modes/AsyncNavWalkingMode.h
--- a/Engine/Plugins/Experimental/Mover/Source/Mover/Public/DefaultMovementSet/Modes/AsyncNavWalkingMode.h
+++ b/Engine/Plugins/Experimental/Mover/Source/Mover/Public/DefaultMovementSet/Modes/AsyncNavWalkingMode.h
@@ -132,4 +132,12 @@ protected:
  	// Collision responses saved at the activation of this mode, to be restored as upon deactivation if the CDO can't be used
  	ECollisionResponse CollideVsWorldStatic = ECollisionResponse::ECR_Block;
  	ECollisionResponse CollideVsWorldDynamic = ECollisionResponse::ECR_Block;
+
+private:
+	// Game thread event callbacks used for manipulating collision on the actor when entering/exiting this mode
+	UFUNCTION()
+	void PostFinalizeSetCollisionForNavWalking(const FMoverSyncState& SyncState, const FMoverAuxStateContext& AuxState);
+	UFUNCTION()
+	void PostFinalizeRestoreOriginalCollision(const FMoverSyncState& SyncState, const FMoverAuxStateContext& AuxState);
+	
 };
diff --git a/Engine/Plugins/Experimental/Mover/Source/Mover/Public/MoverComponent.h b/Engine/Plugins/Experimental/Mover/Source/Mover/Public/MoverComponent.h
--- a/Engine/Plugins/Experimental/Mover/Source/Mover/Public/MoverComponent.h
+++ b/Engine/Plugins/Experimental/Mover/Source/Mover/Public/MoverComponent.h
@@ -49,6 +49,9 @@ DECLARE_DYNAMIC_MULTICAST_DELEGATE_TwoParams(FMover_OnPostSimRollback, const FMo
 // Fired after changing movement modes. First param is the name of the previous movement mode. Second is the name of the new movement mode. 
 DECLARE_DYNAMIC_MULTICAST_DELEGATE_TwoParams(FMover_OnMovementModeChanged, const FName&, PreviousMovementModeName, const FName&, NewMovementModeName);
 
+// Fired after a frame has been finalized. This may be a resimulation or not. No changes to state are possible. Guaranteed to be on the game thread.
+DECLARE_DYNAMIC_MULTICAST_DELEGATE_TwoParams(FMover_OnPostFinalize, const FMoverSyncState&, SyncState, const FMoverAuxStateContext&, AuxState);
+
 // Fired after proposed movement has been generated (i.e. after movement modes and layered moves have generated movement and mixed together).
 DECLARE_DYNAMIC_DELEGATE_ThreeParams(FMover_ProcessGeneratedMovement, const FMoverTickStartData&, StartState, const FMoverTimeStep&, TimeStep, FProposedMove&, OutProposedMove);
 
@@ -94,6 +97,10 @@ public:
 	UPROPERTY(BlueprintAssignable, Category = Mover)
 	FMover_OnMovementModeChanged OnMovementModeChanged;
 
+	// Broadcast after each finalized simulation frame, after the state is finalized. (Game thread only)
+	UPROPERTY(BlueprintAssignable, Category = Mover)
+	FMover_OnPostFinalize OnPostFinalize;
+
 	/**
 	 * Broadcast after proposed movement has been generated. After movement modes and layered moves have generated movement and mixed together.
 	 * This allows for final modifications to proposed movement before it's executed.
