Mover teleport instant effect now supports setting rotation in addition to location

diff --git a/Engine/Plugins/Experimental/Mover/Source/Mover/Private/DefaultMovementSet/InstantMovementEffects/BasicInstantMovementEffects.cpp b/Engine/Plugins/Experimental/Mover/Source/Mover/Private/DefaultMovementSet/InstantMovementEffects/BasicInstantMovementEffects.cpp
--- a/Engine/Plugins/Experimental/Mover/Source/Mover/Private/DefaultMovementSet/InstantMovementEffects/BasicInstantMovementEffects.cpp
+++ b/Engine/Plugins/Experimental/Mover/Source/Mover/Private/DefaultMovementSet/InstantMovementEffects/BasicInstantMovementEffects.cpp
@@ -17,12 +17,16 @@
 
 FTeleportEffect::FTeleportEffect()
 	: TargetLocation(FVector::ZeroVector)
+	, bUseActorRotation(true)
+	, TargetRotation(FRotator::ZeroRotator)
 {
 }
 
 bool FTeleportEffect::ApplyMovementEffect(FApplyMovementEffectParams& ApplyEffectParams, FMoverSyncState& OutputState)
 {
-	if (ApplyEffectParams.UpdatedComponent->GetOwner()->TeleportTo(TargetLocation, ApplyEffectParams.UpdatedComponent->GetComponentRotation()))
+	const FRotator Rotation = bUseActorRotation ? ApplyEffectParams.UpdatedComponent->GetComponentRotation() : TargetRotation;
+	
+	if (ApplyEffectParams.UpdatedComponent->GetOwner()->TeleportTo(TargetLocation, Rotation))
 	{
 		FMoverDefaultSyncState& OutputSyncState = OutputState.SyncStateCollection.FindOrAddMutableDataByType<FMoverDefaultSyncState>();
 		
@@ -58,6 +62,12 @@ void FTeleportEffect::NetSerialize(FArchive& Ar)
 	Super::NetSerialize(Ar);
 
 	Ar << TargetLocation;
+	
+	Ar.SerializeBits(&bUseActorRotation, 1);
+	if (!bUseActorRotation)
+	{
+		Ar << TargetRotation;
+	}
 }
 
 UScriptStruct* FTeleportEffect::GetScriptStruct() const
diff --git a/Engine/Plugins/Experimental/Mover/Source/Mover/Public/DefaultMovementSet/InstantMovementEffects/BasicInstantMovementEffects.h b/Engine/Plugins/Experimental/Mover/Source/Mover/Public/DefaultMovementSet/InstantMovementEffects/BasicInstantMovementEffects.h
--- a/Engine/Plugins/Experimental/Mover/Source/Mover/Public/DefaultMovementSet/InstantMovementEffects/BasicInstantMovementEffects.h
+++ b/Engine/Plugins/Experimental/Mover/Source/Mover/Public/DefaultMovementSet/InstantMovementEffects/BasicInstantMovementEffects.h
@@ -6,7 +6,7 @@
 #include "InstantMovementEffect.h"
 #include "BasicInstantMovementEffects.generated.h"
 
-/** Teleport: instantly moves an actor to a new location */
+/** Teleport: instantly moves an actor to a new location and rotation */
 USTRUCT(BlueprintType)
 struct MOVER_API FTeleportEffect : public FInstantMovementEffect
 {
@@ -19,6 +19,14 @@ struct MOVER_API FTeleportEffect : public FInstantMovementEffect
 	UPROPERTY(EditAnywhere, BlueprintReadWrite, Category = Mover)
 	FVector TargetLocation;
 
+	// Whether this teleport effect should keep the actor's current rotation or use a specified one (TargetRotation)
+	UPROPERTY(EditAnywhere, BlueprintReadWrite, Category = Mover)
+	bool bUseActorRotation;
+	
+	// Actor rotation is set to this value on teleport if bUseActorRotation is false
+	UPROPERTY(EditAnywhere, BlueprintReadWrite, Category = Mover, meta=(EditCondition="bUseActorRotation==false", EditConditionHides))
+	FRotator TargetRotation;
+	
 	virtual bool ApplyMovementEffect(FApplyMovementEffectParams& ApplyEffectParams, FMoverSyncState& OutputState) override;
 
 	virtual FInstantMovementEffect* Clone() const override;
