Mover fixing NetSerialize issue with JumpTo layered move where jump rotation was only synced when bUseActorRotation was true

diff --git a/Engine/Plugins/Experimental/Mover/Source/Mover/Private/DefaultMovementSet/LayeredMoves/BasicLayeredMoves.cpp b/Engine/Plugins/Experimental/Mover/Source/Mover/Private/DefaultMovementSet/LayeredMoves/BasicLayeredMoves.cpp
--- a/Engine/Plugins/Experimental/Mover/Source/Mover/Private/DefaultMovementSet/LayeredMoves/BasicLayeredMoves.cpp
+++ b/Engine/Plugins/Experimental/Mover/Source/Mover/Private/DefaultMovementSet/LayeredMoves/BasicLayeredMoves.cpp
@@ -229,9 +229,9 @@ void FLayeredMove_JumpTo::NetSerialize(FArchive& Ar)
 
 	Ar << JumpDistance;
 	Ar << JumpHeight;
-	Ar << bUseActorRotation;
+	Ar.SerializeBits(&bUseActorRotation, 1);
 
-	if (bUseActorRotation)
+	if (!bUseActorRotation)
 	{
 		Ar << JumpRotation;
 	}
