Mover: adding option to standalone backend to individually choose whether to run movement simulation on worker threads, and make all instances refresh dynamically when the global cvar changes.

diff --git a/Engine/Plugins/Experimental/Mover/Source/Mover/Private/Backends/MoverStandaloneLiaison.cpp b/Engine/Plugins/Experimental/Mover/Source/Mover/Private/Backends/MoverStandaloneLiaison.cpp
--- a/Engine/Plugins/Experimental/Mover/Source/Mover/Private/Backends/MoverStandaloneLiaison.cpp
+++ b/Engine/Plugins/Experimental/Mover/Source/Mover/Private/Backends/MoverStandaloneLiaison.cpp
@@ -8,13 +8,25 @@
 
 namespace MoverStandaloneBackendCVars
 { 
+	static void OnRunMovementSimOnAnyThreadChanged(IConsoleVariable* Var)
+	{
+		// Force a refresh of all liaisons when the global 
+		ForEachObjectOfClass(UMoverStandaloneLiaisonComponent::StaticClass(), [&](UObject* AsObj)
+		{
+			if (UMoverStandaloneLiaisonComponent* StandaloneLiaison = CastChecked<UMoverStandaloneLiaisonComponent>(AsObj, ECastCheckedType::NullAllowed))
+			{
+				StandaloneLiaison->SetUseAsyncMovementSimulationTick(StandaloneLiaison->GetUseAsyncMovementSimulationTick());
+			}
+		});
+	}
+
 	// Whether to allow movement simulation ticks on any thread
 	static int32 RunMovementSimOnAnyThread = 0;
 	FAutoConsoleVariableRef CVarMoverStandaloneRunSimOnAnyThread(
 		TEXT("mover.standalone.RunMovementSimOnAnyThread"),
 		RunMovementSimOnAnyThread,
-		TEXT("Whether to allow Mover simulation ticks to run on any thread. Requires use of threadsafe movement modes that do not modify scene components.\n")
-		TEXT("0: Game thread only, 1: Any thread"),
+		TEXT("Whether to allow Mover simulation ticks to run on any thread. Requires use of threadsafe movement modes that do not modify scene components.\nIndividuals must also have their UseAsyncMovementSimulationTick option enabled.\n0: Game thread only, 1: Any thread"),
+		FConsoleVariableDelegate::CreateStatic(&OnRunMovementSimOnAnyThreadChanged),
 		ECVF_Default);
 }
 
@@ -34,7 +46,7 @@ UMoverStandaloneLiaisonComponent::UMoverStandaloneLiaisonComponent(const FObject
 	SimulateMovementTickFunction.bStartWithTickEnabled = true;
 	SimulateMovementTickFunction.SetTickFunctionEnable(true);
 	SimulateMovementTickFunction.TickGroup = TG_PrePhysics;
-	SimulateMovementTickFunction.bRunOnAnyThread = MoverStandaloneBackendCVars::RunMovementSimOnAnyThread;
+	SimulateMovementTickFunction.bRunOnAnyThread = bUseAsyncMovementSimulationTick && MoverStandaloneBackendCVars::RunMovementSimOnAnyThread;
 
 	ApplyStateTickFunction.bCanEverTick = true;
 	ApplyStateTickFunction.bStartWithTickEnabled = true;
@@ -131,6 +143,19 @@ FTickFunction* UMoverStandaloneLiaisonComponent::FindTickFunction(EMoverTickPhas
 	return FoundFunction;
 }
 
+
+void UMoverStandaloneLiaisonComponent::SetUseAsyncMovementSimulationTick(bool bUseAsyncMovementSim)
+{
+	bUseAsyncMovementSimulationTick = bUseAsyncMovementSim;
+	SimulateMovementTickFunction.bRunOnAnyThread = bUseAsyncMovementSimulationTick && MoverStandaloneBackendCVars::RunMovementSimOnAnyThread;
+}
+
+bool UMoverStandaloneLiaisonComponent::GetUseAsyncMovementSimulationTick() const
+{
+	return bUseAsyncMovementSimulationTick;
+}
+
+
 void UMoverStandaloneLiaisonComponent::RegisterComponentTickFunctions(bool bRegister)
 {
 	Super::RegisterComponentTickFunctions(bRegister);
diff --git a/Engine/Plugins/Experimental/Mover/Source/Mover/Public/Backends/MoverStandaloneLiaison.h b/Engine/Plugins/Experimental/Mover/Source/Mover/Public/Backends/MoverStandaloneLiaison.h
--- a/Engine/Plugins/Experimental/Mover/Source/Mover/Public/Backends/MoverStandaloneLiaison.h
+++ b/Engine/Plugins/Experimental/Mover/Source/Mover/Public/Backends/MoverStandaloneLiaison.h
@@ -114,6 +114,13 @@ public:
 
 	FTickFunction* FindTickFunction(EMoverTickPhase MoverTickPhase);
 
+	/** Sets whether this instance's movement simulation tick can run on worker threads or not. See @bUseAsyncMovementSimulationTick */
+	UFUNCTION(BlueprintCallable, Category = "Async Movement")
+	void SetUseAsyncMovementSimulationTick(bool bUseAsyncMovementSim);
+
+	UFUNCTION(BlueprintCallable, Category = "Async Movement")
+	bool GetUseAsyncMovementSimulationTick() const;
+
 protected:
 	void TickInputProduction(float DeltaSeconds);
 	void TickMovementSimulation(float DeltaSeconds);
@@ -124,6 +131,13 @@ protected:
 	virtual void OnControllerChanged(APawn* Pawn, AController* OldController, AController* NewController);
 
 protected:
+	/**
+	 * Sets whether the movement simulation tick can run on worker threads or not, also gated by global option.
+	 * Changes at runtime will take affect next frame. Has no effect on input production or applying results.
+	 */
+	UPROPERTY(EditDefaultsOnly, Category = "Async Movement")
+	bool bUseAsyncMovementSimulationTick = false;
+
 	TObjectPtr<UMoverComponent> MoverComp;	// the component that we're in charge of driving
 
 	double CurrentSimTimeMs;
