Mover: fix async walking bug where a character would not turn while sliding along a vertical surface that was also marked as walkable. This was also causing motion-matching characters to get hung up on geometry in some cases.

diff --git a/Engine/Plugins/Experimental/Mover/Source/Mover/Private/DefaultMovementSet/Modes/AsyncWalkingMode.cpp b/Engine/Plugins/Experimental/Mover/Source/Mover/Private/DefaultMovementSet/Modes/AsyncWalkingMode.cpp
--- a/Engine/Plugins/Experimental/Mover/Source/Mover/Private/DefaultMovementSet/Modes/AsyncWalkingMode.cpp
+++ b/Engine/Plugins/Experimental/Mover/Source/Mover/Private/DefaultMovementSet/Modes/AsyncWalkingMode.cpp
@@ -242,6 +242,8 @@ void UAsyncWalkingMode::OnSimulationTick(const FSimulationTickParams& Params, FM
 					if (UGroundMovementUtils::TestMoveToStepOver(Params.MovingComps, DownwardDir, CommonLegacySettings->MaxStepHeight, CommonLegacySettings->MaxWalkSlopeCosine, CommonLegacySettings->FloorSweepDistance, OrigMoveDelta * (1.f - PercentTimeAppliedSoFar), RotationInProgress, MoveHitResult, CurrentFloor, false, &StepUpFloorResult, PostStepUpLocation, MoveRecord))
 					{
 						LocationInProgress = PostStepUpLocation;
+						RotationInProgress = TargetRotation;
+						PercentTimeAppliedSoFar = 1.0f;
 					}
 					else
 					{
@@ -250,6 +252,9 @@ void UAsyncWalkingMode::OnSimulationTick(const FSimulationTickParams& Params, FM
 						float PercentAvailableToSlide = 1.f - PercentTimeAppliedSoFar;
 
 						float SlideAmount = UGroundMovementUtils::TestGroundedMoveAlongHitSurface(Params.MovingComps, OrigMoveDelta, LocationInProgress, TargetRotation, /*bHandleImpact=*/true, CommonLegacySettings->MaxStepHeight, CommonLegacySettings->MaxWalkSlopeCosine, IN OUT MoveHitResult, IN OUT MoveRecord);
+
+						LocationInProgress = LocationInProgress + ((MoveHitResult.TraceEnd - MoveHitResult.TraceStart) * SlideAmount);
+						RotationInProgress = FQuat::Slerp(RotationInProgress, TargetRotation, SlideAmount);
 						PercentTimeAppliedSoFar += PercentAvailableToSlide * SlideAmount;
 					}
 				}
