[Mover] [BugFix] SimBlackboard validity check in IsSlopeSliding()

diff --git a/Engine/Plugins/Experimental/Mover/Source/Mover/Private/DefaultMovementSet/CharacterMoverComponent.cpp b/Engine/Plugins/Experimental/Mover/Source/Mover/Private/DefaultMovementSet/CharacterMoverComponent.cpp
--- a/Engine/Plugins/Experimental/Mover/Source/Mover/Private/DefaultMovementSet/CharacterMoverComponent.cpp
+++ b/Engine/Plugins/Experimental/Mover/Source/Mover/Private/DefaultMovementSet/CharacterMoverComponent.cpp
@@ -65,7 +65,7 @@ bool UCharacterMoverComponent::IsSlopeSliding() const
 	{
 		FFloorCheckResult HitResult;
 		const UMoverBlackboard* MoverBlackboard = GetSimBlackboard();
-		if (MoverBlackboard->TryGet(CommonBlackboard::LastFloorResult, HitResult))
+		if (MoverBlackboard && MoverBlackboard->TryGet(CommonBlackboard::LastFloorResult, HitResult))
 		{
 			return HitResult.bBlockingHit && !HitResult.bWalkableFloor;
 		}
