Mover: adding async-compatible teleportation effect

diff --git a/Engine/Plugins/Experimental/Mover/Source/Mover/Private/DefaultMovementSet/InstantMovementEffects/BasicInstantMovementEffects.cpp b/Engine/Plugins/Experimental/Mover/Source/Mover/Private/DefaultMovementSet/InstantMovementEffects/BasicInstantMovementEffects.cpp
--- a/Engine/Plugins/Experimental/Mover/Source/Mover/Private/DefaultMovementSet/InstantMovementEffects/BasicInstantMovementEffects.cpp
+++ b/Engine/Plugins/Experimental/Mover/Source/Mover/Private/DefaultMovementSet/InstantMovementEffects/BasicInstantMovementEffects.cpp
@@ -3,6 +3,7 @@
 
 #include "DefaultMovementSet/InstantMovementEffects/BasicInstantMovementEffects.h"
 
+#include "Engine/World.h"
 #include "MoverComponent.h"
 #include "MoverDataModelTypes.h"
 #include "MoverSimulationTypes.h"
@@ -85,6 +86,60 @@ void FTeleportEffect::AddReferencedObjects(FReferenceCollector& Collector)
 	Super::AddReferencedObjects(Collector);
 }
 
+
+// -------------------------------------------------------------------
+// FAsyncTeleportEffect
+// -------------------------------------------------------------------
+
+bool FAsyncTeleportEffect::ApplyMovementEffect(FApplyMovementEffectParams& ApplyEffectParams, FMoverSyncState& OutputState)
+{
+	FVector TeleportLocation = TargetLocation;
+	const FRotator TeleportRotation = bUseActorRotation ? ApplyEffectParams.UpdatedComponent->GetComponentRotation() : TargetRotation;
+
+	if (UWorld* World = ApplyEffectParams.UpdatedComponent->GetWorld())
+	{
+		World->FindTeleportSpot(ApplyEffectParams.UpdatedComponent->GetOwner(), OUT TeleportLocation, TeleportRotation);
+	}
+
+	FMoverDefaultSyncState& OutputSyncState = OutputState.SyncStateCollection.FindOrAddMutableDataByType<FMoverDefaultSyncState>();
+
+	if (const FMoverDefaultSyncState* StartingSyncState = ApplyEffectParams.StartState->SyncState.SyncStateCollection.FindDataByType<FMoverDefaultSyncState>())
+	{
+		OutputSyncState.SetTransforms_WorldSpace(TeleportLocation,
+			TeleportRotation,
+			OutputSyncState.GetVelocity_WorldSpace(),
+			nullptr); // no movement base
+
+		// TODO: instead of invalidating it, consider checking for a floor. Possibly a dynamic base?
+		if (UMoverBlackboard* SimBlackboard = ApplyEffectParams.MoverComp->GetSimBlackboard_Mutable())
+		{
+			SimBlackboard->Invalidate(CommonBlackboard::LastFloorResult);
+			SimBlackboard->Invalidate(CommonBlackboard::LastFoundDynamicMovementBase);
+		}
+
+		return true;
+	}
+
+	return false;
+}
+
+FInstantMovementEffect* FAsyncTeleportEffect::Clone() const
+{
+	FAsyncTeleportEffect* CopyPtr = new FAsyncTeleportEffect(*this);
+	return CopyPtr;
+}
+
+UScriptStruct* FAsyncTeleportEffect::GetScriptStruct() const
+{
+	return FAsyncTeleportEffect::StaticStruct();
+}
+
+FString FAsyncTeleportEffect::ToSimpleString() const
+{
+	return FString::Printf(TEXT("Async Teleport"));
+}
+
+
 // -------------------------------------------------------------------
 // FJumpImpulseEffect
 // -------------------------------------------------------------------
diff --git a/Engine/Plugins/Experimental/Mover/Source/Mover/Public/DefaultMovementSet/InstantMovementEffects/BasicInstantMovementEffects.h b/Engine/Plugins/Experimental/Mover/Source/Mover/Public/DefaultMovementSet/InstantMovementEffects/BasicInstantMovementEffects.h
--- a/Engine/Plugins/Experimental/Mover/Source/Mover/Public/DefaultMovementSet/InstantMovementEffects/BasicInstantMovementEffects.h
+++ b/Engine/Plugins/Experimental/Mover/Source/Mover/Public/DefaultMovementSet/InstantMovementEffects/BasicInstantMovementEffects.h
@@ -6,11 +6,13 @@
 #include "InstantMovementEffect.h"
 #include "BasicInstantMovementEffects.generated.h"
 
-/** Teleport: instantly moves an actor to a new location and rotation */
+/** Teleport: instantly moves an actor to a new location and rotation
+ * For async-compatible teleportation, use AsyncTeleportEffect instead.
+ */
 USTRUCT(BlueprintType)
 struct MOVER_API FTeleportEffect : public FInstantMovementEffect
 {
-	GENERATED_USTRUCT_BODY()
+	GENERATED_BODY()
 
 	FTeleportEffect();
 	virtual ~FTeleportEffect() {}
@@ -40,13 +42,28 @@ struct MOVER_API FTeleportEffect : public FInstantMovementEffect
 	virtual void AddReferencedObjects(class FReferenceCollector& Collector) override;
 };
 
+/** Async Teleport: instantly moves an actor to a new location and rotation (compatible with async movement simulation) */
+USTRUCT(BlueprintType)
+struct MOVER_API FAsyncTeleportEffect : public FTeleportEffect
+{
+	GENERATED_BODY()
+
+	virtual bool ApplyMovementEffect(FApplyMovementEffectParams& ApplyEffectParams, FMoverSyncState& OutputState) override;
+
+	virtual FInstantMovementEffect* Clone() const override;
+
+	virtual UScriptStruct* GetScriptStruct() const override;
+
+	virtual FString ToSimpleString() const override;
+};
+
 /** Jump Impulse: introduces an instantaneous upwards change in velocity. This overrides the existing 'up' component of the actor's current velocity
   * Note: this only applies the impulse for one tick!
   */
 USTRUCT(BlueprintType)
 struct MOVER_API FJumpImpulseEffect : public FInstantMovementEffect
 {
-	GENERATED_USTRUCT_BODY()
+	GENERATED_BODY()
 
 	FJumpImpulseEffect();
 
@@ -75,7 +92,7 @@ struct MOVER_API FJumpImpulseEffect : public FInstantMovementEffect
 USTRUCT(BlueprintType)
 struct MOVER_API FApplyVelocityEffect : public FInstantMovementEffect
 {
-	GENERATED_USTRUCT_BODY()
+	GENERATED_BODY()
 
 	FApplyVelocityEffect();
 	virtual ~FApplyVelocityEffect() {}
