Mover: avoid using scoped movement update unless we're on the game thread.

diff --git a/Engine/Plugins/Experimental/Mover/Source/Mover/Private/MoverComponent.cpp b/Engine/Plugins/Experimental/Mover/Source/Mover/Private/MoverComponent.cpp
--- a/Engine/Plugins/Experimental/Mover/Source/Mover/Private/MoverComponent.cpp
+++ b/Engine/Plugins/Experimental/Mover/Source/Mover/Private/MoverComponent.cpp
@@ -491,10 +491,16 @@ void UMoverComponent::SimulationTick(const FMoverTimeStep& InTimeStep, const FMo
 	OnPreMovement.Broadcast(MoverTimeStep, SimInput.InputCmd, SimInput.SyncState, SimInput.AuxState);
 
 	// Tick the actual simulation. This is where the proposed moves are queried and executed, affecting change to the moving actor's gameplay state and captured in the output sim state
+	if (IsInGameThread())
 	{
+		// If we're on the game thread, we can make use of a scoped movement update for better perf of multi-step movements.  If not, then we're definitely not moving the component in immediate mode so the scope would have no effect.
 		FScopedMovementUpdate ScopedMovementUpdate(UpdatedComponent, EScopedUpdate::DeferredUpdates);
 		ModeFSM->OnSimulationTick(UpdatedComponent, UpdatedCompAsPrimitive, SimBlackboard.Get(), SimInput, MoverTimeStep, SimOutput);
 	}
+	else
+	{
+		ModeFSM->OnSimulationTick(UpdatedComponent, UpdatedCompAsPrimitive, SimBlackboard.Get(), SimInput, MoverTimeStep, SimOutput);
+	}
 
 	if (FMoverDefaultSyncState* OutputSyncState = SimOutput.SyncState.SyncStateCollection.FindMutableDataByType<FMoverDefaultSyncState>())
 	{
