Mover: Moving bShouldNextModeReenter from UImmediateMovementModeTransition to UBaseMovementModeTransition so we can check it and avoid reentering a mode as the default behavior. You can still reenter it if desired, just set bShouldNextModeReenter to true.

diff --git a/Engine/Plugins/Experimental/Mover/Source/Mover/Private/MovementModeStateMachine.cpp b/Engine/Plugins/Experimental/Mover/Source/Mover/Private/MovementModeStateMachine.cpp
--- a/Engine/Plugins/Experimental/Mover/Source/Mover/Private/MovementModeStateMachine.cpp
+++ b/Engine/Plugins/Experimental/Mover/Source/Mover/Private/MovementModeStateMachine.cpp
@@ -305,8 +305,11 @@ void UMovementModeStateMachine::OnSimulationTick(USceneComponent* UpdatedCompone
 
 						if (!EvalResult.NextMode.IsNone())
 						{
-							TransitionToTrigger = Transition;
-							break;
+							if (EvalResult.NextMode != CurrentModeName || Transition->bAllowModeReentry)
+							{
+								TransitionToTrigger = Transition;
+								break;
+							}
 						}
 					}
 				}
@@ -321,8 +324,11 @@ void UMovementModeStateMachine::OnSimulationTick(USceneComponent* UpdatedCompone
 
 							if (!EvalResult.NextMode.IsNone())
 							{
-								TransitionToTrigger = Transition;
-								break;
+								if (EvalResult.NextMode != CurrentModeName || Transition->bAllowModeReentry)
+								{
+									TransitionToTrigger = Transition;
+									break;
+								}
 							}
 						}
 					}
@@ -654,7 +660,7 @@ FTransitionEvalResult UImmediateMovementModeTransition::OnEvaluate(const FSimula
 {
 	if (NextMode != NAME_None)
 	{
-		if (bShouldNextModeReenter)
+		if (bAllowModeReentry)
 		{
 			return FTransitionEvalResult(NextMode);
 		}
@@ -675,11 +681,11 @@ void UImmediateMovementModeTransition::OnTrigger(const FSimulationTickParams& Pa
 void UImmediateMovementModeTransition::SetNextMode(FName DesiredModeName, bool bShouldReenter)
 {
 	NextMode = DesiredModeName;
-	bShouldNextModeReenter = bShouldReenter;
+	bAllowModeReentry = bShouldReenter;
 }
 
 void UImmediateMovementModeTransition::Clear()
 {
 	NextMode = NAME_None;
-	bShouldNextModeReenter = false;
+	bAllowModeReentry = false;
 }
diff --git a/Engine/Plugins/Experimental/Mover/Source/Mover/Private/MovementModeTransition.cpp b/Engine/Plugins/Experimental/Mover/Source/Mover/Private/MovementModeTransition.cpp
--- a/Engine/Plugins/Experimental/Mover/Source/Mover/Private/MovementModeTransition.cpp
+++ b/Engine/Plugins/Experimental/Mover/Source/Mover/Private/MovementModeTransition.cpp
@@ -13,6 +13,7 @@ const FTransitionEvalResult FTransitionEvalResult::NoTransition = FTransitionEva
 
 UBaseMovementModeTransition::UBaseMovementModeTransition(const FObjectInitializer& ObjectInitializer)
 	: Super(ObjectInitializer)
+	, bAllowModeReentry(false)
 {
 }
 
diff --git a/Engine/Plugins/Experimental/Mover/Source/Mover/Public/MovementModeStateMachine.h b/Engine/Plugins/Experimental/Mover/Source/Mover/Public/MovementModeStateMachine.h
--- a/Engine/Plugins/Experimental/Mover/Source/Mover/Public/MovementModeStateMachine.h
+++ b/Engine/Plugins/Experimental/Mover/Source/Mover/Public/MovementModeStateMachine.h
@@ -132,9 +132,8 @@ private:
 
 
 	 FName GetNextModeName() const { return NextMode; }
-	 bool ShouldReenter() const { return bShouldNextModeReenter; }
+	 bool ShouldReenter() const { return bAllowModeReentry; }
 
  private:
 	 FName NextMode = NAME_None;
-	 bool bShouldNextModeReenter = false;
  };
diff --git a/Engine/Plugins/Experimental/Mover/Source/Mover/Public/MovementModeTransition.h b/Engine/Plugins/Experimental/Mover/Source/Mover/Public/MovementModeTransition.h
--- a/Engine/Plugins/Experimental/Mover/Source/Mover/Public/MovementModeTransition.h
+++ b/Engine/Plugins/Experimental/Mover/Source/Mover/Public/MovementModeTransition.h
@@ -64,6 +64,10 @@ public:
 	UPROPERTY(EditDefaultsOnly, Category = Mover, meta = (MustImplement = "/Script/Mover.MovementSettingsInterface"))
 	TArray<TSubclassOf<UObject>> SharedSettingsClasses;
 
+	/** Whether this transition should reenter a mode if it evaluates true and wants to transition into a mode the actor is already in */
+	UPROPERTY(EditDefaultsOnly, BlueprintReadWrite, Category = Mover)
+	bool bAllowModeReentry;
+	
 protected:
 	virtual void OnRegistered();
 
