Mover fixing transition issues if no transition was used but evaluated.

#rnx

diff --git a/Engine/Plugins/Experimental/Mover/Source/Mover/Private/MovementModeStateMachine.cpp b/Engine/Plugins/Experimental/Mover/Source/Mover/Private/MovementModeStateMachine.cpp
--- a/Engine/Plugins/Experimental/Mover/Source/Mover/Private/MovementModeStateMachine.cpp
+++ b/Engine/Plugins/Experimental/Mover/Source/Mover/Private/MovementModeStateMachine.cpp
@@ -314,7 +314,7 @@ void UMovementModeStateMachine::OnSimulationTick(USceneComponent* UpdatedCompone
 					}
 				}
 
-				if (EvalResult.NextMode.IsNone())
+				if (TransitionToTrigger == nullptr)
 				{
 					for (UBaseMovementModeTransition* Transition : GlobalTransitions)
 					{
@@ -334,7 +334,7 @@ void UMovementModeStateMachine::OnSimulationTick(USceneComponent* UpdatedCompone
 					}
 				}
 
-				if (!EvalResult.NextMode.IsNone())
+				if (TransitionToTrigger && !EvalResult.NextMode.IsNone())
 				{
 					OutputState.MovementEndState.NextModeName = EvalResult.NextMode;
 					OutputState.MovementEndState.RemainingMs = SimTickParams.TimeStep.StepMs; 	// Pass all remaining time to next mode
