Physics Mover - Fix base space for movement intent

diff --git a/Engine/Plugins/Experimental/Mover/Source/Mover/Private/Backends/MoverNetworkPhysicsLiaison.cpp b/Engine/Plugins/Experimental/Mover/Source/Mover/Private/Backends/MoverNetworkPhysicsLiaison.cpp
--- a/Engine/Plugins/Experimental/Mover/Source/Mover/Private/Backends/MoverNetworkPhysicsLiaison.cpp
+++ b/Engine/Plugins/Experimental/Mover/Source/Mover/Private/Backends/MoverNetworkPhysicsLiaison.cpp
@@ -814,6 +814,10 @@ void UMoverNetworkPhysicsLiaisonComponent::PostPhysicsUpdate_External()
 				if (Chaos::FPBDRigidParticle* Particle = Interface.GetRigidParticle(PhysicsObject))
 				{
 					DefaultSyncState.SetTransforms_WorldSpace(Particle->X(), FRotator(Particle->R()), Particle->V(), DefaultSyncState.GetMovementBase(), DefaultSyncState.GetMovementBaseBoneName());
+					
+					// The movement intent is currently in world space. Need to convert to base space
+					DefaultSyncState.MoveDirectionIntent = DefaultSyncState.GetCapturedMovementBaseQuat().UnrotateVector(DefaultSyncState.MoveDirectionIntent);
+
 				}
 			}
 
