Physics Mover - Adjustment for walking to falling transition to prevent character entering falling state when going over uneven terrains and slope changes

diff --git a/Engine/Plugins/Experimental/Mover/Source/Mover/Private/PhysicsMover/Modes/PhysicsDrivenWalkingMode.cpp b/Engine/Plugins/Experimental/Mover/Source/Mover/Private/PhysicsMover/Modes/PhysicsDrivenWalkingMode.cpp
--- a/Engine/Plugins/Experimental/Mover/Source/Mover/Private/PhysicsMover/Modes/PhysicsDrivenWalkingMode.cpp
+++ b/Engine/Plugins/Experimental/Mover/Source/Mover/Private/PhysicsMover/Modes/PhysicsDrivenWalkingMode.cpp
@@ -305,7 +305,7 @@ void UPhysicsDrivenWalkingMode::OnSimulationTick(const FSimulationTickParams& Pa
 	SimBlackboard->Set(CommonBlackboard::LastWaterResult, WaterResult);
 
 	const bool bStartSwimming = WaterResult.WaterSplineData.ImmersionDepth > CommonLegacySettings->SwimmingStartImmersionDepth;
-	
+
 	if (WaterResult.IsSwimmableVolume() && bStartSwimming)
 	{
 		SwitchToState(CommonLegacySettings->SwimmingMovementModeName, Params, OutputState);
@@ -346,10 +346,10 @@ void UPhysicsDrivenWalkingMode::OnSimulationTick(const FSimulationTickParams& Pa
 		}
 		const FVector ProjectedVelocity = TargetVelocity + CharacterGravity * FVector::UpVector * DeltaSeconds;
 		const FVector ProjectedGroundVelocity = UPhysicsMovementUtils::ComputeIntegratedGroundVelocityFromHitResult(StartingSyncState->GetLocation_WorldSpace(), FloorResult.HitResult, DeltaSeconds);
-
+		const bool bIsGroundMoving = FMath::Abs(ProjectedGroundVelocity.Z) > UE_KINDA_SMALL_NUMBER;
 		const float ProjectedRelativeVerticalVelocity = FloorResult.HitResult.ImpactNormal.Dot(ProjectedVelocity - ProjectedGroundVelocity);
 		const float VerticalVelocityLimit = 2.0f / DeltaSeconds;
-		if ((ProjectedRelativeVerticalVelocity > VerticalVelocityLimit) && (ProjectedVelocity.Z > VerticalVelocityLimit))
+		if ((ProjectedRelativeVerticalVelocity > VerticalVelocityLimit) && (ProjectedVelocity.Z > VerticalVelocityLimit) && bIsGroundMoving)
 		{
 			bIsLiftingOffSurface = true;
 		}
