Merging support for async-compatible teleportation
Includes sync-compatible Mover util functions for finding valid teleportation locations, testing whether a Mover actor would encroach at a proposed location, and searching for an adjusted location that would not encroach.

diff --git a/Engine/Plugins/Experimental/Mover/Source/Mover/Private/DefaultMovementSet/InstantMovementEffects/BasicInstantMovementEffects.cpp b/Engine/Plugins/Experimental/Mover/Source/Mover/Private/DefaultMovementSet/InstantMovementEffects/BasicInstantMovementEffects.cpp
--- a/Engine/Plugins/Experimental/Mover/Source/Mover/Private/DefaultMovementSet/InstantMovementEffects/BasicInstantMovementEffects.cpp
+++ b/Engine/Plugins/Experimental/Mover/Source/Mover/Private/DefaultMovementSet/InstantMovementEffects/BasicInstantMovementEffects.cpp
@@ -3,7 +3,6 @@
 
 #include "DefaultMovementSet/InstantMovementEffects/BasicInstantMovementEffects.h"
 
-#include "Engine/World.h"
 #include "MoverComponent.h"
 #include "MoverDataModelTypes.h"
 #include "MoverSimulationTypes.h"
@@ -96,10 +95,7 @@ bool FAsyncTeleportEffect::ApplyMovementEffect(FApplyMovementEffectParams& Apply
 	FVector TeleportLocation = TargetLocation;
 	const FRotator TeleportRotation = bUseActorRotation ? ApplyEffectParams.UpdatedComponent->GetComponentRotation() : TargetRotation;
 
-	if (UWorld* World = ApplyEffectParams.UpdatedComponent->GetWorld())
-	{
-		World->FindTeleportSpot(ApplyEffectParams.UpdatedComponent->GetOwner(), OUT TeleportLocation, TeleportRotation);
-	}
+	UMovementUtils::FindTeleportSpot(ApplyEffectParams.MoverComp, OUT TeleportLocation, TeleportRotation);
 
 	FMoverDefaultSyncState& OutputSyncState = OutputState.SyncStateCollection.FindOrAddMutableDataByType<FMoverDefaultSyncState>();
 
diff --git a/Engine/Plugins/Experimental/Mover/Source/Mover/Private/MoveLibrary/MovementUtils.cpp b/Engine/Plugins/Experimental/Mover/Source/Mover/Private/MoveLibrary/MovementUtils.cpp
--- a/Engine/Plugins/Experimental/Mover/Source/Mover/Private/MoveLibrary/MovementUtils.cpp
+++ b/Engine/Plugins/Experimental/Mover/Source/Mover/Private/MoveLibrary/MovementUtils.cpp
@@ -3,7 +3,9 @@
 #include "MoveLibrary/MovementUtils.h"
 #include "MoveLibrary/MovementRecord.h"
 #include "MoverLog.h"
+#include "Components/BoxComponent.h"
 #include "Components/SkeletalMeshComponent.h"
+#include "Engine/OverlapResult.h"
 #include "Engine/ScopedMovementUpdate.h"
 #include "Engine/World.h"
 #include "GameFramework/Pawn.h"
@@ -567,6 +569,178 @@ bool UMovementUtils::IsAngularVelocityZero(const FRotator& AngularVelocity)
 }
 
 
+bool UMovementUtils::FindTeleportSpot(const UMoverComponent* MoverComp, FVector& TestLocation, FRotator TestRotation)
+{
+	if (!MoverComp || !MoverComp->GetUpdatedComponent())
+	{
+		return true;
+	}
+	FVector Adjust(0.f);
+
+	const FVector OriginalTestLocation = TestLocation;
+
+	// check if fits at desired location
+	if (!TestEncroachment(MoverComp, TestLocation, TestRotation, &Adjust))
+	{
+		return true;
+	}
+
+	if (Adjust.IsNearlyZero())
+	{
+		// Reset in case Adjust is not actually zero
+		TestLocation = OriginalTestLocation;
+		return false;
+	}
+
+	// first do only Z
+	const FVector::FReal ZeroThreshold = UE_KINDA_SMALL_NUMBER;
+	const bool bZeroZ = FMath::IsNearlyZero(Adjust.Z, ZeroThreshold);
+	if (!bZeroZ)
+	{
+		TestLocation.Z += Adjust.Z;
+		if (!TestEncroachment(MoverComp, TestLocation, TestRotation))
+		{
+			return true;
+		}
+
+		TestLocation = OriginalTestLocation;
+	}
+
+	// now try just XY
+	const bool bZeroX = FMath::IsNearlyZero(Adjust.X, ZeroThreshold);
+	const bool bZeroY = FMath::IsNearlyZero(Adjust.Y, ZeroThreshold);
+	if (!bZeroX || !bZeroY)
+	{
+		const float X = bZeroX ? 0.f : Adjust.X;
+		const float Y = bZeroY ? 0.f : Adjust.Y;
+		FVector Adjustments[8];
+		Adjustments[0] = FVector(X, Y, 0);
+
+		// If initially spawning allow testing a few permutations (though this needs improvement).
+		// During play only test the first adjustment, permuting axes could put the location on other sides of geometry.
+		const int32 Iterations = (MoverComp->HasBegunPlay() ? 1 : 8);
+
+		if (Iterations > 1)
+		{
+			if (!bZeroX && !bZeroY)
+			{
+				Adjustments[1] = FVector(-X, Y, 0);
+				Adjustments[2] = FVector(X, -Y, 0);
+				Adjustments[3] = FVector(-X, -Y, 0);
+				Adjustments[4] = FVector(Y, X, 0);
+				Adjustments[5] = FVector(-Y, X, 0);
+				Adjustments[6] = FVector(Y, -X, 0);
+				Adjustments[7] = FVector(-Y, -X, 0);
+			}
+			else
+			{
+				// If either X or Y was zero, the permutations above would result in only 4 unique attempts.
+				Adjustments[1] = FVector(-X, -Y, 0);
+				Adjustments[2] = FVector(Y, X, 0);
+				Adjustments[3] = FVector(-Y, -X, 0);
+				// Mirror the dominant non-zero value
+				const float D = bZeroY ? X : Y;
+				Adjustments[4] = FVector(D, D, 0);
+				Adjustments[5] = FVector(D, -D, 0);
+				Adjustments[6] = FVector(-D, D, 0);
+				Adjustments[7] = FVector(-D, -D, 0);
+			}
+		}
+
+		for (int i = 0; i < Iterations; ++i)
+		{
+			TestLocation = OriginalTestLocation + Adjustments[i];
+			if (!TestEncroachment(MoverComp, TestLocation, TestRotation))
+			{
+				return true;
+			}
+		}
+
+		// Try XY adjustment including Z. Note that even with only 1 iteration, this will still try the full proposed (X,Y,Z) adjustment.
+		if (!bZeroZ)
+		{
+			for (int i = 0; i < Iterations; ++i)
+			{
+				TestLocation = OriginalTestLocation + Adjustments[i];
+				TestLocation.Z += Adjust.Z;
+				if (!TestEncroachment(MoverComp, TestLocation, TestRotation))
+				{
+					return true;
+				}
+			}
+		}
+	}
+
+	// Don't write out the last failed test location, we promised to only if we find a good spot, in case the caller re-uses the original input.
+	TestLocation = OriginalTestLocation;
+	return false;
+}
+
+bool UMovementUtils::TestEncroachment(const UMoverComponent* MoverComp, FVector TestLocation, FRotator TestRotation, OUT FVector* OutProposedAdjustment)
+{
+	if (MoverComp == nullptr)
+	{
+		return false;
+	}
+
+	USceneComponent* const RootComponent = MoverComp->GetUpdatedComponent();
+	if (RootComponent == nullptr)
+	{
+		return false;
+	}
+
+	bool bFoundEncroacher = false;
+
+	FVector TotalAdjustment(0.f);
+	FTransform const TestRootToWorld = FTransform(TestRotation, TestLocation);
+	FTransform const WorldToOldRoot = RootComponent->GetComponentToWorld().Inverse();
+
+	const UPrimitiveComponent* MovingPrimitiveRoot = Cast<UPrimitiveComponent>(MoverComp->GetUpdatedComponent());
+	if (MovingPrimitiveRoot)
+	{
+		// This actor has a movement component, which we interpret to mean that this actor has a primary component being swept around
+		// the world, and that component is the only one we care about encroaching (since the movement code will happily embedding
+		// other components in the world during movement updates)
+
+		if (MovingPrimitiveRoot->IsQueryCollisionEnabled())
+		{
+			// might not be the root, so we need to compute the transform
+			FTransform const CompToRoot = MovingPrimitiveRoot->GetComponentToWorld() * WorldToOldRoot;
+			FTransform const CompToNewWorld = CompToRoot * TestRootToWorld;
+
+			TArray<AActor*> ChildActors;
+			MoverComp->GetOwner()->GetAllChildActors(ChildActors);
+
+			FVector FoundAdjustment;
+
+			if (TestEncroachmentWithAdjustment_Internal(MoverComp->GetOwner()->GetWorld(), MoverComp->GetOwner(), MovingPrimitiveRoot, CompToNewWorld, FoundAdjustment, ChildActors))
+			{
+				if (OutProposedAdjustment == nullptr)
+				{
+					// don't need an adjustment and we know we are overlapping, so we can be done
+					return true;
+				}
+				else
+				{
+					TotalAdjustment = FoundAdjustment;
+				}
+
+				bFoundEncroacher = true;
+			}
+		}
+	}
+
+	// copy over total adjustment
+	if (OutProposedAdjustment)
+	{
+		*OutProposedAdjustment = TotalAdjustment;
+	}
+
+	return bFoundEncroacher;
+}
+
+
+
 bool UMovementUtils::TryMoveUpdatedComponent_Internal(const FMovingComponentSet& MovingComps, const FVector& Delta, const FQuat& NewRotation, bool bSweep, EMoveComponentFlags MoveComponentFlags, FHitResult* OutHit, ETeleportType Teleport)
 {
 	if (MovingComps.UpdatedComponent.IsValid())
@@ -583,3 +757,118 @@ bool UMovementUtils::TryMoveUpdatedComponent_Internal(const FMovingComponentSet&
 
 	return false;
 }
+
+
+
+// How much to adjust out collision shape box during encroachment testing, for a little leeway. This is applied to each axis, in cm.
+static const float ENCROACH_SHRINK_EPSILON = 0.15f;
+
+bool UMovementUtils::TestEncroachmentWithAdjustment_Internal(const UWorld* World, const AActor* TestActor, const UPrimitiveComponent* PrimComp, const FTransform& TestWorldTransform, FVector& OutProposedAdjustment, const TArray<AActor*>& IgnoreActors)
+{
+	// init our output
+	OutProposedAdjustment = FVector::ZeroVector;
+
+	float const Epsilon = ENCROACH_SHRINK_EPSILON;
+
+	if (World && PrimComp)
+	{
+		bool bFoundBlockingHit = false;
+		bool bComputePenetrationAdjustment = true;
+
+		TArray<FOverlapResult> Overlaps;
+		ECollisionChannel const BlockingChannel = PrimComp->GetCollisionObjectType();
+		FCollisionShape const CollisionShape = PrimComp->GetCollisionShape(-Epsilon);
+
+		if (CollisionShape.IsBox() && (Cast<UBoxComponent>(PrimComp) == nullptr))
+		{
+			// we have a bounding box not for a box component, which means this was the fallback aabb
+			// so lets test the actual component instead of it's aabb
+			// note we won't get penetration adjustment but that's ok
+			if (PrimComp->IsRegistered())
+			{
+				// must be registered
+				FComponentQueryParams Params(SCENE_QUERY_STAT(ComponentEncroachesBlockingGeometry_WithAdjustment), TestActor);
+				FCollisionResponseParams ResponseParams;
+				PrimComp->InitSweepCollisionParams(Params, ResponseParams);
+				Params.AddIgnoredActors(IgnoreActors);
+				bFoundBlockingHit = World->ComponentOverlapMultiByChannel(Overlaps, PrimComp, TestWorldTransform.GetLocation(), TestWorldTransform.GetRotation(), BlockingChannel, Params);
+				bComputePenetrationAdjustment = false;
+			}
+			else
+			{
+				UE_LOG(LogPhysics, Log, TEXT("Components must be registered in order to be used in a ComponentOverlapMulti call. PriComp: %s TestActor: %s"), *PrimComp->GetName(), *TestActor->GetName());
+			}
+		}
+		else
+		{
+			// overlap our shape
+			FCollisionQueryParams Params(SCENE_QUERY_STAT(ComponentEncroachesBlockingGeometry_WithAdjustment), false, TestActor);
+			FCollisionResponseParams ResponseParams;
+			PrimComp->InitSweepCollisionParams(Params, ResponseParams);
+			Params.AddIgnoredActors(IgnoreActors);
+			bFoundBlockingHit = World->OverlapMultiByChannel(Overlaps, TestWorldTransform.GetLocation(), TestWorldTransform.GetRotation(), BlockingChannel, CollisionShape, Params, ResponseParams);
+		}
+
+		// compute adjustment
+		if (bFoundBlockingHit && bComputePenetrationAdjustment)
+		{
+			// if encroaching, add up all the MTDs of overlapping shapes
+			FMTDResult MTDResult;
+			uint32 NumBlockingHits = 0;
+			for (int32 HitIdx = 0; HitIdx < Overlaps.Num(); HitIdx++)
+			{
+				UPrimitiveComponent* const OverlapComponent = Overlaps[HitIdx].Component.Get();
+				// first determine closest impact point along each axis
+				if (OverlapComponent && OverlapComponent->GetCollisionResponseToChannel(BlockingChannel) == ECR_Block)
+				{
+					NumBlockingHits++;
+					FCollisionShape const NonShrunkenCollisionShape = PrimComp->GetCollisionShape();
+					const FBodyInstance* OverlapBodyInstance = OverlapComponent->GetBodyInstance(NAME_None, true, Overlaps[HitIdx].ItemIndex);
+					bool bSuccess = OverlapBodyInstance && OverlapBodyInstance->OverlapTest(TestWorldTransform.GetLocation(), TestWorldTransform.GetRotation(), NonShrunkenCollisionShape, &MTDResult);
+					if (bSuccess)
+					{
+						OutProposedAdjustment += MTDResult.Direction * MTDResult.Distance;
+					}
+					else
+					{
+						UE_LOG(LogPhysics, Log, TEXT("OverlapTest says we are overlapping, yet MTD says we're not. Something is wrong"));
+						// It's not safe to use a partial result, that could push us out to an invalid location (like the other side of a wall).
+						OutProposedAdjustment = FVector::ZeroVector;
+						return true;
+					}
+
+					// #hack: sometimes for boxes, physx returns a 0 MTD even though it reports a contact (returns true)
+					// to get around this, let's go ahead and test again with the epsilon-shrunken collision shape to see if we're really in 
+					// the clear.
+					if (bSuccess && FMath::IsNearlyZero(MTDResult.Distance))
+					{
+						FCollisionShape const ShrunkenCollisionShape = PrimComp->GetCollisionShape(-Epsilon);
+						bSuccess = OverlapBodyInstance && OverlapBodyInstance->OverlapTest(TestWorldTransform.GetLocation(), TestWorldTransform.GetRotation(), ShrunkenCollisionShape, &MTDResult);
+						if (bSuccess)
+						{
+							OutProposedAdjustment += MTDResult.Direction * MTDResult.Distance;
+						}
+						else
+						{
+							// Ignore this overlap.
+							UE_LOG(LogPhysics, Log, TEXT("OverlapTest says we are overlapping, yet MTD says we're not (with smaller shape). Ignoring this overlap."));
+							NumBlockingHits--;
+							continue;
+						}
+					}
+				}
+			}
+
+			// See if we chose to invalidate all of our supposed "blocking hits".
+			if (NumBlockingHits == 0)
+			{
+				OutProposedAdjustment = FVector::ZeroVector;
+				bFoundBlockingHit = false;
+			}
+		}
+
+		return bFoundBlockingHit;
+	}
+
+	return false;
+}
diff --git a/Engine/Plugins/Experimental/Mover/Source/Mover/Public/MoveLibrary/MovementUtils.h b/Engine/Plugins/Experimental/Mover/Source/Mover/Public/MoveLibrary/MovementUtils.h
--- a/Engine/Plugins/Experimental/Mover/Source/Mover/Public/MoveLibrary/MovementUtils.h
+++ b/Engine/Plugins/Experimental/Mover/Source/Mover/Public/MoveLibrary/MovementUtils.h
@@ -217,6 +217,18 @@ public:
 	UFUNCTION(BlueprintCallable, Category = Mover)
 	static bool IsAngularVelocityZero(const FRotator& AngularVelocity);
 
+	/**
+	 * Try to find an acceptable non-colliding location to place TestActor as close to possible to PlaceLocation. Expects PlaceLocation to be a valid location inside the level.
+	 * Returns true if a location without blocking collision is found, in which case PlaceLocation is overwritten with the new clear location.
+	 * Returns false if no suitable location could be found, in which case PlaceLocation is unmodified.
+	 */
+	UFUNCTION(BlueprintCallable, Category = Mover)
+	static bool FindTeleportSpot(const UMoverComponent* MoverComp, FVector& TestLocation, FRotator TestRotation);
+
+	/** @Return true if MoverComp's Actor would encroach at TestLocation on something that blocks it.  Sets OutProposedAdjustment that might result in an unblocked TestLocation. */
+	static bool TestEncroachment(const UMoverComponent* MoverComp, FVector TestLocation, FRotator TestRotation, OUT FVector* OutProposedAdjustment=nullptr);
+
+
 	// Internal functions - not meant to be called outside of this library
 	
 	/** Internal function that other move functions use to perform all actual component movement and retrieve results
@@ -224,6 +236,9 @@ public:
 	 */
 	static bool TryMoveUpdatedComponent_Internal(const FMovingComponentSet& MovingComps, const FVector& Delta, const FQuat& NewRotation, bool bSweep, EMoveComponentFlags MoveComponentFlags, FHitResult* OutHit, ETeleportType Teleport);
 
+	// Internal function for testing whether a Mover actor would encroach at a test location, and computes a proposed adjustment where they won't encrouch (if found).
+	static bool TestEncroachmentWithAdjustment_Internal(const UWorld* World, const AActor* TestActor, const UPrimitiveComponent* PrimComp, const FTransform& TestWorldTransform, OUT FVector& OutProposedAdjustment, const TArray<AActor*>& IgnoreActors);
+
 };
 
 template <class ComponentType>
