Mover fixing issue with sim step infinite loop detection if movement modes refund an impossible amount

diff --git a/Engine/Plugins/Experimental/Mover/Source/Mover/Private/MovementModeStateMachine.cpp b/Engine/Plugins/Experimental/Mover/Source/Mover/Private/MovementModeStateMachine.cpp
--- a/Engine/Plugins/Experimental/Mover/Source/Mover/Private/MovementModeStateMachine.cpp
+++ b/Engine/Plugins/Experimental/Mover/Source/Mover/Private/MovementModeStateMachine.cpp
@@ -344,6 +344,8 @@ void UMovementModeStateMachine::OnSimulationTick(USceneComponent* UpdatedCompone
 				{
 					CurrentMode->DoSimulationTick(SimTickParams, OutputState);
 				}
+
+				OutputState.MovementEndState.RemainingMs = FMath::Clamp(OutputState.MovementEndState.RemainingMs, 0.0f, SubTimeStep.StepMs);
 			}
 
 			QueueNextMode(OutputState.MovementEndState.NextModeName);
