Mover: fixes rotation divergence during local->world space conversion in async simulations

diff --git a/Engine/Plugins/Experimental/Mover/Source/Mover/Private/MoverComponent.cpp b/Engine/Plugins/Experimental/Mover/Source/Mover/Private/MoverComponent.cpp
--- a/Engine/Plugins/Experimental/Mover/Source/Mover/Private/MoverComponent.cpp
+++ b/Engine/Plugins/Experimental/Mover/Source/Mover/Private/MoverComponent.cpp
@@ -1598,6 +1598,25 @@ bool UMoverComponent::RemoveMovementMode(FName ModeName)
 }
 
 
+/** Converts localspace root motion to a specific alternate worldspace location, taking the relative transform of the localspace component into account. */
+static FTransform ConvertLocalRootMotionToAltWorldSpace(const FTransform& LocalRootMotionTransform, const FTransform& AltWorldspaceTransform, const USceneComponent& RelativeComp)
+{
+	const FTransform TrueActorToWorld = RelativeComp.GetOwner()->GetTransform();
+	const FTransform RelativeCompToActor = TrueActorToWorld.GetRelativeTransform(RelativeComp.GetComponentTransform());
+
+	const FTransform AltComponentWorldTransform = RelativeCompToActor.Inverse() * AltWorldspaceTransform;
+
+	const FTransform NewComponentToWorld = LocalRootMotionTransform * AltComponentWorldTransform;
+	const FTransform NewActorTransform = RelativeCompToActor * NewComponentToWorld;
+
+	const FVector DeltaWorldTranslation = NewActorTransform.GetTranslation() - AltWorldspaceTransform.GetTranslation();
+
+	const FQuat NewWorldRotation = AltComponentWorldTransform.GetRotation() * LocalRootMotionTransform.GetRotation();
+	const FQuat DeltaWorldRotation = NewWorldRotation * AltComponentWorldTransform.GetRotation().Inverse();
+
+	return FTransform(DeltaWorldRotation, DeltaWorldTranslation);
+}
+
 FTransform UMoverComponent::ConvertLocalRootMotionToWorld(const FTransform& LocalRootMotionTransform, float DeltaSeconds, const FTransform* AlternateActorToWorld, const FMotionWarpingUpdateContext* OptionalWarpingContext) const
 {
 	// Optionally process/warp localspace root motion
@@ -1610,11 +1629,18 @@ FTransform UMoverComponent::ConvertLocalRootMotionToWorld(const FTransform& Loca
 
 	if (USkeletalMeshComponent* SkeletalMesh = GetPrimaryVisualComponent<USkeletalMeshComponent>())
 	{
-		 WorldSpaceRootMotion = SkeletalMesh->ConvertLocalRootMotionToWorld(ProcessedLocalRootMotion);
+		if (AlternateActorToWorld)
+		{
+			WorldSpaceRootMotion = ConvertLocalRootMotionToAltWorldSpace(ProcessedLocalRootMotion, *AlternateActorToWorld, *SkeletalMesh);
+		}
+		else
+		{
+			WorldSpaceRootMotion = SkeletalMesh->ConvertLocalRootMotionToWorld(ProcessedLocalRootMotion);
+		}
 	}
 	else
 	{
-		const FTransform PresentationActorToWorldTransform = GetOwner()->GetTransform();
+		const FTransform PresentationActorToWorldTransform = AlternateActorToWorld ? *AlternateActorToWorld : GetOwner()->GetTransform();
 		const FVector DeltaWorldTranslation = ProcessedLocalRootMotion.GetTranslation() - PresentationActorToWorldTransform.GetTranslation();
 
 		const FQuat NewWorldRotation = PresentationActorToWorldTransform.GetRotation() * ProcessedLocalRootMotion.GetRotation();
@@ -1623,18 +1649,6 @@ FTransform UMoverComponent::ConvertLocalRootMotionToWorld(const FTransform& Loca
 		WorldSpaceRootMotion.SetComponents(DeltaWorldRotation, DeltaWorldTranslation, FVector::OneVector);
 	}
 
-	// Optionally convert this to be relative to a different space
-	if (AlternateActorToWorld)
-	{
-		const FTransform AlternateActorToWorldNoTrans(AlternateActorToWorld->GetRotation(), FVector::ZeroVector, AlternateActorToWorld->GetScale3D());
-
-		FTransform WorldToActorNoTrans(GetOwner()->GetTransform().Inverse());
-		WorldToActorNoTrans.SetTranslation(FVector::ZeroVector);
-
-		const FTransform ActorSpaceRootMotion = WorldSpaceRootMotion * WorldToActorNoTrans;
-		WorldSpaceRootMotion = ActorSpaceRootMotion * AlternateActorToWorldNoTrans;
-	}
-	
 	// Optionally process/warp worldspace root motion
 	return ProcessWorldRootMotionDelegate.IsBound()
 		? ProcessWorldRootMotionDelegate.Execute(WorldSpaceRootMotion, DeltaSeconds, OptionalWarpingContext)
