[Mover] Set UBaseMovementMode to always be outer'd to a MoverComponent (it was already making that assumption anyway)
- Also set the default type for GetMoverComponentChecked so it can be called without <> if a subtype isn't needed

diff --git a/Engine/Plugins/Experimental/Mover/Source/Mover/Public/MovementMode.h b/Engine/Plugins/Experimental/Mover/Source/Mover/Public/MovementMode.h
--- a/Engine/Plugins/Experimental/Mover/Source/Mover/Public/MovementMode.h
+++ b/Engine/Plugins/Experimental/Mover/Source/Mover/Public/MovementMode.h
@@ -33,7 +33,7 @@ public:
 /**
  * Base class for all movement modes, exposing simulation update methods for both C++ and blueprint extension
  */
-UCLASS(Abstract, Blueprintable, BlueprintType, EditInlineNew)
+UCLASS(Abstract, Within = MoverComponent, Blueprintable, BlueprintType, EditInlineNew, DefaultToInstanced)
 class MOVER_API UBaseMovementMode : public UObject
 {
 	GENERATED_UCLASS_BODY()
@@ -51,11 +51,11 @@ public:
 	UMoverComponent* GetMoverComponent() const;
 
 	/** Templated convenience version of GetMoverComponent which checks the type is as presumed. */
-	template<class T>
+	template<class T = UMoverComponent>
 	T* GetMoverComponentChecked() const
 	{
 		static_assert(TPointerIsConvertibleFromTo<T, const UMoverComponent>::Value, "'T' template parameter to GetMoverComponentChecked must be derived from UMoverComponent");
-		return CastChecked<T>(GetOuter());
+		return CastChecked<T>(GetOuterUMoverComponent());
 	}
 
 #if WITH_EDITOR
