[Mover] - Adding more information to the check message when an FMoverDataStructBase virtual is called but must be overridden
- FMoverDataCollection::Interpolate now includes any struct present only in From or To in the interpolated result
- Added Merge virtual, atm only expected to be necessary for physics mover input structs

diff --git a/Engine/Plugins/Experimental/Mover/Source/Mover/Private/MoverTypes.cpp b/Engine/Plugins/Experimental/Mover/Source/Mover/Private/MoverTypes.cpp
--- a/Engine/Plugins/Experimental/Mover/Source/Mover/Private/MoverTypes.cpp
+++ b/Engine/Plugins/Experimental/Mover/Source/Mover/Private/MoverTypes.cpp
@@ -4,6 +4,7 @@
 #include "MoverTypes.h"
 #include "Blueprint/BlueprintExceptionInfo.h"
 #include "MoverLog.h"
+#include "UObject/ObjectKey.h"
 
 #define LOCTEXT_NAMESPACE "MoverData"
 
@@ -34,29 +35,34 @@ FMoverDataStructBase::FMoverDataStructBase()
 FMoverDataStructBase* FMoverDataStructBase::Clone() const
 {
 	// If child classes don't override this, collections will not work
-	checkf(false, TEXT("FMoverDataStructBase::Clone() being called erroneously. This should always be overridden in derived types!"));
+	checkf(false, TEXT("%hs is being called erroneously on [%s]. This must be overridden in derived types!"), __FUNCTION__, *GetScriptStruct()->GetName());
 	return nullptr;
 }
 
 UScriptStruct* FMoverDataStructBase::GetScriptStruct() const
 {
-	checkf(false, TEXT("FMoverDataStructBase::GetScriptStruct() being called erroneously. This should always be overridden in derived types!"));
+	checkf(false, TEXT("%hs is being called erroneously. This must be overridden in derived types!"), __FUNCTION__);
 	return FMoverDataStructBase::StaticStruct();
 }
 
-
 bool FMoverDataStructBase::ShouldReconcile(const FMoverDataStructBase& AuthorityState) const
 {
-	checkf(false, TEXT("FMoverDataStructBase::ShouldReconcile being called erroneously. This should always be overridden in derived types that comprise STATE data (sync/aux)!"));
+	checkf(false, TEXT("%hs is being called erroneously on [%s]. This must be overridden in derived types that comprise STATE data (sync/aux) "
+					"or INPUT data for use with physics-based movement"), __FUNCTION__, *GetScriptStruct()->GetName());
 	return false;
 }
 
 void FMoverDataStructBase::Interpolate(const FMoverDataStructBase& From, const FMoverDataStructBase& To, float Pct)
 {
-	checkf(false, TEXT("FMoverDataStructBase::Interpolate being called erroneously. This should always be overridden in derived types that comprise STATE data (sync/aux)!"));
+	checkf(false, TEXT("%hs is being called erroneously on [%s]. This must be overridden in derived types that comprise STATE data (sync/aux) "
+					"or INPUT data for use with physics-based movement"), __FUNCTION__, *GetScriptStruct()->GetName());
 }
 
-
+void FMoverDataStructBase::Merge(const FMoverDataStructBase& From)
+{
+	checkf(false, TEXT("%hs is being called erroneously on [%s]. This must be overridden in derived types that comprise INPUT data for use with physics-based movement"),
+		__FUNCTION__, *GetScriptStruct()->GetName());
+}
 
 
 FMoverDataCollection::FMoverDataCollection()
@@ -164,21 +170,48 @@ bool FMoverDataCollection::ShouldReconcile(const FMoverDataCollection& Other) co
 
 void FMoverDataCollection::Interpolate(const FMoverDataCollection& From, const FMoverDataCollection& To, float Pct)
 {
+	TSet<TObjectKey<UScriptStruct>> InterpolatedTypes;
+	
 	// Piece-wise interpolation of matching data blocks
-	for (int32 i = 0; i < From.DataArray.Num(); ++i)
+	for (const TSharedPtr<FMoverDataStructBase>& FromElement : From.DataArray)
 	{
-		const FMoverDataStructBase* FromElement = From.DataArray[i].Get();
-		
+		InterpolatedTypes.Add(FromElement->GetScriptStruct());
 		if (const FMoverDataStructBase* ToElement = To.FindDataByType(FromElement->GetScriptStruct()))
 		{
 			FMoverDataStructBase* InterpElement = FindOrAddDataByType(FromElement->GetScriptStruct());
 			InterpElement->Interpolate(*FromElement, *ToElement, Pct);
 		}
+		else
+		{
+			// If only present in From, add the block directly to this collection
+			AddOrOverwriteData(FromElement);
+		}
 	}
 
-	// TODO: What if From or To have any block types that are unique to them?
-
+	// Add any types present only in To as well
+	for (const TSharedPtr<FMoverDataStructBase>& ToElement : To.DataArray)
+	{
+		if (!InterpolatedTypes.Contains(ToElement->GetScriptStruct()))
+		{
+			AddOrOverwriteData(ToElement);
+		}
+	}
+}
 
+void FMoverDataCollection::Merge(const FMoverDataCollection& From)
+{
+	for (const TSharedPtr<FMoverDataStructBase>& FromElement : From.DataArray)
+	{
+		if (FMoverDataStructBase* ExistingElement = FindDataByType(FromElement->GetScriptStruct()))
+		{
+			ExistingElement->Merge(*FromElement);
+		}
+		else
+		{
+			// If only present in the previous block, copy it into this block
+			AddOrOverwriteData(FromElement);
+		}
+	}
 }
 
 
diff --git a/Engine/Plugins/Experimental/Mover/Source/Mover/Public/MoverTypes.h b/Engine/Plugins/Experimental/Mover/Source/Mover/Public/MoverTypes.h
--- a/Engine/Plugins/Experimental/Mover/Source/Mover/Public/MoverTypes.h
+++ b/Engine/Plugins/Experimental/Mover/Source/Mover/Public/MoverTypes.h
@@ -141,16 +141,26 @@ struct MOVER_API FMoverDataStructBase
 	/** If derived classes hold any object references, override this function and add them to the collector. */
 	virtual void AddReferencedObjects(FReferenceCollector& Collector) {}
 
-	/** Checks if the contained data is equal, within reason. MUST be override by types that compose STATE data (sync or aux). 
-	 *   AuthorityState is guaranteed to be the same concrete type as 'this'.
+	/**
+	 * Checks if the contained data is equal, within reason.
+	 * AuthorityState is guaranteed to be the same concrete type as 'this'.
+	 * Override for: All types that compose STATE data (sync or aux) and types that compose INPUT data for physics-based movement
 	 */
 	virtual bool ShouldReconcile(const FMoverDataStructBase& AuthorityState) const;
 
-	/** Interpolates contained data between a starting and ending block. MUST be override by types that compose STATE data (sync or aux). 
+	/**
+	 * Interpolates contained data between a starting and ending block.
 	 * From and To are guaranteed to be the same concrete type as 'this'.
+	 * Override for: All types that compose STATE data (sync or aux) and types that compose INPUT data for physics-based movement  
 	 */
 	virtual void Interpolate(const FMoverDataStructBase& From, const FMoverDataStructBase& To, float Pct);
 
+	/**
+	 * Merges contained data from a previous frame with that of the current frame. For more information, see the comments on FNetworkPhysicsData::MergeData.
+	 * From is guaranteed to be the same concrete type as 'this'.
+	 * Override for: Types that compose INPUT data for physics-based movement
+	 */
+	virtual void Merge(const FMoverDataStructBase& From);
 };
 
 template<>
@@ -197,6 +207,9 @@ struct MOVER_API FMoverDataCollection
 	/** Make this collection a piece-wise interpolation between 2 collections */
 	void Interpolate(const FMoverDataCollection& From, const FMoverDataCollection& To, float Pct);
 
+	/** Merge a previous frame's collection with this collection */
+	void Merge(const FMoverDataCollection& From);
+	
 	/** Exposes references to GC system */
 	void AddStructReferencedObjects(FReferenceCollector& Collector) const;
 
