[Chaos Mover] Implement ShouldReconcile, Interpolate, and Merge for FCharacterDefaultInputs, FMovementSettingsInputs, and FMoverAIInputs

diff --git a/Engine/Plugins/Experimental/Mover/Source/Mover/Private/MoverDataModelTypes.cpp b/Engine/Plugins/Experimental/Mover/Source/Mover/Private/MoverDataModelTypes.cpp
--- a/Engine/Plugins/Experimental/Mover/Source/Mover/Private/MoverDataModelTypes.cpp
+++ b/Engine/Plugins/Experimental/Mover/Source/Mover/Private/MoverDataModelTypes.cpp
@@ -111,6 +111,34 @@ void FCharacterDefaultInputs::ToString(FAnsiStringBuilderBase& Out) const
 	Out.Appendf("bIsJumpPressed: %i\tbIsJumpJustPressed: %i\n", bIsJumpPressed, bIsJumpJustPressed);
 }
 
+bool FCharacterDefaultInputs::ShouldReconcile(const FMoverDataStructBase& AuthorityState) const
+{
+	const FCharacterDefaultInputs& TypedAuthority = static_cast<const FCharacterDefaultInputs&>(AuthorityState);
+	return *this != TypedAuthority;
+}
+
+void FCharacterDefaultInputs::Interpolate(const FMoverDataStructBase& From, const FMoverDataStructBase& To, float Pct)
+{
+	const FCharacterDefaultInputs& TypedFrom = static_cast<const FCharacterDefaultInputs&>(From);
+	const FCharacterDefaultInputs& TypedTo = static_cast<const FCharacterDefaultInputs&>(To);
+	
+	// Note, this ignores movement base as this is not used by the physics mover
+	const FCharacterDefaultInputs* ClosestInputs = Pct < 0.5f ? &TypedFrom : &TypedTo;
+	bIsJumpJustPressed = ClosestInputs->bIsJumpJustPressed;
+	bIsJumpPressed = ClosestInputs->bIsJumpPressed;
+	SuggestedMovementMode = ClosestInputs->SuggestedMovementMode;
+
+	SetMoveInput(ClosestInputs->GetMoveInputType(), FMath::Lerp(TypedFrom.GetMoveInput(), TypedTo.GetMoveInput(), Pct));
+	OrientationIntent = FMath::Lerp(TypedFrom.OrientationIntent, TypedTo.OrientationIntent, Pct);
+	ControlRotation = FMath::Lerp(TypedFrom.ControlRotation, TypedTo.ControlRotation, Pct);
+}
+
+void FCharacterDefaultInputs::Merge(const FMoverDataStructBase& From)
+{
+	const FCharacterDefaultInputs& TypedFrom = static_cast<const FCharacterDefaultInputs&>(From);
+	bIsJumpJustPressed |= TypedFrom.bIsJumpJustPressed;
+	bIsJumpPressed |= TypedFrom.bIsJumpPressed;
+}
 
 
 // FMoverDefaultSyncState //////////////////////////////////////////////////////////////
diff --git a/Engine/Plugins/Experimental/Mover/Source/Mover/Private/PhysicsMover/PhysicsMoverSimulationTypes.cpp b/Engine/Plugins/Experimental/Mover/Source/Mover/Private/PhysicsMover/PhysicsMoverSimulationTypes.cpp
--- a/Engine/Plugins/Experimental/Mover/Source/Mover/Private/PhysicsMover/PhysicsMoverSimulationTypes.cpp
+++ b/Engine/Plugins/Experimental/Mover/Source/Mover/Private/PhysicsMover/PhysicsMoverSimulationTypes.cpp
@@ -45,8 +45,28 @@ void FMovementSettingsInputs::ToString(FAnsiStringBuilderBase& Out) const
 {
 	Super::ToString(Out);
 
-	Out.Appendf("MaxSpeed=%.2f\n", MaxSpeed);
-	Out.Appendf("Acceleration=%.2f\n", Acceleration);
+	Out.Appendf("MaxSpeed=%.2f | ", MaxSpeed);
+	Out.Appendf("Acceleration=%.2f", Acceleration);
+}
+
+bool FMovementSettingsInputs::ShouldReconcile(const FMoverDataStructBase& AuthorityState) const
+{
+	const FMovementSettingsInputs& TypedAuthority = static_cast<const FMovementSettingsInputs&>(AuthorityState);
+	return !FMath::IsNearlyEqual(Acceleration, TypedAuthority.Acceleration) || !FMath::IsNearlyEqual(MaxSpeed, TypedAuthority.MaxSpeed);
+}
+
+void FMovementSettingsInputs::Interpolate(const FMoverDataStructBase& From, const FMoverDataStructBase& To, float Pct)
+{
+	const FMovementSettingsInputs& TypedFrom = static_cast<const FMovementSettingsInputs&>(From);
+	const FMovementSettingsInputs& TypedTo = static_cast<const FMovementSettingsInputs&>(To);
+
+	MaxSpeed = FMath::Lerp(TypedFrom.MaxSpeed, TypedTo.MaxSpeed, Pct);;
+	Acceleration = FMath::Lerp(TypedFrom.Acceleration, TypedTo.Acceleration, Pct);
+
+}
+
+void FMovementSettingsInputs::Merge(const FMoverDataStructBase& From)
+{
 }
 
 
@@ -74,4 +94,22 @@ void FMoverAIInputs::ToString(FAnsiStringBuilderBase& Out) const
 	Super::ToString(Out);
 
 	Out.Appendf("ROVVelDelta: X=%.2f Y=%.2f Z=%.2f\n", RVOVelocityDelta.X, RVOVelocityDelta.Y, RVOVelocityDelta.Z);
-}
\ No newline at end of file
+}
+
+bool FMoverAIInputs::ShouldReconcile(const FMoverDataStructBase& AuthorityState) const
+{
+	const FMoverAIInputs& TypedAuthority = static_cast<const FMoverAIInputs&>(AuthorityState);
+	return TypedAuthority.RVOVelocityDelta != RVOVelocityDelta;
+}
+
+void FMoverAIInputs::Interpolate(const FMoverDataStructBase& From, const FMoverDataStructBase& To, float Pct)
+{
+	const FMoverAIInputs& TypedFrom = static_cast<const FMoverAIInputs&>(From);
+	const FMoverAIInputs& TypedTo = static_cast<const FMoverAIInputs&>(To);
+
+	RVOVelocityDelta = FMath::Lerp(TypedFrom.RVOVelocityDelta, TypedTo.RVOVelocityDelta, Pct);
+}
+
+void FMoverAIInputs::Merge(const FMoverDataStructBase& From)
+{
+}
diff --git a/Engine/Plugins/Experimental/Mover/Source/Mover/Public/MoverDataModelTypes.h b/Engine/Plugins/Experimental/Mover/Source/Mover/Public/MoverDataModelTypes.h
--- a/Engine/Plugins/Experimental/Mover/Source/Mover/Public/MoverDataModelTypes.h
+++ b/Engine/Plugins/Experimental/Mover/Source/Mover/Public/MoverDataModelTypes.h
@@ -94,18 +94,31 @@ public:
 
 	virtual ~FCharacterDefaultInputs() {}
 
+	bool operator==(const FCharacterDefaultInputs& Other) const
+	{
+		return MoveInputType == Other.MoveInputType
+			&& MoveInput == Other.MoveInput
+			&& OrientationIntent == Other.OrientationIntent
+			&& ControlRotation == Other.ControlRotation
+			&& SuggestedMovementMode == Other.SuggestedMovementMode
+			&& bUsingMovementBase == Other.bUsingMovementBase
+			&& MovementBase == Other.MovementBase
+			&& MovementBaseBoneName == Other.MovementBaseBoneName
+			&& bIsJumpJustPressed == Other.bIsJumpJustPressed
+			&& bIsJumpPressed == Other.bIsJumpPressed;
+	}
+	
+	bool operator!=(const FCharacterDefaultInputs& Other) const { return !operator==(Other); }
+	
 	// @return newly allocated copy of this FCharacterDefaultInputs. Must be overridden by child classes
 	virtual FMoverDataStructBase* Clone() const override;
-
 	virtual bool NetSerialize(FArchive& Ar, UPackageMap* Map, bool& bOutSuccess) override;
-
 	virtual UScriptStruct* GetScriptStruct() const override { return StaticStruct(); }
-
 	virtual void ToString(FAnsiStringBuilderBase& Out) const override;
-
 	virtual void AddReferencedObjects(FReferenceCollector& Collector) override { Super::AddReferencedObjects(Collector); }
-
-
+	virtual bool ShouldReconcile(const FMoverDataStructBase& AuthorityState) const override;
+	virtual void Interpolate(const FMoverDataStructBase& From, const FMoverDataStructBase& To, float Pct) override;
+	virtual void Merge(const FMoverDataStructBase& From) override;
 };
 
 template<>
diff --git a/Engine/Plugins/Experimental/Mover/Source/Mover/Public/PhysicsMover/PhysicsMoverSimulationTypes.h b/Engine/Plugins/Experimental/Mover/Source/Mover/Public/PhysicsMover/PhysicsMoverSimulationTypes.h
--- a/Engine/Plugins/Experimental/Mover/Source/Mover/Public/PhysicsMover/PhysicsMoverSimulationTypes.h
+++ b/Engine/Plugins/Experimental/Mover/Source/Mover/Public/PhysicsMover/PhysicsMoverSimulationTypes.h
@@ -126,16 +126,14 @@ public:
 
 	virtual ~FMovementSettingsInputs() {}
 
-	// @return newly allocated copy of this FMovementSettingsInputs. Must be overridden by child classes
 	virtual FMoverDataStructBase* Clone() const override;
-
 	virtual bool NetSerialize(FArchive& Ar, UPackageMap* Map, bool& bOutSuccess) override;
-
 	virtual UScriptStruct* GetScriptStruct() const override { return StaticStruct(); }
-
 	virtual void ToString(FAnsiStringBuilderBase& Out) const override;
-
 	virtual void AddReferencedObjects(FReferenceCollector& Collector) override { Super::AddReferencedObjects(Collector); }
+	virtual bool ShouldReconcile(const FMoverDataStructBase& AuthorityState) const override;
+	virtual void Interpolate(const FMoverDataStructBase& From, const FMoverDataStructBase& To, float Pct) override;
+	virtual void Merge(const FMoverDataStructBase& From) override;
 };
 
 template<>
@@ -143,7 +141,6 @@ struct TStructOpsTypeTraits< FMovementSettingsInputs > : public TStructOpsTypeTr
 {
 	enum
 	{
-		WithNetSerializer = true,
 		WithCopy = true
 	};
 };
@@ -171,16 +168,14 @@ public:
 
 	virtual ~FMoverAIInputs() {}
 
-	// @return newly allocated copy of this FMoverAIInputs. Must be overridden by child classes
 	virtual FMoverDataStructBase* Clone() const override;
-
 	virtual bool NetSerialize(FArchive& Ar, UPackageMap* Map, bool& bOutSuccess) override;
-
 	virtual UScriptStruct* GetScriptStruct() const override { return StaticStruct(); }
-
 	virtual void ToString(FAnsiStringBuilderBase& Out) const override;
-
 	virtual void AddReferencedObjects(FReferenceCollector& Collector) override { Super::AddReferencedObjects(Collector); }
+	virtual bool ShouldReconcile(const FMoverDataStructBase& AuthorityState) const override;
+	virtual void Interpolate(const FMoverDataStructBase& From, const FMoverDataStructBase& To, float Pct) override;
+	virtual void Merge(const FMoverDataStructBase& From) override;
 };
 
 template<>
@@ -188,7 +183,6 @@ struct TStructOpsTypeTraits< FMoverAIInputs > : public TStructOpsTypeTraitsBase2
 {
 	enum
 	{
-		WithNetSerializer = true,
 		WithCopy = true
 	};
 };
\ No newline at end of file
