[Chaos Mover] Make the TargetHeight property of the physics-driven walking, falling, and swimming modes an optional override param. If absent, the -Z offset of the owning ACharacter::Mesh is used.

diff --git a/Engine/Plugins/Experimental/Mover/Source/Mover/Private/PhysicsMover/Modes/PhysicsDrivenFallingMode.cpp b/Engine/Plugins/Experimental/Mover/Source/Mover/Private/PhysicsMover/Modes/PhysicsDrivenFallingMode.cpp
--- a/Engine/Plugins/Experimental/Mover/Source/Mover/Private/PhysicsMover/Modes/PhysicsDrivenFallingMode.cpp
+++ b/Engine/Plugins/Experimental/Mover/Source/Mover/Private/PhysicsMover/Modes/PhysicsDrivenFallingMode.cpp
@@ -16,6 +16,8 @@
 #include "Misc/DataValidation.h"
 #endif // WITH_EDITOR
 
+#include "Components/SkeletalMeshComponent.h"
+#include "GameFramework/Character.h"
 #include "MoveLibrary/MovementUtils.h"
 
 #include UE_INLINE_GENERATED_CPP_BY_NAME(PhysicsDrivenFallingMode)
@@ -50,6 +52,40 @@ EDataValidationResult UPhysicsDrivenFallingMode::IsDataValid(FDataValidationCont
 }
 #endif // WITH_EDITOR
 
+void UPhysicsDrivenFallingMode::SetTargetHeightOverride(float InTargetHeight)
+{
+	TargetHeightOverride = InTargetHeight;
+	TargetHeight = InTargetHeight;
+}
+
+void UPhysicsDrivenFallingMode::ClearTargetHeightOverride()
+{
+	TargetHeightOverride.Reset();
+	 
+	if (const ACharacter* CharacterOwner = GetTypedOuter<ACharacter>())
+	{
+		TargetHeight = -CharacterOwner->GetMesh()->GetRelativeLocation().Z;
+	}
+	else
+	{
+		TargetHeight = GetDefault<UPhysicsDrivenFallingMode>(GetClass())->TargetHeight;
+	}
+}
+
+void UPhysicsDrivenFallingMode::OnRegistered(const FName ModeName)
+{
+	Super::OnRegistered(ModeName);
+
+	if (TargetHeightOverride.IsSet())
+	{
+		TargetHeight = TargetHeightOverride.GetValue();
+	}
+	else if (const ACharacter* CharacterOwner = GetTypedOuter<ACharacter>())
+	{
+		TargetHeight = -CharacterOwner->GetMesh()->GetRelativeLocation().Z;
+	}
+}
+
 void UPhysicsDrivenFallingMode::OnSimulationTick(const FSimulationTickParams& Params, FMoverTickEndData& OutputState)
 {
 	const FMoverTickStartData& StartState = Params.StartState;
diff --git a/Engine/Plugins/Experimental/Mover/Source/Mover/Private/PhysicsMover/Modes/PhysicsDrivenSwimmingMode.cpp b/Engine/Plugins/Experimental/Mover/Source/Mover/Private/PhysicsMover/Modes/PhysicsDrivenSwimmingMode.cpp
--- a/Engine/Plugins/Experimental/Mover/Source/Mover/Private/PhysicsMover/Modes/PhysicsDrivenSwimmingMode.cpp
+++ b/Engine/Plugins/Experimental/Mover/Source/Mover/Private/PhysicsMover/Modes/PhysicsDrivenSwimmingMode.cpp
@@ -18,6 +18,9 @@
 #include "Misc/DataValidation.h"
 #endif // WITH_EDITOR
 
+#include "Components/SkeletalMeshComponent.h"
+#include "GameFramework/Character.h"
+
 #include UE_INLINE_GENERATED_CPP_BY_NAME(PhysicsDrivenSwimmingMode)
 
 #define LOCTEXT_NAMESPACE "PhysicsDrivenSwimmingMode"
@@ -50,6 +53,40 @@ EDataValidationResult UPhysicsDrivenSwimmingMode::IsDataValid(FDataValidationCon
 }
 #endif // WITH_EDITOR
 
+void UPhysicsDrivenSwimmingMode::SetTargetHeightOverride(float InTargetHeight)
+{
+	TargetHeightOverride = InTargetHeight;
+	TargetHeight = InTargetHeight;
+}
+
+void UPhysicsDrivenSwimmingMode::ClearTargetHeightOverride()
+{
+	TargetHeightOverride.Reset();
+	
+	if (const ACharacter* CharacterOwner = GetTypedOuter<ACharacter>())
+	{
+		TargetHeight = -CharacterOwner->GetMesh()->GetRelativeLocation().Z;
+	}
+	else
+	{
+		TargetHeight = GetDefault<UPhysicsDrivenSwimmingMode>(GetClass())->TargetHeight;
+	}
+}
+
+void UPhysicsDrivenSwimmingMode::OnRegistered(const FName ModeName)
+{
+	Super::OnRegistered(ModeName);
+
+	if (TargetHeightOverride.IsSet())
+	{
+		TargetHeight = TargetHeightOverride.GetValue();
+	}
+	else if (const ACharacter* CharacterOwner = GetTypedOuter<ACharacter>())
+	{
+		TargetHeight = -CharacterOwner->GetMesh()->GetRelativeLocation().Z;
+	}
+}
+
 void UPhysicsDrivenSwimmingMode::OnSimulationTick(const FSimulationTickParams& Params, FMoverTickEndData& OutputState)
 {
 	const UMoverComponent* MoverComp = GetMoverComponent();
diff --git a/Engine/Plugins/Experimental/Mover/Source/Mover/Private/PhysicsMover/Modes/PhysicsDrivenWalkingMode.cpp b/Engine/Plugins/Experimental/Mover/Source/Mover/Private/PhysicsMover/Modes/PhysicsDrivenWalkingMode.cpp
--- a/Engine/Plugins/Experimental/Mover/Source/Mover/Private/PhysicsMover/Modes/PhysicsDrivenWalkingMode.cpp
+++ b/Engine/Plugins/Experimental/Mover/Source/Mover/Private/PhysicsMover/Modes/PhysicsDrivenWalkingMode.cpp
@@ -24,6 +24,9 @@
 #include "Misc/DataValidation.h"
 #endif // WITH_EDITOR
 
+#include "Components/SkeletalMeshComponent.h"
+#include "GameFramework/Character.h"
+
 #include UE_INLINE_GENERATED_CPP_BY_NAME(PhysicsDrivenWalkingMode)
 
 extern FPhysicsDrivenMotionDebugParams GPhysicsDrivenMotionDebugParams;
@@ -260,6 +263,40 @@ EDataValidationResult UPhysicsDrivenWalkingMode::IsDataValid(FDataValidationCont
 }
 #endif // WITH_EDITOR
 
+void UPhysicsDrivenWalkingMode::SetTargetHeightOverride(float InTargetHeight)
+{
+	TargetHeightOverride = InTargetHeight;
+	TargetHeight = InTargetHeight;
+}
+
+void UPhysicsDrivenWalkingMode::ClearTargetHeightOverride()
+{
+	TargetHeightOverride.Reset();
+	
+	if (const ACharacter* CharacterOwner = GetTypedOuter<ACharacter>())
+	{
+		TargetHeight = -CharacterOwner->GetMesh()->GetRelativeLocation().Z;
+	}
+	else
+	{
+		TargetHeight = GetDefault<UPhysicsDrivenWalkingMode>(GetClass())->TargetHeight;
+	}
+}
+
+void UPhysicsDrivenWalkingMode::OnRegistered(const FName ModeName)
+{
+	Super::OnRegistered(ModeName);
+
+	if (TargetHeightOverride.IsSet())
+	{
+		TargetHeight = TargetHeightOverride.GetValue();
+	}
+	else if (const ACharacter* CharacterOwner = GetTypedOuter<ACharacter>())
+	{
+		TargetHeight = -CharacterOwner->GetMesh()->GetRelativeLocation().Z;
+	}
+}
+
 void UPhysicsDrivenWalkingMode::OnSimulationTick(const FSimulationTickParams& Params, FMoverTickEndData& OutputState)
 {
 	const UMoverComponent* MoverComp = GetMoverComponent();
diff --git a/Engine/Plugins/Experimental/Mover/Source/Mover/Public/PhysicsMover/Modes/PhysicsDrivenFallingMode.h b/Engine/Plugins/Experimental/Mover/Source/Mover/Public/PhysicsMover/Modes/PhysicsDrivenFallingMode.h
--- a/Engine/Plugins/Experimental/Mover/Source/Mover/Public/PhysicsMover/Modes/PhysicsDrivenFallingMode.h
+++ b/Engine/Plugins/Experimental/Mover/Source/Mover/Public/PhysicsMover/Modes/PhysicsDrivenFallingMode.h
@@ -20,15 +20,18 @@ class MOVER_API UPhysicsDrivenFallingMode : public UFallingMode, public IPhysics
 	GENERATED_UCLASS_BODY()
 
 public:
-
+	virtual void OnRegistered(const FName ModeName) override;
 	virtual void OnSimulationTick(const FSimulationTickParams& Params, FMoverTickEndData& OutputState) override;
-
 	virtual void UpdateConstraintSettings(Chaos::FCharacterGroundConstraint& Constraint) const override;
 
 #if WITH_EDITOR
 	virtual EDataValidationResult IsDataValid(class FDataValidationContext& Context) const override;
 #endif // WITH_EDITOR
 
+	void SetTargetHeightOverride(float InTargetHeight);
+	void ClearTargetHeightOverride();
+	
+protected:
 	// Maximum torque the character can apply to rotate in air about the vertical axis
 	UPROPERTY(EditAnywhere, Category = "Physics Mover", meta = (ClampMin = "0", UIMin = "0", ForceUnits = "NewtonMeters"))
 	float TwistTorqueLimit = 0.0f;
@@ -36,12 +39,17 @@ public:
 	// Maximum torque the character can apply to remain upright
 	UPROPERTY(EditAnywhere, Category = "Physics Mover", meta = (ClampMin = "0", UIMin = "0", ForceUnits = "NewtonMeters"))
 	float SwingTorqueLimit = 3000.0f;
-
-	// Target height for the character. This is the desired distance from the center of the capsule to the floor
+	
+	/**
+	 * Optional override target height for the character (the desired distance from the center of the capsule to the floor).
+	 * If left blank, the -Z offset of the owning character's skeletal mesh comp will be used automatically.
+	 */
 	UPROPERTY(EditAnywhere, Category = "Physics Mover", meta = (ClampMin = "0", UIMin = "0", ForceUnits = "cm"))
-	float TargetHeight = 95.0f;
+	TOptional<float> TargetHeightOverride;
 
 	// Radius used for ground queries
 	UPROPERTY(EditAnywhere, Category = "Physics Mover", meta = (ClampMin = "0", UIMin = "0", ForceUnits = "cm"))
 	float QueryRadius = 20.0f;
+
+	float TargetHeight = 95.f;
 };
diff --git a/Engine/Plugins/Experimental/Mover/Source/Mover/Public/PhysicsMover/Modes/PhysicsDrivenSwimmingMode.h b/Engine/Plugins/Experimental/Mover/Source/Mover/Public/PhysicsMover/Modes/PhysicsDrivenSwimmingMode.h
--- a/Engine/Plugins/Experimental/Mover/Source/Mover/Public/PhysicsMover/Modes/PhysicsDrivenSwimmingMode.h
+++ b/Engine/Plugins/Experimental/Mover/Source/Mover/Public/PhysicsMover/Modes/PhysicsDrivenSwimmingMode.h
@@ -20,7 +20,7 @@ class MOVER_API UPhysicsDrivenSwimmingMode : public USwimmingMode, public IPhysi
 	GENERATED_UCLASS_BODY()
 
 public:
-
+	virtual void OnRegistered(const FName ModeName) override;
 	virtual void OnSimulationTick(const FSimulationTickParams& Params, FMoverTickEndData& OutputState) override;
 	virtual bool AttemptJump(float UpwardsSpeed, FMoverTickEndData& Output) override;
 
@@ -30,11 +30,20 @@ public:
 	virtual EDataValidationResult IsDataValid(class FDataValidationContext& Context) const override;
 #endif // WITH_EDITOR
 
-	// Target height for the character. This is the desired distance from the center of the capsule to the floor
+	void SetTargetHeightOverride(float InTargetHeight);
+	void ClearTargetHeightOverride();
+	
+protected:
+	/**
+	 * Optional override target height for the character (the desired distance from the center of the capsule to the floor).
+	 * If left blank, the -Z offset of the owning character's skeletal mesh comp will be used automatically.
+	 */
 	UPROPERTY(EditAnywhere, Category = "Physics Mover", meta = (ClampMin = "0", UIMin = "0", ForceUnits = "cm"))
-	float TargetHeight = 95.0f;
+	TOptional<float> TargetHeightOverride;
 
 	// Radius used for ground queries
 	UPROPERTY(EditAnywhere, Category = "Physics Mover", meta = (ClampMin = "0", UIMin = "0", ForceUnits = "cm"))
 	float QueryRadius = 30.0f;
+	
+	float TargetHeight = 95.f;
 };
diff --git a/Engine/Plugins/Experimental/Mover/Source/Mover/Public/PhysicsMover/Modes/PhysicsDrivenWalkingMode.h b/Engine/Plugins/Experimental/Mover/Source/Mover/Public/PhysicsMover/Modes/PhysicsDrivenWalkingMode.h
--- a/Engine/Plugins/Experimental/Mover/Source/Mover/Public/PhysicsMover/Modes/PhysicsDrivenWalkingMode.h
+++ b/Engine/Plugins/Experimental/Mover/Source/Mover/Public/PhysicsMover/Modes/PhysicsDrivenWalkingMode.h
@@ -21,7 +21,7 @@ class MOVER_API UPhysicsDrivenWalkingMode : public UWalkingMode, public IPhysics
 	GENERATED_UCLASS_BODY()
 
 public:
-
+	virtual void OnRegistered(const FName ModeName) override;
 	virtual void OnSimulationTick(const FSimulationTickParams& Params, FMoverTickEndData& OutputState) override;
 
 	virtual void UpdateConstraintSettings(Chaos::FCharacterGroundConstraint& Constraint) const override;
@@ -31,6 +31,10 @@ public:
 	virtual EDataValidationResult IsDataValid(class FDataValidationContext& Context) const override;
 #endif // WITH_EDITOR
 
+	void SetTargetHeightOverride(float InTargetHeight);
+	void ClearTargetHeightOverride();
+	
+protected:
 	// Maximum force the character can apply to reach the motion target
 	UPROPERTY(EditAnywhere, Category = "Physics Mover", meta = (ClampMin = "0", UIMin = "0", ForceUnits = "Newtons"))
 	float RadialForceLimit = 1500.0f;
@@ -47,9 +51,12 @@ public:
 	UPROPERTY(EditAnywhere, Category = "Physics Mover", meta = (ClampMin = "0", UIMin = "0", ForceUnits = "NewtonMeters"))
 	float SwingTorqueLimit = 3000.0f;
 
-	// Target height for the character. This is the desired distance from the center of the capsule to the floor
+	/**
+	 * Optional override target height for the character (the desired distance from the center of the capsule to the floor).
+	 * If left blank, the -Z offset of the owning character's skeletal mesh comp will be used automatically.
+	 */
 	UPROPERTY(EditAnywhere, Category = "Physics Mover", meta = (ClampMin = "0", UIMin = "0", ForceUnits = "cm"))
-	float TargetHeight = 95.0f;
+	TOptional<float> TargetHeightOverride;
 
 	// Radius used for ground queries
 	UPROPERTY(EditAnywhere, Category = "Physics Mover", meta = (ClampMin = "0", UIMin = "0", ForceUnits = "cm"))
@@ -90,6 +97,8 @@ public:
 	UPROPERTY(EditAnywhere, Category = "Physics Mover")
 	bool bHandleVerticalLandingSeparately = true;
 
+	float TargetHeight = 95.f;
+	
 protected:
 	void SwitchToState(const FName& StateName, const FSimulationTickParams& Params, FMoverTickEndData& OutputState);
 
diff --git a/Engine/Plugins/Experimental/Mover/Source/Mover/Public/PhysicsMover/MovementModifiers/PhysicsStanceModifier.h b/Engine/Plugins/Experimental/Mover/Source/Mover/Public/PhysicsMover/MovementModifiers/PhysicsStanceModifier.h
--- a/Engine/Plugins/Experimental/Mover/Source/Mover/Public/PhysicsMover/MovementModifiers/PhysicsStanceModifier.h
+++ b/Engine/Plugins/Experimental/Mover/Source/Mover/Public/PhysicsMover/MovementModifiers/PhysicsStanceModifier.h
@@ -39,7 +39,7 @@ protected:
 		if (T* MovementMode = MoverComp->FindMode_Mutable<T>())
 		{
 			const float ColliderOffset = -MoverComp->GetPrimaryVisualComponent()->GetRelativeLocation().Z - OldHalfHeight;
-			MovementMode->TargetHeight = NewHalfHeight + ColliderOffset;
+			MovementMode->SetTargetHeightOverride(NewHalfHeight + ColliderOffset);
 		}
 	}
 };
\ No newline at end of file
