[Mover] MoverComponent:
- Adding a templated way to find a movement mode by name
- Adding UE_REQUIRES() to all the templates in MoverComponent to enforce base class requirements and added the base type as the default arg
- Changed the BackendClass property to be a TScriptInterface hard ref, as the component is the creator, owner, and sole user of the component and needs it to function

LayeredMove:
- Adding more convenient means of querying the active and queue moves in FLayeredMoveGroup
- Also added comments on FLayeredMove::DurationMs to call out that setting it to 0 is equivalent to returning true from IsFinished()

diff --git a/Engine/Plugins/Experimental/Mover/Source/Mover/Private/Debug/GameplayDebuggerCategory_Mover.cpp b/Engine/Plugins/Experimental/Mover/Source/Mover/Private/Debug/GameplayDebuggerCategory_Mover.cpp
--- a/Engine/Plugins/Experimental/Mover/Source/Mover/Private/Debug/GameplayDebuggerCategory_Mover.cpp
+++ b/Engine/Plugins/Experimental/Mover/Source/Mover/Private/Debug/GameplayDebuggerCategory_Mover.cpp
@@ -147,9 +147,9 @@ void FGameplayDebuggerCategory_Mover::CollectData(APlayerController* OwnerPC, AA
 		{
 			const FMoverSyncState& SyncState = MyMoverComponent->GetSyncState();
 
-			for (auto it = SyncState.LayeredMoves.GetActiveMovesIterator(); it; ++it)
+			for (const TSharedPtr<FLayeredMoveBase>& ActiveMove : SyncState.LayeredMoves.GetActiveMoves())
 			{
-				DataPack.ActiveLayeredMoves.Add(*it->Get()->ToSimpleString());
+				DataPack.ActiveLayeredMoves.Add(ActiveMove->ToSimpleString());
 			}
 
 			for (auto it = SyncState.MovementModifiers.GetActiveModifiersIterator(); it; ++it)
diff --git a/Engine/Plugins/Experimental/Mover/Source/Mover/Private/LayeredMove.cpp b/Engine/Plugins/Experimental/Mover/Source/Mover/Private/LayeredMove.cpp
--- a/Engine/Plugins/Experimental/Mover/Source/Mover/Private/LayeredMove.cpp
+++ b/Engine/Plugins/Experimental/Mover/Source/Mover/Private/LayeredMove.cpp
@@ -255,14 +255,28 @@ FString FLayeredMoveGroup::ToSimpleString() const
 	return FString::Printf(TEXT("FLayeredMoveGroup. Active: %i Queued: %i"), ActiveLayeredMoves.Num(), QueuedLayeredMoves.Num());
 }
 
-TArray<TSharedPtr<FLayeredMoveBase>>::TConstIterator FLayeredMoveGroup::GetActiveMovesIterator() const
+const FLayeredMoveBase* FLayeredMoveGroup::FindActiveMove(const UScriptStruct* LayeredMoveStructType) const
 {
-	return ActiveLayeredMoves.CreateConstIterator();
+	for (const TSharedPtr<FLayeredMoveBase>& ActiveMove : ActiveLayeredMoves)
+	{
+		if (ActiveMove && ActiveMove->GetScriptStruct()->IsChildOf(LayeredMoveStructType))
+		{
+			return ActiveMove.Get();
+		}
+	}
+	return nullptr;
 }
 
-TArray<TSharedPtr<FLayeredMoveBase>>::TConstIterator FLayeredMoveGroup::GetQueuedMovesIterator() const
+const FLayeredMoveBase* FLayeredMoveGroup::FindQueuedMove(const UScriptStruct* LayeredMoveStructType) const
 {
-	return QueuedLayeredMoves.CreateConstIterator();
+	for (const TSharedPtr<FLayeredMoveBase>& QueuedMove : QueuedLayeredMoves)
+	{
+		if (QueuedMove && QueuedMove->GetScriptStruct()->IsChildOf(LayeredMoveStructType))
+		{
+			return QueuedMove.Get();
+		}
+	}
+	return nullptr;
 }
 
 void FLayeredMoveGroup::FlushMoveArrays(const UMoverComponent* MoverComp, UMoverBlackboard* SimBlackboard, float CurrentSimTimeMs)
diff --git a/Engine/Plugins/Experimental/Mover/Source/Mover/Private/MoverComponent.cpp b/Engine/Plugins/Experimental/Mover/Source/Mover/Private/MoverComponent.cpp
--- a/Engine/Plugins/Experimental/Mover/Source/Mover/Private/MoverComponent.cpp
+++ b/Engine/Plugins/Experimental/Mover/Source/Mover/Private/MoverComponent.cpp
@@ -57,7 +57,6 @@ namespace MoverComponentConstants
 
 static constexpr float ROTATOR_TOLERANCE = (1e-3);
 
-
 UMoverComponent::UMoverComponent()
 {
 	PrimaryComponentTick.TickGroup = TG_PrePhysics;
@@ -90,9 +89,10 @@ void UMoverComponent::InitializeComponent()
 		// Instantiate out sister backend component that will actually talk to the system driving the simulation
 		if (BackendClass)
 		{
-			UActorComponent* NewLiaisonComp = NewObject<UActorComponent>(GetOwner(), BackendClass, TEXT("BackendLiaisonComponent"));
-			BackendLiaisonComp = CastChecked<IMoverBackendLiaisonInterface>(NewLiaisonComp);
-			if (BackendLiaisonComp.Get())
+			UActorComponent* NewLiaisonComp = NewObject<UActorComponent>(this, BackendClass, TEXT("BackendLiaisonComponent"));
+			BackendLiaisonComp.SetObject(NewLiaisonComp);
+			BackendLiaisonComp.SetInterface(CastChecked<IMoverBackendLiaisonInterface>(NewLiaisonComp));
+			if (BackendLiaisonComp)
 			{
 				NewLiaisonComp->RegisterComponent();
 				NewLiaisonComp->InitializeComponent();
@@ -111,7 +111,7 @@ void UMoverComponent::InitializeComponent()
 
 void UMoverComponent::UninitializeComponent()
 {
-	if (UActorComponent* LiaisonAsComp = Cast<UActorComponent>(BackendLiaisonComp.Get()))
+	if (UActorComponent* LiaisonAsComp = Cast<UActorComponent>(BackendLiaisonComp.GetObject()))
 	{
 		LiaisonAsComp->DestroyComponent();
 		BackendLiaisonComp = nullptr;
@@ -265,7 +265,6 @@ void UMoverComponent::ProduceInput(const int32 DeltaTimeMS, FMoverInputCmdContex
 
 	CachedLastProducedInputCmd = *Cmd;
 	bHasValidLastProducedInput = true;
-
 }
 
 void UMoverComponent::RestoreFrame(const FMoverSyncState* SyncState, const FMoverAuxStateContext* AuxState)
@@ -704,10 +703,9 @@ bool UMoverComponent::HasGameplayTag(FGameplayTag TagToFind, bool bExactMatch) c
 		}
 		
 		// Search Layered Moves
-		for (auto LayeredMoveFromSyncStateIt = CachedLastSyncState.LayeredMoves.GetActiveMovesIterator(); LayeredMoveFromSyncStateIt; ++LayeredMoveFromSyncStateIt)
+		for (const TSharedPtr<FLayeredMoveBase>& LayeredMove :  CachedLastSyncState.LayeredMoves.GetActiveMoves())
 		{
-			const TSharedPtr<FLayeredMoveBase> LayeredMoveFromSyncState = *LayeredMoveFromSyncStateIt;
-			if (LayeredMoveFromSyncState->HasGameplayTag(TagToFind, bExactMatch))
+			if (LayeredMove->HasGameplayTag(TagToFind, bExactMatch))
 			{
 				return true;
 			}
@@ -1435,6 +1433,15 @@ void UMoverComponent::QueueInstantMovementEffect(TSharedPtr<FInstantMovementEffe
 	ModeFSM->QueueInstantMovementEffect(InstantMovementEffect);
 }
 
+UBaseMovementMode* UMoverComponent::FindMovementModeByName(FName MovementModeName) const
+{
+	if (const TObjectPtr<UBaseMovementMode>* FoundMode = MovementModes.Find(MovementModeName))
+	{
+		return *FoundMode;
+	}
+	return nullptr;
+}
+
 void UMoverComponent::K2_FindActiveLayeredMove(bool& DidSucceed, int32& TargetAsRawBytes) const
 {
 	// This will never be called, the exec version below will be hit instead
@@ -1497,23 +1504,11 @@ DEFINE_FUNCTION(UMoverComponent::execK2_FindActiveLayeredMove)
 	}
 }
 
-const FLayeredMoveBase* UMoverComponent::FindActiveLayeredMoveByType(const UScriptStruct* DataStructType) const
+const FLayeredMoveBase* UMoverComponent::FindActiveLayeredMoveByType(const UScriptStruct* LayeredMoveStructType) const
 {
 	if (bHasValidCachedState)
 	{
-		for (auto it = CachedLastSyncState.LayeredMoves.GetActiveMovesIterator(); it; ++it)
-		{
-			UStruct* CandidateStruct = it->Get()->GetScriptStruct();
-			while (CandidateStruct)
-			{
-				if (DataStructType == CandidateStruct)
-				{
-					return it->Get();
-				}
-
-				CandidateStruct = CandidateStruct->GetSuperStruct();
-			}
-		}
+		return CachedLastSyncState.LayeredMoves.FindActiveMove(LayeredMoveStructType);
 	}
 
 	return nullptr;
@@ -1947,15 +1942,7 @@ FName UMoverComponent::GetMovementModeName() const
 
 const UBaseMovementMode* UMoverComponent::GetMovementMode() const
 {
-	if (bHasValidCachedState)
-	{
-		if (const TObjectPtr<UBaseMovementMode>* CurrentMode = MovementModes.Find(CachedLastSyncState.MovementMode))
-		{
-			return CurrentMode->Get();
-		}
-	}
-
-	return nullptr;
+	return GetActiveModeInternal(UBaseMovementMode::StaticClass());
 }
 
 UPrimitiveComponent* UMoverComponent::GetMovementBase() const
@@ -2076,14 +2063,49 @@ const UObject* UMoverComponent::FindSharedSettings_BP(TSubclassOf<UObject> Share
 	return nullptr;
 }
 
-UBaseMovementMode* UMoverComponent::FindMode_Mutable(const UClass* ByType, bool bRequireExactClass) const
+UBaseMovementMode* UMoverComponent::FindMode_Mutable(TSubclassOf<UBaseMovementMode> ModeType, bool bRequireExactClass) const
 {
-	for (const TPair<FName, TObjectPtr<UBaseMovementMode>>& Element : MovementModes)
+	if (ModeType)
 	{
-		if ( (!bRequireExactClass && Element.Value->IsA(ByType)) || 
-		     (Element.Value->GetClass() == ByType) )
+		for (const TPair<FName, TObjectPtr<UBaseMovementMode>>& NameModePair : MovementModes)
 		{
-			return Element.Value.Get();
+			if ( (!bRequireExactClass && NameModePair.Value->IsA(ModeType)) || 
+				 (NameModePair.Value->GetClass() == ModeType) )
+			{
+				return NameModePair.Value.Get();
+			}
+		}		
+	}
+
+	return nullptr;
+}
+
+UBaseMovementMode* UMoverComponent::FindMode_Mutable(TSubclassOf<UBaseMovementMode> ModeType, FName ModeName, bool bRequireExactClass) const
+{
+	if (!ModeName.IsNone())
+	{
+		if (const TObjectPtr<UBaseMovementMode>* FoundMode = MovementModes.Find(ModeName))
+		{
+			if ((!bRequireExactClass && FoundMode->IsA(ModeType)) || FoundMode->GetClass() == ModeType)
+			{
+				return *FoundMode;
+			} 
+		}
+	}
+	return nullptr;
+}
+
+UBaseMovementMode* UMoverComponent::GetActiveModeInternal(TSubclassOf<UBaseMovementMode> ModeType, bool bRequireExactClass) const
+{
+	if (bHasValidCachedState)
+	{
+		if (const TObjectPtr<UBaseMovementMode>* CurrentMode = MovementModes.Find(CachedLastSyncState.MovementMode))
+		{
+			if ((!bRequireExactClass && CurrentMode->IsA(ModeType)) ||
+				CurrentMode->GetClass() == ModeType)
+			{
+				return CurrentMode->Get();
+			}
 		}
 	}
 
diff --git a/Engine/Plugins/Experimental/Mover/Source/Mover/Public/LayeredMove.h b/Engine/Plugins/Experimental/Mover/Source/Mover/Public/LayeredMove.h
--- a/Engine/Plugins/Experimental/Mover/Source/Mover/Public/LayeredMove.h
+++ b/Engine/Plugins/Experimental/Mover/Source/Mover/Public/LayeredMove.h
@@ -84,7 +84,10 @@ struct MOVER_API FLayeredMoveBase
 	UPROPERTY(EditAnywhere, BlueprintReadWrite, Category = Mover)
 	uint8 Priority;
 	
-	// This move will expire after a set amount of time if > 0. If 0, it will be ticked only once, regardless of time step. It will need to be manually ended if < 0. 
+	/**
+	 * This move will expire after a set amount of time if > 0. If 0, it will be ticked only once, regardless of time step. It will need to be manually ended if < 0.
+	 * Note: If changed after starting to a value beneath the current lifetime of the move, it will immediately finish (so if your move finishes early, setting this to 0 is equivalent to returning true from IsFinished())
+	 */
 	UPROPERTY(EditAnywhere, BlueprintReadWrite, Category = Mover)
 	float DurationMs;
 
@@ -100,7 +103,7 @@ struct MOVER_API FLayeredMoveBase
 	 * Check Layered Move for a gameplay tag.
 	 *
 	 * @param TagToFind			Tag to check on the Mover systems
-	 * @param bExactMatch		If true, the tag has to be exactly present, if false then TagToFind will include it's parent tags while matching
+	 * @param bExactMatch		If true, the tag has to be exactly present, if false then TagToFind will include its parent tags while matching
 	 * 
 	 * @return True if the TagToFind was found
 	 */
@@ -185,6 +188,13 @@ struct MOVER_API FLayeredMoveGroup
 
 	bool HasAnyMoves() const { return (!ActiveLayeredMoves.IsEmpty() || !QueuedLayeredMoves.IsEmpty()); }
 
+	/** @return True if there is at least one layered move of MoveType that's either active or queued */
+	template <typename MoveType UE_REQUIRES(std::is_base_of_v<FLayeredMoveBase, MoveType>)>
+	bool HasMove() const
+	{
+		return !FindActiveMove<MoveType>() && !FindQueuedMove<MoveType>();
+	}
+
 	// Generates active layered move list (by calling FlushMoveArrays) and returns the an array of all currently active layered moves
 	TArray<TSharedPtr<FLayeredMoveBase>> GenerateActiveMoves(const FMoverTimeStep& TimeStep, const UMoverComponent* MoverComp, UMoverBlackboard* SimBlackboard);
 
@@ -206,12 +216,25 @@ struct MOVER_API FLayeredMoveGroup
 	/** Get a simplified string representation of this group. Typically for debugging. */
 	FString ToSimpleString() const;
 
-	/** Const access to active moves */
-	TArray<TSharedPtr<FLayeredMoveBase>>::TConstIterator GetActiveMovesIterator() const;
-	
-	/** Const access to queued moves */
-	TArray<TSharedPtr<FLayeredMoveBase>>::TConstIterator GetQueuedMovesIterator() const;
+	const TArray<TSharedPtr<FLayeredMoveBase>>& GetActiveMoves() const { return ActiveLayeredMoves; }
+	const TArray<TSharedPtr<FLayeredMoveBase>>& GetQueuedMoves() const { return QueuedLayeredMoves; }
 
+	/** Returns the first active layered move of the specified type, if one exists */
+	template <typename MoveType UE_REQUIRES(std::is_base_of_v<FLayeredMoveBase, MoveType>)>
+	const MoveType* FindActiveMove() const
+	{
+		return static_cast<const MoveType*>(FindActiveMove(MoveType::StaticStruct()));
+	}
+	const FLayeredMoveBase* FindActiveMove(const UScriptStruct* LayeredMoveStructType) const;
+
+	/** Returns the first queued layered move of the specified type, if one exists */
+	template <typename MoveType UE_REQUIRES(std::is_base_of_v<FLayeredMoveBase, MoveType>)>
+	const MoveType* FindQueuedMove() const
+	{
+		return static_cast<const MoveType*>(FindQueuedMove(MoveType::StaticStruct()));
+	}
+	const FLayeredMoveBase* FindQueuedMove(const UScriptStruct* LayeredMoveStructType) const;
+	
 	// Resets residual velocity related settings
 	void ResetResidualVelocity();
 
diff --git a/Engine/Plugins/Experimental/Mover/Source/Mover/Public/MoverComponent.h b/Engine/Plugins/Experimental/Mover/Source/Mover/Public/MoverComponent.h
--- a/Engine/Plugins/Experimental/Mover/Source/Mover/Public/MoverComponent.h
+++ b/Engine/Plugins/Experimental/Mover/Source/Mover/Public/MoverComponent.h
@@ -324,7 +324,7 @@ public:	// Queries
 	// Get the current movement mode name
 	UFUNCTION(BlueprintPure, Category = Mover)
 	FName GetMovementModeName() const;
-
+	
 	// Get the current movement mode 
 	UFUNCTION(BlueprintPure, Category = Mover)
 	const UBaseMovementMode* GetMovementMode() const;
@@ -369,13 +369,13 @@ public:	// Queries
 
 	/** Find settings object by type. Returns null if there is none of that type */
 	const IMovementSettingsInterface* FindSharedSettings(const UClass* ByType) const { return FindSharedSettings_Mutable(ByType); }
-	template<class T>
-	const T* FindSharedSettings() const { return Cast<const T>(FindSharedSettings(T::StaticClass())); }
+	template<typename SettingsT = IMovementSettingsInterface UE_REQUIRES(std::is_base_of_v<IMovementSettingsInterface, SettingsT>)>
+	const SettingsT* FindSharedSettings() const { return Cast<const SettingsT>(FindSharedSettings(SettingsT::StaticClass())); }
 
 	/** Find mutable settings object by type. Returns null if there is none of that type */
 	IMovementSettingsInterface* FindSharedSettings_Mutable(const UClass* ByType) const;
-	template<class T>
-	T* FindSharedSettings_Mutable() const { return Cast<T>(FindSharedSettings_Mutable(T::StaticClass())); }
+	template<typename SettingsT = IMovementSettingsInterface UE_REQUIRES(std::is_base_of_v<IMovementSettingsInterface, SettingsT>)>
+	SettingsT* FindSharedSettings_Mutable() const { return Cast<SettingsT>(FindSharedSettings_Mutable(SettingsT::StaticClass())); }
 
 	/** Find mutable settings object by type. Returns null if there is none of that type */
 	UFUNCTION(BlueprintPure, Category = Mover,  meta=(DeterminesOutputType="SharedSetting", DisplayName="Find Shared Settings Mutable"))
@@ -385,14 +385,30 @@ public:	// Queries
 	UFUNCTION(BlueprintPure, Category = Mover,  meta=(DeterminesOutputType="SharedSetting", DisplayName="Find Shared Settings"))
 	const UObject* FindSharedSettings_BP(UPARAM(meta = (MustImplement = "MovementSettingsInterface")) TSubclassOf<UObject> SharedSetting) const;
 
-	/** Find movement mode by type. Returns null if there is none of that type */
-	UBaseMovementMode* FindMode_Mutable(const UClass* ByType, bool bRequireExactClass=false) const;
-	template<class T>
-	T* FindMode_Mutable(bool bRequireExactClass=false) const { return Cast<T>(FindMode_Mutable(T::StaticClass(), bRequireExactClass)); }
+	/** Gets the currently active movement mode, provided it is of the given type. Returns nullptr if there is no active mode yet, or if it's of a different type. */
+	template<typename ModeT = UBaseMovementMode UE_REQUIRES(std::is_base_of_v<UBaseMovementMode, ModeT>)>
+	const ModeT* GetActiveMode(bool bRequireExactClass = false) const { return Cast<ModeT>(GetActiveModeInternal(ModeT::StaticClass(), bRequireExactClass)); }
 
+	/** Gets the currently active movement mode, provided it is of the given type. Returns nullptr if there is no active mode yet, or if it's of a different type. */
+	template<typename ModeT = UBaseMovementMode UE_REQUIRES(std::is_base_of_v<UBaseMovementMode, ModeT>)>
+	ModeT* GetActiveMode_Mutable(bool bRequireExactClass = false) const { return Cast<ModeT>(GetActiveModeInternal(ModeT::StaticClass(), bRequireExactClass)); }
+
+	/** Find the first movement mode on this component with the given type, optionally of the given type exactly. Returns null if there is none of that type */
+	template<typename ModeT = UBaseMovementMode UE_REQUIRES(std::is_base_of_v<UBaseMovementMode, ModeT>)>
+	ModeT* FindMode_Mutable(bool bRequireExactClass = false) const { return Cast<ModeT>(FindMode_Mutable(ModeT::StaticClass(), NAME_None, bRequireExactClass)); }
+	UBaseMovementMode* FindMode_Mutable(TSubclassOf<UBaseMovementMode> ModeType, bool bRequireExactClass = false) const;
+
+	/** Find the movement mode on this component the given name and type, optionally of the given type exactly. Returns null if there is no mode by that name, or if it's of a different type. */
+	template<typename ModeT = UBaseMovementMode UE_REQUIRES(std::is_base_of_v<UBaseMovementMode, ModeT>)>
+	ModeT* FindMode_Mutable(FName MovementModeName, bool bRequireExactClass = false) const { return Cast<ModeT>(FindMode_Mutable(ModeT::StaticClass(), MovementModeName, bRequireExactClass)); }
+	UBaseMovementMode* FindMode_Mutable(TSubclassOf<UBaseMovementMode> ModeType, FName ModeName, bool bRequireExactClass = false) const;
+	
 	UFUNCTION(BlueprintPure, Category = Mover,  meta=(DeterminesOutputType="MovementMode"))
 	UBaseMovementMode* FindMovementMode(TSubclassOf<UBaseMovementMode> MovementMode) const;
 
+	UFUNCTION(BlueprintPure, Category = Mover)
+	UBaseMovementMode* FindMovementModeByName(FName MovementModeName) const;
+	
 	/**
 	 * Retrieves an active layered move, by writing to a target instance if it is the matching type. Note: Writing to the struct returned will not modify the active struct.
 	 * @param DidSucceed			Flag indicating whether data was actually written to target struct instance
@@ -406,17 +422,8 @@ public:	// Queries
 	const FLayeredMoveBase* FindActiveLayeredMoveByType(const UScriptStruct* DataStructType) const;
 
 	/** Find a layered move of a specific type in this components active layered moves. If not found, null will be returned. */
-	template <typename T>
-	const T* FindActiveLayeredMoveByType() const
-	{
-		if (const FLayeredMoveBase* FoundData = FindActiveLayeredMoveByType(T::StaticStruct()))
-		{
-			return static_cast<const T*>(FoundData);
-		}
-
-		return nullptr;
-	}
-
+	template <typename MoveT = FLayeredMoveBase UE_REQUIRES(std::is_base_of_v<FLayeredMoveBase, MoveT>)>
+	const MoveT* FindActiveLayeredMoveByType() const { return static_cast<const MoveT*>(FindActiveLayeredMoveByType(MoveT::StaticStruct())); }
 	
 	/**
 	 * Retrieves Movement modifier by writing to a target instance if it is the matching type. Note: Writing to the struct returned will not modify the active struct.
@@ -439,16 +446,8 @@ public:	// Queries
 	const FMovementModifierBase* FindMovementModifierByType(const UScriptStruct* DataStructType) const;
 	
 	/** Find a movement modifier of a specific type in this components movement modifiers. If not found, null will be returned. */
-	template <typename T>
-	const T* FindMovementModifierByType() const
-	{
-		if (const FMovementModifierBase* FoundData = FindMovementModifierByType(T::StaticStruct()))
-		{
-			return static_cast<const T*>(FoundData);
-		}
-
-		return nullptr;
-	}
+	template <typename ModifierT = FMovementModifierBase UE_REQUIRES(std::is_base_of_v<FMovementModifierBase, ModifierT>)>
+	const ModifierT* FindMovementModifierByType() const { return static_cast<const ModifierT*>(FindMovementModifierByType(ModifierT::StaticStruct())); }
 	
 	/**
  	 * Check Mover systems for a gameplay tag.
@@ -514,6 +513,8 @@ protected:
 	/** internal function to perform post-sim scheduling to optionally support simple based movement */
 	void UpdateBasedMovementScheduling(const FMoverTickEndData& SimOutput);
 
+	UBaseMovementMode* GetActiveModeInternal(TSubclassOf<UBaseMovementMode> ModeType, bool bRequireExactClass = false) const;
+
 	TObjectPtr<UPrimitiveComponent> MovementBaseDependency;	// used internally for based movement scheduling management
 	
 	/** internal function to ensure SharedSettings array matches what's needed by the list of Movement Modes */
@@ -547,7 +548,8 @@ protected:
 	FMoverTimeStep CachedLastSimTickTimeStep;	// Saved timestep info from our last simulation tick, used during rollback handling. This will rewind during corrections.
 	FMoverTimeStep CachedNewestSimTickTimeStep;	// Saved timestep info from the newest (farthest-advanced) simulation tick. This will not rewind during corrections.
 
-	TWeakInterfacePtr<IMoverBackendLiaisonInterface> BackendLiaisonComp;
+	UPROPERTY(Transient)
+	TScriptInterface<IMoverBackendLiaisonInterface> BackendLiaisonComp;
 
 	/** Tick function that may be called anytime after this actor's movement step, useful as a way to support based movement on objects that are not */
 	FMoverDynamicBasedMovementTickFunction BasedMovementTickFunction;
