Physics Mover - Add support for mass conditioning

diff --git a/Engine/Plugins/Experimental/Mover/Source/Mover/Private/Backends/MoverNetworkPhysicsLiaison.cpp b/Engine/Plugins/Experimental/Mover/Source/Mover/Private/Backends/MoverNetworkPhysicsLiaison.cpp
--- a/Engine/Plugins/Experimental/Mover/Source/Mover/Private/Backends/MoverNetworkPhysicsLiaison.cpp
+++ b/Engine/Plugins/Experimental/Mover/Source/Mover/Private/Backends/MoverNetworkPhysicsLiaison.cpp
@@ -366,6 +366,7 @@ void UMoverNetworkPhysicsLiaisonComponent::SetupConstraint()
 							// The rest get set every frame depending on the current movement mode
 							Constraint->SetCosMaxWalkableSlopeAngle(CommonMovementSettings->MaxWalkSlopeCosine);
 							Constraint->SetVerticalAxis(MoverComp->GetUpDirection());
+							Constraint->SetMaxCharacterGroundMassRatio(GPhysicsDrivenMotionDebugParams.MaxCharacterGroundMassRatio);
 
 							// Enable Physics Simulation
 							MoverComp->UpdatedCompAsPrimitive->SetSimulatePhysics(true);
@@ -583,6 +584,7 @@ void UMoverNetworkPhysicsLiaisonComponent::UpdateConstraintSettings()
 	{
 		Constraint->SetVerticalAxis(MoverComp->GetUpDirection());
 		Constraint->SetCosMaxWalkableSlopeAngle(CommonMovementSettings->MaxWalkSlopeCosine);
+		Constraint->SetMaxCharacterGroundMassRatio(GPhysicsDrivenMotionDebugParams.MaxCharacterGroundMassRatio);
 
 		const UBaseMovementMode* CurrentMode = MoverComp->ModeFSM->GetCurrentMode();
 		if (CurrentMode && CurrentMode->Implements<UPhysicsCharacterMovementModeInterface>())
diff --git a/Engine/Plugins/Experimental/Mover/Source/Mover/Private/PhysicsMover/PhysicsMoverSimulationTypes.cpp b/Engine/Plugins/Experimental/Mover/Source/Mover/Private/PhysicsMover/PhysicsMoverSimulationTypes.cpp
--- a/Engine/Plugins/Experimental/Mover/Source/Mover/Private/PhysicsMover/PhysicsMoverSimulationTypes.cpp
+++ b/Engine/Plugins/Experimental/Mover/Source/Mover/Private/PhysicsMover/PhysicsMoverSimulationTypes.cpp
@@ -20,6 +20,9 @@ FAutoConsoleVariableRef CVarPhysicsDrivenMotionTeleportThreshold(TEXT("p.mover.p
 FAutoConsoleVariableRef CVarPhysicsDrivenMotionMinStepUpDistance(TEXT("p.mover.physics.MinStepUpDistance"),
 	GPhysicsDrivenMotionDebugParams.MinStepUpDistance, TEXT("Minimum distance that will be considered a step up."));
 
+FAutoConsoleVariableRef CVarPhysicsDrivenMotionMaxCharacterGroundMassRatio(TEXT("p.mover.physics.MaxCharacterGroundMassRatio"),
+	GPhysicsDrivenMotionDebugParams.MaxCharacterGroundMassRatio, TEXT("Maximum ratio between character mass and ground mass as seen by the ground constraint."));
+
 //////////////////////////////////////////////////////////////////////////
 // FMovementSettingsInputs
 
diff --git a/Engine/Plugins/Experimental/Mover/Source/Mover/Public/PhysicsMover/PhysicsMoverSimulationTypes.h b/Engine/Plugins/Experimental/Mover/Source/Mover/Public/PhysicsMover/PhysicsMoverSimulationTypes.h
--- a/Engine/Plugins/Experimental/Mover/Source/Mover/Public/PhysicsMover/PhysicsMoverSimulationTypes.h
+++ b/Engine/Plugins/Experimental/Mover/Source/Mover/Public/PhysicsMover/PhysicsMoverSimulationTypes.h
@@ -27,6 +27,7 @@ struct FPhysicsDrivenMotionDebugParams
 {
 	float TeleportThreshold = 100.0f;
 	float MinStepUpDistance = 5.0f;
+	float MaxCharacterGroundMassRatio = 1.5f;
 	bool EnableMultithreading = true;
 	bool DebugDrawGroundQueries = false;
 };
