[Chaos Mover] Splitting out a MoverNetworkPhysicsLiaisonBase from the existing (character-centric) MoverNetworkPhysicsLiaison
- UMoverNetworkPhysicsLiaisonComponent is now a subclass that tacks on the character movement-specific logic and is pending a rename to UMoverCharacterPhysicsLiaisonComponent
- Forthcoming in other CLs are liaisons for physics-based kinematic and pathed movement

Additional changes beyond simply splitting things up:
- Added lots of comments to (hopefully) make things more approachable to others who come in as ignorant as I did
- Most of the Internal/External sim callback methods have a separate CanX virtual method that gets checked first, then a virtual PerformX called if CanX passes. This is to avoid the need for redundant validation checks in child overrides (see overrides of UMovementComponent::TickComponent() for examples of what I'm seeking to avoid)
- Support for scenarios where input for the physics-based movement originates on the server (i.e. all cases that don't involve a player-posessed pawn)
- FPhysicsNetworkMoverData structs implement CompareData by forwarding it to their respective FMoverDataCollection's ShouldReconcile()
- FMoverDataStructBase now complains about missing a ShouldReconcile(), Interpolate(), and Merge() override on types that comprise physics-based input data. Ensure messages and comments have been updated accordingly.

Still to do:
- Rename to UMoverNetworkPhysicsLiaisonComponent to UMoverCharacterPhysicsLiaisonComponent
- We probably don't need to involve FMoverDefaultSyncState at all in the base physics liaison, aside from adding to the CachedLatestSyncState that gets set on the MoverComp
- UMoverComponent::OnMovementModeChanged still gets triggered twice (once on each thread)

diff --git a/Engine/Plugins/Experimental/Mover/Source/Mover/Private/Backends/MoverNetworkPhysicsLiaison.cpp b/Engine/Plugins/Experimental/Mover/Source/Mover/Private/Backends/MoverNetworkPhysicsLiaison.cpp
--- a/Engine/Plugins/Experimental/Mover/Source/Mover/Private/Backends/MoverNetworkPhysicsLiaison.cpp
+++ b/Engine/Plugins/Experimental/Mover/Source/Mover/Private/Backends/MoverNetworkPhysicsLiaison.cpp
@@ -4,379 +4,69 @@
 
 #include "Chaos/Character/CharacterGroundConstraint.h"
 #include "Chaos/PhysicsObject.h"
-#include "Chaos/PhysicsObjectInternalInterface.h"
-#include "Components/ShapeComponent.h"
-#include "Framework/Threading.h"
-#include "GameFramework/Actor.h"
-#include "GameFramework/Pawn.h"
 #include "MovementModeStateMachine.h"
 #include "DefaultMovementSet/Settings/CommonLegacyMovementSettings.h"
 #include "PBDRigidsSolver.h"
+#include "Chaos/PhysicsObjectInterface.h"
+#include "Chaos/PhysicsObjectInternalInterface.h"
 #include "Physics/Experimental/PhysScene_Chaos.h"
-#include "PhysicsEngine/PhysicsObjectExternalInterface.h"
-#include "PhysicsMover/Modes/PhysicsDrivenFallingMode.h"
-#include "PhysicsMover/Modes/PhysicsDrivenWalkingMode.h"
 #include "PhysicsMover/PhysicsMovementUtils.h"
-#include "PhysicsMover/PhysicsMoverManager.h"
+#include "PhysicsMover/Modes/PhysicsDrivenWalkingMode.h"
 #include "PhysicsProxy/CharacterGroundConstraintProxy.h"
+#include "PhysicsMover/PhysicsMoverSimulationTypes.h"
+#include "DefaultMovementSet/Modes/FallingMode.h"
+#include "Misc/DataValidation.h"
+
 
 #if WITH_EDITOR
-#include "Misc/DataValidation.h"
 #endif
 
 #define LOCTEXT_NAMESPACE "Mover"
 
-//////////////////////////////////////////////////////////////////////////
-
-extern FPhysicsDrivenMotionDebugParams GPhysicsDrivenMotionDebugParams;
-
-//////////////////////////////////////////////////////////////////////////
-// FNetworkPhysicsMoverInputs
-
-void FNetworkPhysicsMoverInputs::ApplyData(UActorComponent* NetworkComponent) const
-{
-	if (NetworkComponent)
-	{
-		if (UMoverNetworkPhysicsLiaisonComponent* LiaisonComp = Cast<UMoverNetworkPhysicsLiaisonComponent>(NetworkComponent))
-		{
-			LiaisonComp->SetCurrentInputData(InputCmdContext);
-		}
-	}
-}
-
-void FNetworkPhysicsMoverInputs::BuildData(const UActorComponent* NetworkComponent)
-{
-	if (NetworkComponent)
-	{
-		if (const UMoverNetworkPhysicsLiaisonComponent* LiaisonComp = Cast<const UMoverNetworkPhysicsLiaisonComponent>(NetworkComponent))
-		{
-			LiaisonComp->GetCurrentInputData(InputCmdContext);
-		}
-	}
-}
-
-bool FNetworkPhysicsMoverInputs::NetSerialize(FArchive& Ar, class UPackageMap* PackageMap, bool& bOutSuccess)
-{
-	FNetworkPhysicsData::SerializeFrames(Ar);
-
-	if (PackageMap)
-	{
-		InputCmdContext.NetSerialize(FNetSerializeParams(Ar));
-		bOutSuccess = true;
-	}
-	else
-	{
-		bOutSuccess = false;
-	}
-	
-	return bOutSuccess;
-}
-
-void FNetworkPhysicsMoverInputs::InterpolateData(const FNetworkPhysicsData& MinData, const FNetworkPhysicsData& MaxData)
-{
-	const FNetworkPhysicsMoverInputs& MinDataInput = static_cast<const FNetworkPhysicsMoverInputs&>(MinData);
-	const FNetworkPhysicsMoverInputs& MaxDataInput = static_cast<const FNetworkPhysicsMoverInputs&>(MaxData);
-
-	const float LerpFactor = (LocalFrame - MinDataInput.LocalFrame) / (MaxDataInput.LocalFrame - MinDataInput.LocalFrame);
-
-	const FCharacterDefaultInputs* MinInput = MinDataInput.InputCmdContext.InputCollection.FindDataByType<FCharacterDefaultInputs>();
-	const FCharacterDefaultInputs* MaxInput = MaxDataInput.InputCmdContext.InputCollection.FindDataByType<FCharacterDefaultInputs>();
-
-	FCharacterDefaultInputs& LocalInput = InputCmdContext.InputCollection.FindOrAddMutableDataByType<FCharacterDefaultInputs>();
-
-	if (MinInput && MaxInput)
-	{
-		// Note, this ignores movement base as this is not used by the physics mover
-		const FCharacterDefaultInputs* ClosestInputs = LerpFactor < 0.5f ? MinInput : MaxInput;
-		LocalInput.bIsJumpJustPressed = ClosestInputs->bIsJumpJustPressed;
-		LocalInput.bIsJumpPressed = ClosestInputs->bIsJumpPressed;
-		LocalInput.SuggestedMovementMode = ClosestInputs->SuggestedMovementMode;
-
-		LocalInput.SetMoveInput(ClosestInputs->GetMoveInputType(), FMath::Lerp(MinInput->GetMoveInput(), MaxInput->GetMoveInput(), LerpFactor));
-		LocalInput.OrientationIntent = FMath::Lerp(MinInput->OrientationIntent, MaxInput->OrientationIntent, LerpFactor);
-		LocalInput.ControlRotation = FMath::Lerp(MinInput->ControlRotation, MaxInput->ControlRotation, LerpFactor);
-
-	}
-	else if (MinInput)
-	{
-		LocalInput = *MinInput;
-	}
-	else if (MaxInput)
-	{
-		LocalInput = *MaxInput;
-	}
-}
-
-void FNetworkPhysicsMoverInputs::MergeData(const FNetworkPhysicsData& FromData)
-{
-	const FNetworkPhysicsMoverInputs& FromDataInput = static_cast<const FNetworkPhysicsMoverInputs&>(FromData);
-
-	if (const FCharacterDefaultInputs* FromInput = FromDataInput.InputCmdContext.InputCollection.FindDataByType<FCharacterDefaultInputs>())
-	{
-		FCharacterDefaultInputs& LocalInputs = InputCmdContext.InputCollection.FindOrAddMutableDataByType<FCharacterDefaultInputs>();
-
-		LocalInputs.bIsJumpJustPressed |= FromInput->bIsJumpJustPressed;
-		LocalInputs.bIsJumpPressed |= FromInput->bIsJumpPressed;
-	}
-}
-
-void FNetworkPhysicsMoverInputs::ValidateData(const UActorComponent* NetworkComponent)
-{
-	if (NetworkComponent)
-	{
-		if (const UMoverNetworkPhysicsLiaisonComponent* LiaisonComp = Cast<UMoverNetworkPhysicsLiaisonComponent>(NetworkComponent))
-		{
-			LiaisonComp->ValidateInputData(InputCmdContext);
-		}
-	}
-}
-
-const FString FNetworkPhysicsMoverInputs::DebugData()
-{
-	const FCharacterDefaultInputs Input = InputCmdContext.InputCollection.FindOrAddDataByType<FCharacterDefaultInputs>();
-
-	return FString::Printf(TEXT("FNetworkPhysicsMoverInputs | MoveInput = %s OrientationInput = %s"),
-			*Input.GetMoveInput_WorldSpace().ToString(),
-			*Input.GetOrientationIntentDir_WorldSpace().ToString());
-}
-
-//////////////////////////////////////////////////////////////////////////
-// FNetworkPhysicsMoverState
-
-void FNetworkPhysicsMoverState::ApplyData(UActorComponent* NetworkComponent) const
-{
-	if (NetworkComponent)
-	{
-		if (UMoverNetworkPhysicsLiaisonComponent* LiaisonComp = Cast<UMoverNetworkPhysicsLiaisonComponent>(NetworkComponent))
-		{
-			LiaisonComp->SetCurrentStateData(SyncStateContext);
-		}
-	}
-}
-
-void FNetworkPhysicsMoverState::BuildData(const UActorComponent* NetworkComponent)
-{
-	if (NetworkComponent)
-	{
-		if (const UMoverNetworkPhysicsLiaisonComponent* LiaisonComp = Cast<const UMoverNetworkPhysicsLiaisonComponent>(NetworkComponent))
-		{
-			LiaisonComp->GetCurrentStateData(SyncStateContext);
-		}
-	}
-}
-
-bool FNetworkPhysicsMoverState::NetSerialize(FArchive& Ar, class UPackageMap* PackageMap, bool& bOutSuccess)
-{
-	FNetworkPhysicsData::SerializeFrames(Ar);
-
-	if (PackageMap)
-	{
-		FNetSerializeParams Params(Ar);
-		SyncStateContext.NetSerialize(Params);
-		bOutSuccess = true;
-	}
-	else
-	{
-		bOutSuccess = false;
-	}
-
-	return bOutSuccess;
-}
-
-void FNetworkPhysicsMoverState::InterpolateData(const FNetworkPhysicsData& MinData, const FNetworkPhysicsData& MaxData)
-{
-	const FNetworkPhysicsMoverState& MinState = static_cast<const FNetworkPhysicsMoverState&>(MinData);
-	const FNetworkPhysicsMoverState& MaxState = static_cast<const FNetworkPhysicsMoverState&>(MaxData);
-
-	const float LerpFactor = (LocalFrame - MinState.LocalFrame) / (MaxState.LocalFrame - MinState.LocalFrame);
-	SyncStateContext.Interpolate(&MinState.SyncStateContext, &MaxState.SyncStateContext, LerpFactor);
-}
-
-const FString FNetworkPhysicsMoverState::DebugData()
-{
-	const FMoverDefaultSyncState SyncState = SyncStateContext.SyncStateCollection.FindOrAddDataByType<FMoverDefaultSyncState>();
-
-	return FString::Printf(TEXT("FNetworkPhysicsMoverState | Location = %s Velocity = %s MovementBase = %s"),
-			*SyncState.GetLocation_WorldSpace().ToString(),
-			*SyncState.GetVelocity_WorldSpace().ToString(),
-			*SyncState.GetMovementBase()->GetName());
-}
-
 //////////////////////////////////////////////////////////////////////////
 // UMoverNetworkPhysicsLiaisonComponent
 
-void UMoverNetworkPhysicsLiaisonComponent::GetCurrentInputData(OUT FMoverInputCmdContext& InputCmd) const
-{
-	InputCmd = LatestInputCmd;
-}
-
-void UMoverNetworkPhysicsLiaisonComponent::GetCurrentStateData(OUT FMoverSyncState& SyncState) const
-{
-	SyncState = LatestSyncState;
-}
-
-void UMoverNetworkPhysicsLiaisonComponent::SetCurrentInputData(const FMoverInputCmdContext& InputCmd)
-{
-	LatestInputCmd = InputCmd;
-}
-
-void UMoverNetworkPhysicsLiaisonComponent::SetCurrentStateData(const FMoverSyncState& SyncState)
-{
-	LatestSyncState = SyncState;
-}
-
-bool UMoverNetworkPhysicsLiaisonComponent::ValidateInputData(FMoverInputCmdContext& InputCmd) const
-{
-	// TODO - proper data validation
-	return true;
-}
-
 UMoverNetworkPhysicsLiaisonComponent::UMoverNetworkPhysicsLiaisonComponent()
 {
-	PrimaryComponentTick.TickGroup = TG_PrePhysics;
-	PrimaryComponentTick.bCanEverTick = true;
-
-	bWantsInitializeComponent = true;
-	bAutoActivate = true;
-
-	if (UWorld* World = GetWorld())
-	{
-		if (FPhysScene_Chaos* PhysScene = GetWorld()->GetPhysicsScene())
-		{
-			if (Chaos::FPhysicsSolver* Solver = PhysScene->GetSolver())
-			{
-				bUsingAsyncPhysics = Solver->IsUsingAsyncResults();
-			}
-		}
-	}
-
-	if (Chaos::FPhysicsSolverBase::IsNetworkPhysicsPredictionEnabled())
-	{
-		// Network physics relies on movement being replicated
-		SetIsReplicatedByDefault(true);
-
-		if (AActor* MyActor = GetOwner())
-		{
-			MyActor->SetReplicatingMovement(true);
-			MyActor->SetReplicateMovement(true);
-		}
-	}
-}
-
-//////////////////////////////////////////////////////////////////////////
-//  UMoverNetworkPhysicsLiaisonComponent IMoverBackendLiaisonInterface
-
-float UMoverNetworkPhysicsLiaisonComponent::GetCurrentSimTimeMs()
-{
-	if (FPhysScene* Scene = GetWorld()->GetPhysicsScene())
-	{
-		if (Chaos::FPhysicsSolver* Solver = Scene->GetSolver())
-		{
-			return bUsingAsyncPhysics ? Solver->GetAsyncDeltaTime() * GetCurrentSimFrame() * 1000.0f : Solver->GetSolverTime() * 1000.0f;
-		}
-	}
-
-	return 0.0f;
-}
-
-int32 UMoverNetworkPhysicsLiaisonComponent::GetCurrentSimFrame()
-{
-	if (FPhysScene* Scene = GetWorld()->GetPhysicsScene())
-	{
-		if (Chaos::FPhysicsSolver* Solver = Scene->GetSolver())
-		{
-			return Solver->GetCurrentFrame() + GetNetworkPhysicsTickOffset();
-		}
-	}
-
-	return 0;
-}
-
-#if WITH_EDITOR
-EDataValidationResult UMoverNetworkPhysicsLiaisonComponent::ValidateData(FDataValidationContext& Context, const UMoverComponent& ValidationMoverComp) const
-{
-	if (const AActor* OwnerActor = ValidationMoverComp.GetOwner())
-	{
-		if (!OwnerActor->IsReplicatingMovement())
-		{
-			Context.AddError(FText::Format(LOCTEXT("RequiresReplicateMovementProperty", "The owning actor ({0}) does not have the ReplicateMovement property enabled. This is required for use with Chaos Networked Physics and poor quality movement with occur without it. Please enable it."),
-				FText::FromString(GetNameSafe(OwnerActor))));
-
-			return EDataValidationResult::Invalid;
-		}
-	}
-
-	return EDataValidationResult::Valid;
 }
-#endif // WITH_EDITOR
-
-//////////////////////////////////////////////////////////////////////////
-// UMoverNetworkPhysicsLiaisonComponent UObject interface
 
 void UMoverNetworkPhysicsLiaisonComponent::OnRegister()
 {
 	Super::OnRegister();
 
-	// Need to set this here as physics creation requires access to MoverComp
-	
-	if ((MoverComp = GetOwner()->FindComponentByClass<UMoverComponent>()) != nullptr)
-	{
-	
-		CommonMovementSettings = MoverComp->FindSharedSettings_Mutable<UCommonLegacyMovementSettings>();
-		check(CommonMovementSettings);
-
-		if (MoverComp->UpdatedCompAsPrimitive)
-		{
-			MoverComp->UpdatedCompAsPrimitive->OnComponentPhysicsStateChanged.AddUniqueDynamic(this, &ThisClass::OnComponentPhysicsStateChanged);
-		}
-	}
+	CommonMovementSettings = GetMoverComponent().FindSharedSettings_Mutable<UCommonLegacyMovementSettings>();
+	check(CommonMovementSettings);
 }
 
-void UMoverNetworkPhysicsLiaisonComponent::OnUnregister()
-{
-	if (MoverComp && MoverComp->UpdatedCompAsPrimitive)
-	{
-		MoverComp->UpdatedCompAsPrimitive->OnComponentPhysicsStateChanged.RemoveDynamic(this, &ThisClass::OnComponentPhysicsStateChanged);
-	}
-
-	Super::OnUnregister();
-}
+//////////////////////////////////////////////////////////////////////////
+// UMoverNetworkPhysicsLiaisonComponent UObject interface
 
 void UMoverNetworkPhysicsLiaisonComponent::SetupConstraint()
 {
-	UWorld* World = GetWorld();
-	if (World && World->IsGameWorld())
+	if (Chaos::FPhysicsSolver* Solver = GetPhysicsSolver())
 	{
-		if (FPhysScene_Chaos* PhysScene = World->GetPhysicsScene())
+		const UMoverComponent& MoverComp = GetMoverComponent();
+		if (FBodyInstance* BI = MoverComp.UpdatedCompAsPrimitive ? MoverComp.UpdatedCompAsPrimitive->GetBodyInstance() : nullptr)
 		{
-			if (Chaos::FPhysicsSolver* Solver = PhysScene->GetSolver())
+			if (Chaos::FSingleParticlePhysicsProxy* CharacterProxy = BI->ActorHandle)
 			{
-				if (MoverComp && MoverComp->UpdatedCompAsPrimitive)
-				{
-					if (FBodyInstance* BI = MoverComp->UpdatedCompAsPrimitive->GetBodyInstance())
-					{
-						if (Chaos::FSingleParticlePhysicsProxy* CharacterProxy = BI->ActorHandle)
-						{
-							// Create and register the constraint
-							Constraint = MakeUnique<Chaos::FCharacterGroundConstraint>();
-							Constraint->Init(CharacterProxy);
-							Solver->RegisterObject(Constraint.Get());
-
-							// Set the common settings
-							// The rest get set every frame depending on the current movement mode
-							Constraint->SetCosMaxWalkableSlopeAngle(CommonMovementSettings->MaxWalkSlopeCosine);
-							Constraint->SetVerticalAxis(MoverComp->GetUpDirection());
-							Constraint->SetMaxCharacterGroundMassRatio(GPhysicsDrivenMotionDebugParams.MaxCharacterGroundMassRatio);
-
-							// Enable Physics Simulation
-							MoverComp->UpdatedCompAsPrimitive->SetSimulatePhysics(true);
-							
-							// Turn off sleeping
-							Chaos::FRigidBodyHandle_External& PhysicsBody = CharacterProxy->GetGameThreadAPI();
-							PhysicsBody.SetSleepType(Chaos::ESleepType::NeverSleep);
-						}
-					}
-				}
+				// Create and register the constraint
+				Constraint = MakeUnique<Chaos::FCharacterGroundConstraint>();
+				Constraint->Init(CharacterProxy);
+				Solver->RegisterObject(Constraint.Get());
+
+				// Set the common settings
+				// The rest get set every frame depending on the current movement mode
+				Constraint->SetCosMaxWalkableSlopeAngle(CommonMovementSettings->MaxWalkSlopeCosine);
+				Constraint->SetVerticalAxis(MoverComp.GetUpDirection());
+				Constraint->SetMaxCharacterGroundMassRatio(GPhysicsDrivenMotionDebugParams.MaxCharacterGroundMassRatio);
+
+				// Enable Physics Simulation
+				MoverComp.UpdatedCompAsPrimitive->SetSimulatePhysics(true);
+
+				// Turn off sleeping
+				Chaos::FRigidBodyHandle_External& PhysicsBody = CharacterProxy->GetGameThreadAPI();
+				PhysicsBody.SetSleepType(Chaos::ESleepType::NeverSleep);
 			}
 		}
 	}
@@ -398,7 +88,7 @@ void UMoverNetworkPhysicsLiaisonComponent::DestroyConstraint()
 	}
 }
 
-void UMoverNetworkPhysicsLiaisonComponent::OnComponentPhysicsStateChanged(UPrimitiveComponent* ChangedComponent, EComponentPhysicsStateChange StateChange)
+void UMoverNetworkPhysicsLiaisonComponent::HandleComponentPhysicsStateChanged(UPrimitiveComponent* ChangedComponent, EComponentPhysicsStateChange StateChange)
 {
 	if (StateChange == EComponentPhysicsStateChange::Destroyed)
 	{
@@ -408,79 +98,10 @@ void UMoverNetworkPhysicsLiaisonComponent::OnComponentPhysicsStateChanged(UPrimi
 	{
 		SetupConstraint();
 
-		if (MoverComp && MoverComp->ModeFSM)
-		{
-			MoverComp->ModeFSM->SetModeImmediately(MoverComp->StartingMovementMode);
-		}
-
 		InitializeCachedSyncState();
 	}
-}
-
-void UMoverNetworkPhysicsLiaisonComponent::InitializeComponent()
-{
-	Super::InitializeComponent();
-
-	UWorld* World = GetWorld();
-	if (World && World->IsGameWorld())
-	{
-		if (ensureAlwaysMsgf(MoverComp, TEXT("UMoverNetworkPhysicsLiaisonComponent on actor %s failed to find associated Mover component. This actor's movement will not be simulated. Verify its setup."), *GetNameSafe(GetOwner())))
-		{
-			MoverComp->InitMoverSimulation();
-			MoverComp->ModeFSM->SetModeImmediately(MoverComp->StartingMovementMode);
-		}
-
-		if (Chaos::FPhysicsSolverBase::IsNetworkPhysicsPredictionEnabled())
-		{
-			NetworkPhysicsComponent = NewObject<UNetworkPhysicsComponent>(GetOwner(), TEXT("PhysMover_NetworkPhysicsComponent"));
-			if (NetworkPhysicsComponent)
-			{
-				NetworkPhysicsComponent->SetNetAddressable(); // Make DSO components net addressable
-				NetworkPhysicsComponent->SetIsReplicated(true);
-				NetworkPhysicsComponent->RegisterComponent();
-				if (!NetworkPhysicsComponent->HasBeenInitialized())
-				{
-					NetworkPhysicsComponent->InitializeComponent();
-				}
-				NetworkPhysicsComponent->Activate(true);
-
-				// Register network data for recording and rewind/resim
-				NetworkPhysicsComponent->CreateDataHistory<FNetworkPhysicsMoverTraits>(this);
-			}
-		}
-	}
-}
-
-void UMoverNetworkPhysicsLiaisonComponent::UninitializeComponent()
-{
-	if (NetworkPhysicsComponent)
-	{
-		NetworkPhysicsComponent->RemoveDataHistory();
-		NetworkPhysicsComponent->DestroyComponent();
-	}
 
-	Super::UninitializeComponent();
-}
-
-bool UMoverNetworkPhysicsLiaisonComponent::ShouldCreatePhysicsState() const
-{
-	if (!IsRegistered() || IsBeingDestroyed())
-	{
-		return false;
-	}
-
-	UWorld* World = GetWorld();
-	if (World->IsGameWorld())
-	{
-		FPhysScene* PhysScene = World->GetPhysicsScene();
-
-		if (PhysScene && CanCreatePhysics())
-		{
-			return true;
-		}
-	}
-
-	return false;
+	Super::HandleComponentPhysicsStateChanged(ChangedComponent, StateChange);
 }
 
 bool UMoverNetworkPhysicsLiaisonComponent::HasValidPhysicsState() const
@@ -488,13 +109,6 @@ bool UMoverNetworkPhysicsLiaisonComponent::HasValidPhysicsState() const
 	return Constraint.IsValid() && Constraint->IsValid();
 }
 
-bool UMoverNetworkPhysicsLiaisonComponent::HasValidState() const
-{
-	return HasValidPhysicsState() && MoverComp && MoverComp->UpdatedCompAsPrimitive && MoverComp->UpdatedComponent
-		&& MoverComp->ModeFSM->IsValidLowLevelFast() && MoverComp->SimBlackboard->IsValidLowLevelFast()
-		&& MoverComp->InputProducer && MoverComp->MovementMixer;
-}
-
 void UMoverNetworkPhysicsLiaisonComponent::OnCreatePhysicsState()
 {
 	Super::OnCreatePhysicsState();
@@ -505,88 +119,40 @@ void UMoverNetworkPhysicsLiaisonComponent::OnCreatePhysicsState()
 void UMoverNetworkPhysicsLiaisonComponent::OnDestroyPhysicsState()
 {
 	DestroyConstraint();
-	
-	Super::OnDestroyPhysicsState();
-}
-
-bool UMoverNetworkPhysicsLiaisonComponent::CanCreatePhysics() const
-{
-	check(GetOwner());
-	FString ActorName = GetOwner()->GetName();
-
-	if (!IsValid(MoverComp->UpdatedComponent))
-	{
-		UE_LOG(LogMover, Warning, TEXT("Can't create physics %s (%s). UpdatedComponent is not set."), *ActorName, *GetPathName());
-		return false;
-	}
-
-	if (!IsValid(MoverComp->UpdatedCompAsPrimitive))
-	{
-		UE_LOG(LogMover, Warning, TEXT("Can't create physics %s (%s). UpdatedComponent is not a PrimitiveComponent."), *ActorName, *GetPathName());
-		return false;
-	}
 
-	return true;
+	Super::OnDestroyPhysicsState();
 }
 
-void UMoverNetworkPhysicsLiaisonComponent::BeginPlay()
+#if WITH_EDITOR
+EDataValidationResult UMoverNetworkPhysicsLiaisonComponent::ValidateData(FDataValidationContext& Context, const UMoverComponent& ValidationMoverComp) const
 {
-	Super::BeginPlay();
-
-	// Register with the physics mover manager
-	UWorld* World = GetWorld();
-	if (World && World->IsGameWorld())
+	if (const AActor* OwnerActor = ValidationMoverComp.GetOwner())
 	{
-		if (UPhysicsMoverManager* Manager = World->GetSubsystem<UPhysicsMoverManager>())
+		if (!OwnerActor->IsReplicatingMovement())
 		{
-			Manager->RegisterPhysicsMoverComponent(this);
-		}
-	}
-
-	InitializeCachedSyncState();
-}
+			Context.AddError(FText::Format(LOCTEXT("RequiresReplicateMovementProperty", "The owning actor ({0}) does not have the ReplicateMovement property enabled. This is required for use with Chaos Networked Physics and poor quality movement with occur without it. Please enable it."),
+				FText::FromString(GetNameSafe(OwnerActor))));
 
-void UMoverNetworkPhysicsLiaisonComponent::EndPlay(const EEndPlayReason::Type EndPlayReason)
-{
-	// Unregister with the physics mover manager
-	if (UWorld* World = GetWorld())
-	{
-		if (UPhysicsMoverManager* Manager = World->GetSubsystem<UPhysicsMoverManager>())
-		{
-			Manager->UnregisterPhysicsMoverComponent(this);
+			return EDataValidationResult::Invalid;
 		}
 	}
 
-	Super::EndPlay(EndPlayReason);
+	return Super::ValidateData(Context, ValidationMoverComp);
 }
+#endif
 
 //////////////////////////////////////////////////////////////////////////
 
-Chaos::FUniqueIdx UMoverNetworkPhysicsLiaisonComponent::GetUniqueIdx() const
-{
-	if (MoverComp && MoverComp->UpdatedCompAsPrimitive)
-	{
-		if (FBodyInstance* BI = MoverComp->UpdatedCompAsPrimitive->GetBodyInstance())
-		{
-			if (FPhysicsActorHandle ActorHandle = BI->ActorHandle)
-			{
-				return ActorHandle->GetGameThreadAPI().UniqueIdx();
-			}
-		}
-	}
-
-	return Chaos::FUniqueIdx();
-}
-
 void UMoverNetworkPhysicsLiaisonComponent::UpdateConstraintSettings()
 {
 	if (HasValidState())
 	{
-		Constraint->SetVerticalAxis(MoverComp->GetUpDirection());
+		const UMoverComponent& MoverComp = GetMoverComponent();
+		Constraint->SetVerticalAxis(MoverComp.GetUpDirection());
 		Constraint->SetCosMaxWalkableSlopeAngle(CommonMovementSettings->MaxWalkSlopeCosine);
 		Constraint->SetMaxCharacterGroundMassRatio(GPhysicsDrivenMotionDebugParams.MaxCharacterGroundMassRatio);
 
-		const UBaseMovementMode* CurrentMode = MoverComp->ModeFSM->GetCurrentMode();
+		const UBaseMovementMode* CurrentMode = MoverComp.ModeFSM->GetCurrentMode();
 		if (CurrentMode && CurrentMode->Implements<UPhysicsCharacterMovementModeInterface>())
 		{
 			const IPhysicsCharacterMovementModeInterface* PhysicsMode = CastChecked<IPhysicsCharacterMovementModeInterface>(CurrentMode);
@@ -595,323 +161,82 @@ void UMoverNetworkPhysicsLiaisonComponent::UpdateConstraintSettings()
 	}
 }
 
-int32 UMoverNetworkPhysicsLiaisonComponent::GetNetworkPhysicsTickOffset() const
-{
-	if (NetworkPhysicsComponent && !NetworkPhysicsComponent->HasServerWorld())
-	{
-		const APlayerController* PlayerController = NetworkPhysicsComponent->GetPlayerController();
-		if (!PlayerController)
-		{
-			PlayerController = GetWorld()->GetFirstPlayerController();
-		}
-
-		if (PlayerController)
-		{
-			return PlayerController->GetNetworkPhysicsTickOffset();
-		}
-	}
-
-	return 0;
-}
-
-FMoverTimeStep UMoverNetworkPhysicsLiaisonComponent::GetCurrentAsyncMoverTimeStep_Internal() const
+void UMoverNetworkPhysicsLiaisonComponent::ConsumeOutput_External(const FPhysicsMoverAsyncOutput& Output, const double OutputTimeInSeconds)
 {
-	Chaos::EnsureIsInPhysicsThreadContext();
-	ensure(bUsingAsyncPhysics);
-
-	FMoverTimeStep TimeStep;
-
-	if (FPhysScene* Scene = GetWorld()->GetPhysicsScene())
+	if (Output.bIsValid)
 	{
-		if (Chaos::FPhysicsSolver* Solver = Scene->GetSolver())
-		{
-			TimeStep.ServerFrame = Solver->GetCurrentFrame() + GetNetworkPhysicsTickOffset();
-			TimeStep.StepMs = Solver->GetAsyncDeltaTime() * 1000.0f;
-			TimeStep.BaseSimTimeMs = TimeStep.ServerFrame * TimeStep.StepMs;
-			TimeStep.bIsResimulating = Solver->GetEvolution()->IsResimming();
-		}
+		UpdateConstraintSettings();
 	}
 
-	return TimeStep;
+	Super::ConsumeOutput_External(Output, OutputTimeInSeconds);
 }
 
-FMoverTimeStep UMoverNetworkPhysicsLiaisonComponent::GetCurrentAsyncMoverTimeStep_External() const
+void UMoverNetworkPhysicsLiaisonComponent::PostPhysicsUpdate_External()
 {
-	Chaos::EnsureIsInGameThreadContext();
-	ensure(bUsingAsyncPhysics);
+	const UMoverComponent& MoverComp = GetMoverComponent();
+	const UBaseMovementMode* PrevMode = MoverComp.GetMovementMode();
+	bool bWasFalling = MoverComp.HasGameplayTag(Mover_IsFalling, true);
 
-	FMoverTimeStep TimeStep;
+	Super::PostPhysicsUpdate_External();
 
-	if (FPhysScene* Scene = GetWorld()->GetPhysicsScene())
+	UpdateConstraintSettings();
+	if (bWasFalling && PrevMode != MoverComp.GetMovementMode() && MoverComp.HasGameplayTag(Mover_IsOnGround, true))
 	{
-		if (Chaos::FPhysicsSolver* Solver = Scene->GetSolver())
+		if (UFallingMode* FallingMode = MoverComp.FindMode_Mutable<UFallingMode>())
 		{
-			const int32 Offset = GetNetworkPhysicsTickOffset();
-			TimeStep.ServerFrame = Solver->GetCurrentFrame() + Offset;
-			TimeStep.StepMs = Solver->GetAsyncDeltaTime() * 1000.0f;
-			TimeStep.BaseSimTimeMs = Solver->GetPhysicsResultsTime_External() * 1000.0f + Offset * TimeStep.StepMs;
-			TimeStep.bIsResimulating = Solver->GetEvolution()->IsResimming();
-		}
-	}
-
-	return TimeStep;
-}
+			FHitResult HitResult;
+			MoverComp.TryGetFloorCheckHitResult(HitResult);
 
-FMoverTimeStep UMoverNetworkPhysicsLiaisonComponent::GetCurrentMoverTimeStep(float DeltaSeconds) const
-{
-	ensure(!bUsingAsyncPhysics);
-
-	FMoverTimeStep TimeStep;
-
-	if (FPhysScene* Scene = GetWorld()->GetPhysicsScene())
-	{
-		if (Chaos::FPhysicsSolver* Solver = Scene->GetSolver())
-		{
-			TimeStep.ServerFrame = Solver->GetCurrentFrame();
-			TimeStep.StepMs = DeltaSeconds * 1000.0f;
-			TimeStep.BaseSimTimeMs = Solver->GetSolverTime() * 1000.0f;
-			TimeStep.bIsResimulating = Solver->GetEvolution()->IsResimming();
+			FallingMode->OnLanded.Broadcast(MoverComp.GetMovementModeName(), HitResult);
 		}
 	}
-
-	return TimeStep;
 }
 
-void UMoverNetworkPhysicsLiaisonComponent::InitializeCachedSyncState()
+void UMoverNetworkPhysicsLiaisonComponent::OnContactModification_Internal(const FPhysicsMoverAsyncInput& Input, Chaos::FCollisionContactModifier& Modifier) const
 {
-	SyncStateRecord.Clear();
+	Super::OnContactModification_Internal(Input, Modifier);
 
-	if (const UMoverComponent* MoverComponentPtr = MoverComp.Get())
-	{
-		FMoverSyncState SyncState;
-		SyncState.MovementMode = MoverComponentPtr->StartingMovementMode;
-
-		double SimTime = 0.0;
-		if (Chaos::FPhysicsSolver* Solver = GetPhysicsSolver())
-		{
-			SimTime = Solver->GetPhysicsResultsTime_External();
-		}
-
-		SyncStateRecord.Add(SimTime, SyncState);
-		LatestSyncState = SyncState;
-	}
-}
-
-Chaos::FPhysicsSolver* UMoverNetworkPhysicsLiaisonComponent::GetPhysicsSolver() const
-{
-	if (const UWorld* World = GetWorld())
+	if (!HasValidState())
 	{
-		if (const FPhysScene* Scene = World->GetPhysicsScene())
-		{
-			return Scene->GetSolver();
-		}
+		return;
 	}
 
-	return nullptr;
-}
-
-void UMoverNetworkPhysicsLiaisonComponent::ProduceInput_External(int32 PhysicsStep, int32 NumSteps, OUT FPhysicsMoverAsyncInput& Input)
-{
-	Chaos::EnsureIsInGameThreadContext();
-
-	if (HasValidState())
+	const UMoverComponent& MoverComp = GetMoverComponent();
+	if (MoverComp.MovementModes.Contains(Input.SyncState.MovementMode))
 	{
-		Input.MoverIdx = GetUniqueIdx();
-		Input.MoverSimulation = this;
-
-		// Produce input
-		if (!MoverComp)
+		if (const IPhysicsCharacterMovementModeInterface* PhysicsMode = Cast<const IPhysicsCharacterMovementModeInterface>(MoverComp.MovementModes[Input.SyncState.MovementMode]))
 		{
-			return;
-		}
-
-		// If in a networked game only produce input for the locally controlled character
-		APawn* PawnOwner = Cast<APawn>(GetOwner());
-		bool bProduceInput = PawnOwner ? PawnOwner->IsLocallyControlled() : false;
-
-		if (PawnOwner && !PawnOwner->IsPlayerControlled() && NetworkPhysicsComponent && !NetworkPhysicsComponent->GetIsRelayingLocalInputs() && NetworkPhysicsComponent->HasServerWorld())
-		{
-			NetworkPhysicsComponent->SetIsRelayingLocalInputs(true);
-		}
-
-		float DeltaSeconds = 0.0f;
-		if (const UWorld* World = GetWorld())
-		{
-			if (bUsingAsyncPhysics)
-			{
-				if (Chaos::FPhysicsSolver* Solver = GetPhysicsSolver())
-				{
-					if (Solver->IsUsingAsyncResults())
-					{
-						DeltaSeconds = Solver->GetAsyncDeltaTime();
-					}
-				}
-			}
-			else
+			Chaos::FCharacterGroundConstraintHandle* ConstraintHandle = Constraint->GetProxy<Chaos::FCharacterGroundConstraintProxy>()->GetPhysicsThreadAPI();
+			if (!ConstraintHandle || !ConstraintHandle->IsEnabled() || !ConstraintHandle->GetCharacterParticle())
 			{
-				DeltaSeconds = World->GetDeltaSeconds();
+				return;
 			}
-		}
 
-		if (bProduceInput)
-		{
-			const int DeltaTimeMS = FMath::RoundToInt(DeltaSeconds * 1000.0f);
-			MoverComp->ProduceInput(DeltaTimeMS, &Input.InputCmd);
+			const FPhysicsMoverSimulationContactModifierParams Params{ConstraintHandle, MoverComp.UpdatedCompAsPrimitive};
+			PhysicsMode->OnContactModification_Internal(Params, Modifier);
 		}
-
-		Input.SyncState.SyncStateCollection.Empty();
-		Input.SyncState = LatestSyncState;
-
-		const FMoverTimeStep MoverTimeStep = bUsingAsyncPhysics ? GetCurrentAsyncMoverTimeStep_External() : GetCurrentMoverTimeStep(DeltaSeconds);
-		MoverComp->CachedLastSimTickTimeStep = MoverTimeStep;
-
-		MoverComp->OnPreSimulationTick.Broadcast(MoverTimeStep, Input.InputCmd);
-
-		// This is required so that the physics thread can have a copy of the data to access
-		LatestInputCmd = Input.InputCmd;
 	}
 }
 
-void UMoverNetworkPhysicsLiaisonComponent::ConsumeOutput_External(const FPhysicsMoverAsyncOutput& Output, const double OutputTimeInSeconds)
+bool UMoverNetworkPhysicsLiaisonComponent::HasValidState() const
 {
-	Chaos::EnsureIsInGameThreadContext();
-
-	if (Output.bIsValid)
-	{
-		SyncStateRecord.Add(OutputTimeInSeconds, Output.SyncState);
-	}
+	// For character movement, we need to have a valid input producer (i.e. the character itself)
+	return Super::HasValidState() && GetMoverComponent().InputProducer;
 }
 
-void UMoverNetworkPhysicsLiaisonComponent::PostPhysicsUpdate_External()
+bool UMoverNetworkPhysicsLiaisonComponent::CanProcessInputs_Internal(const FPhysicsMoverAsyncInput& Input) const
 {
-	Chaos::EnsureIsInGameThreadContext();
-
-	// Write the interpolated result to the mover component
-	if (UMoverComponent* MoverComponentPtr = MoverComp.Get())
+	if (!Input.InputCmd.InputCollection.FindDataByType<FCharacterDefaultInputs>())
 	{
-		if (const Chaos::FPhysicsSolver* Solver = GetPhysicsSolver())
-		{
-			const float DeltaSeconds = MoverComp->CachedLastSimTickTimeStep.StepMs * 0.001f;
-			const FMoverTimeStep MoverTimeStep = bUsingAsyncPhysics ? GetCurrentAsyncMoverTimeStep_External() : GetCurrentMoverTimeStep(DeltaSeconds);
-
-			FMoverSyncState OutputSyncState;
-			if (bUsingAsyncPhysics)
-			{
-				const float ResultsTime = Solver->GetPhysicsResultsTime_External();
-				SyncStateRecord.GetInterpolated(ResultsTime, OutputSyncState);
-			}
-			else
-			{
-				OutputSyncState = SyncStateRecord.GetLatest();
-			}
-
-			FMoverDefaultSyncState& DefaultSyncState = OutputSyncState.SyncStateCollection.FindOrAddMutableDataByType<FMoverDefaultSyncState>();
-
-			// The physics mover does not set the sync state transform, it uses the physics transform,
-			// so get the world transform from the physics particle.
-			if (IPhysicsComponent* PhysComp = Cast<IPhysicsComponent>(MoverComponentPtr->GetUpdatedComponent()))
-			{
-				Chaos::FPhysicsObject* PhysicsObject = PhysComp->GetPhysicsObjectByName(NAME_None);
-				Chaos::FReadPhysicsObjectInterface_External Interface = FPhysicsObjectExternalInterface::GetRead_AssumesLocked();
-				if (Chaos::FPBDRigidParticle* Particle = Interface.GetRigidParticle(PhysicsObject))
-				{
-					DefaultSyncState.SetTransforms_WorldSpace(Particle->X(), FRotator(Particle->R()), Particle->V(), DefaultSyncState.GetMovementBase(), DefaultSyncState.GetMovementBaseBoneName());
-					
-					// The movement intent is currently in world space. Need to convert to base space
-					DefaultSyncState.MoveDirectionIntent = DefaultSyncState.GetCapturedMovementBaseQuat().UnrotateVector(DefaultSyncState.MoveDirectionIntent);
-
-				}
-			}
-
-			const FName CachedMovementMode = MoverComponentPtr->GetMovementModeName();
-
-			const bool bIsCachedModeFalling = MoverComponentPtr->HasGameplayTag(Mover_IsFalling, true);
-
-			MoverComponentPtr->CachedLastSyncState = OutputSyncState;
-			MoverComponentPtr->CachedLastSimTickTimeStep = MoverTimeStep;
-			MoverComponentPtr->bHasValidCachedState = true;
-
-			const FName NextMode = MoverComponentPtr->GetMovementModeName();
-			UpdateConstraintSettings();
-			if (CachedMovementMode != NextMode)
-			{
-				MoverComponentPtr->OnMovementModeChanged.Broadcast(CachedMovementMode, NextMode);
-
-				if (bIsCachedModeFalling && MoverComponentPtr->HasGameplayTag(Mover_IsOnGround, true))
-				{
-					if (UFallingMode* FallingMode = MoverComponentPtr->FindMode_Mutable<UFallingMode>())
-					{
-						FHitResult HitResult;
-						MoverComponentPtr->TryGetFloorCheckHitResult(HitResult);
-
-						FallingMode->OnLanded.Broadcast(OutputSyncState.MovementMode, HitResult);
-					}
-				}
-			}
-
-			MoverComponentPtr->OnPostSimulationTick.Broadcast(MoverTimeStep);
-		}
+		return false;
 	}
+	
+	return Super::CanProcessInputs_Internal(Input);
 }
 
-void UMoverNetworkPhysicsLiaisonComponent::ProcessInputs_Internal(int32 PhysicsStep, float DeltaTime, const FPhysicsMoverAsyncInput& Input) const
+void UMoverNetworkPhysicsLiaisonComponent::PerformProcessInputs_Internal(int32 PhysicsStep, float DeltaTime, const FPhysicsMoverAsyncInput& Input) const
 {
-	Chaos::EnsureIsInPhysicsThreadContext();
-
-	// Override input data unless player is local client
-	if (FPhysScene* PhysScene = GetWorld()->GetPhysicsScene())
-	{
-		if (NetworkPhysicsComponent)
-		{
-			if (NetworkPhysicsComponent->HasServerWorld())
-			{
-				// Server remote player
-				GetCurrentInputData(Input.InputCmd);
-			}
-			else
-			{
-				bool bIsSolverResim = false;
-				bool bIsFirstResimFrame = false;
-				if (Chaos::FPhysicsSolver* Solver = PhysScene->GetSolver())
-				{
-					bIsSolverResim = Solver->GetEvolution()->IsResimming();
-					bIsFirstResimFrame = Solver->GetEvolution()->IsResetting();
-				}
-
-				bool bLocalPlayer = false;
-				APlayerController* PlayerController = NetworkPhysicsComponent->GetPlayerController();
-				if (PlayerController && PlayerController->IsLocalController())
-				{
-					bLocalPlayer = true;
-				}
-
-				if (!bLocalPlayer || bIsSolverResim)
-				{
-					ensureMsgf(bUsingAsyncPhysics, TEXT("Physics mover requires async physics to work in multiplayer."));
-					GetCurrentInputData(Input.InputCmd);
-				}
-
-				// Rollback mover state if on the first resimulation frame
-				if (bLocalPlayer && bIsSolverResim && bIsFirstResimFrame)
-				{
-					FMoverAuxStateContext UnusedInvalidAuxState;
-					FMoverAuxStateContext UnusedAuxState;
-
-					MoverComp->OnSimulationPreRollback(&Input.SyncState, &LatestSyncState, &UnusedInvalidAuxState, &UnusedAuxState);
-					GetCurrentStateData(Input.SyncState);
-					MoverComp->OnSimulationRollback(&Input.SyncState, &UnusedAuxState);
-				}
-				else
-				{
-					GetCurrentStateData(Input.SyncState);
-				}
-			}
-
-			const FMoverTimeStep MoverTimeStep = bUsingAsyncPhysics ? GetCurrentAsyncMoverTimeStep_Internal() : GetCurrentMoverTimeStep(DeltaTime);
-			FMoverAuxStateContext UnusedAuxState;
-			MoverComp->OnPreMovement.Broadcast(MoverTimeStep, Input.InputCmd, Input.SyncState, UnusedAuxState);
-		}
-	}
+	Super::PerformProcessInputs_Internal(PhysicsStep, DeltaTime, Input);
 
 	// Override common settings data with data from FMovementSettingsInputs if present in the input cmd
 	if (const FMovementSettingsInputs* MovementSettings = Input.InputCmd.InputCollection.FindDataByType<FMovementSettingsInputs>())
@@ -921,67 +246,51 @@ void UMoverNetworkPhysicsLiaisonComponent::ProcessInputs_Internal(int32 PhysicsS
 	}
 }
 
-void UMoverNetworkPhysicsLiaisonComponent::OnPreSimulate_Internal(const FPhysicsMoverSimulationTickParams& TickParams, const FPhysicsMoverAsyncInput& Input, OUT FPhysicsMoverAsyncOutput& Output)
+bool UMoverNetworkPhysicsLiaisonComponent::CanSimulate_Internal(const FPhysicsMoverSimulationTickParams& TickParams, const FPhysicsMoverAsyncInput& Input) const
 {
-	Chaos::EnsureIsInPhysicsThreadContext();
+	if (Super::CanSimulate_Internal(TickParams, Input))
+	{
+		Chaos::FCharacterGroundConstraintHandle* ConstraintHandle = Constraint->GetProxy<Chaos::FCharacterGroundConstraintProxy>()->GetPhysicsThreadAPI();
+		if (!ConstraintHandle || !ConstraintHandle->IsEnabled())
+		{
+			return false;
+		}
 
-	// Sync state should carry over to the next sim frame by default unless something modifies it
-	Output.SyncState = Input.SyncState;
+		if (!Cast<IPhysicsCharacterMovementModeInterface>(GetMoverComponent().MovementModes.FindChecked(Input.SyncState.MovementMode)))
+		{
+			UE_LOG(LogMover, Verbose, TEXT("Attempting to run non-character physics movement mode %s in physics mover update. Only modes that implement IPhysicsCharacterMovementModeInterface can be used with the CharacterPhysicsLiaison."),
+				*Input.SyncState.MovementMode.ToString());
+			return false;
+		}
 
-	if (!HasValidState() || !Input.InputCmd.InputCollection.FindDataByType<FCharacterDefaultInputs>())
-	{
-		return;
+		return true;
 	}
 
-	// Exit if Physics is not enabled on the CollisionShape
-	const UShapeComponent* CollisionShape = Cast<const UShapeComponent>(MoverComp->UpdatedComponent);
-	if (!CollisionShape || !CollisionShape->IsSimulatingPhysics())
-	{
-		return;
-	}
-	
-	Chaos::FCharacterGroundConstraintHandle* ConstraintHandle = Constraint->GetProxy<Chaos::FCharacterGroundConstraintProxy>()->GetPhysicsThreadAPI();
-	if (!ConstraintHandle || !ConstraintHandle->IsEnabled() || !ConstraintHandle->GetCharacterParticle())
-	{
-		return;
-	}
+	return false;
+}
 
-	ConstraintHandle->SetGroundParticle(nullptr);
+void UMoverNetworkPhysicsLiaisonComponent::PerformPreSimulate_Internal(const FPhysicsMoverSimulationTickParams& TickParams, const FPhysicsMoverAsyncInput& Input, FPhysicsMoverAsyncOutput& Output) const
+{
+	FCharacterDefaultInputs* CharacterDefaultInputs = Input.InputCmd.InputCollection.FindMutableDataByType<FCharacterDefaultInputs>();
+	check(CharacterDefaultInputs);
 
-	// Update sync state from physics
+	Chaos::FCharacterGroundConstraintHandle* ConstraintHandle = Constraint->GetProxy<Chaos::FCharacterGroundConstraintProxy>()->GetPhysicsThreadAPI();
 	Chaos::FPBDRigidParticleHandle* CharacterParticle = ConstraintHandle->GetCharacterParticle()->CastToRigidParticle();
-	if (!CharacterParticle || CharacterParticle->Disabled())
-	{
-		return;
-	}
-
-	// Check that the input movement mode is valid
-	if (MoverComp->MovementModes.Contains(Input.SyncState.MovementMode))
-	{
-		const IPhysicsCharacterMovementModeInterface* PhysicsMode = Cast<const IPhysicsCharacterMovementModeInterface>(MoverComp->MovementModes[Input.SyncState.MovementMode]);
-		if (!PhysicsMode)
-		{
-			UE_LOG(LogMover, Verbose, TEXT("Attempting to run non-physics movement mode %s in physics mover update."), *Input.SyncState.MovementMode.ToString());
-			return;
-		}
-	}
-	else
-	{
-		return;
-	}
+	UMoverComponent& MoverComp = GetMoverComponent();
 
+	ConstraintHandle->SetGroundParticle(nullptr);
+	
+	//@todo DanH: Ideally the physics walking mode can extricate the ground velocity from the original calculation
 	// Make the sync state velocity relative to the ground if walking
 	FVector LocalGroundVelocity = FVector::ZeroVector;
-	const UBaseMovementMode* InputMode = MoverComp->MovementModes[Input.SyncState.MovementMode];
+	const UBaseMovementMode* InputMode = MoverComp.MovementModes[Input.SyncState.MovementMode];
 	if (InputMode->HasGameplayTag(Mover_IsOnGround, true))
 	{
-		if (const UMoverBlackboard* Blackboard = MoverComp->GetSimBlackboard())
+		const UMoverBlackboard* Blackboard = MoverComp.GetSimBlackboard();
+		FFloorCheckResult LastFloorResult;
+		if (Blackboard->TryGet(CommonBlackboard::LastFloorResult, LastFloorResult))
 		{
-			FFloorCheckResult LastFloorResult;
-			if (Blackboard->TryGet(CommonBlackboard::LastFloorResult, LastFloorResult))
-			{
-				LocalGroundVelocity = UPhysicsMovementUtils::ComputeGroundVelocityFromHitResult(CharacterParticle->GetX(), LastFloorResult.HitResult, TickParams.DeltaTimeSeconds);
-			}
+			LocalGroundVelocity = UPhysicsMovementUtils::ComputeGroundVelocityFromHitResult(CharacterParticle->GetX(), LastFloorResult.HitResult, TickParams.DeltaTimeSeconds);
 		}
 	}
 
@@ -992,83 +301,44 @@ void UMoverNetworkPhysicsLiaisonComponent::OnPreSimulate_Internal(const FPhysics
 		AIMoveVelocity = MoverAIInputs->RVOVelocityDelta;
 	}
 
-	FMoverDefaultSyncState& SyncState = Input.SyncState.SyncStateCollection.FindOrAddMutableDataByType<FMoverDefaultSyncState>();
-	SyncState.SetTransforms_WorldSpace(CharacterParticle->GetX(), FRotator(CharacterParticle->GetR()), CharacterParticle->GetV() - LocalGroundVelocity + AIMoveVelocity);
-
-	///////////////////////////////////////////////////////////////////////////////////////////////////////////////////////////
-	// Update the simulation
-
-	FMoverTickStartData TickStartData(Input.InputCmd, Input.SyncState, FMoverAuxStateContext());
-	FMoverTickEndData TickEndData;
-
-	// Sync state should carry over to the next sim frame by default unless something modifies it
-	TickEndData.SyncState = TickStartData.SyncState;
-
-	const FMoverDefaultSyncState* StartingSyncState = TickStartData.SyncState.SyncStateCollection.FindDataByType<FMoverDefaultSyncState>();
-	check(StartingSyncState);
-
-	FCharacterDefaultInputs* InputCmd = TickStartData.InputCmd.InputCollection.FindMutableDataByType<FCharacterDefaultInputs>();
-	check(InputCmd);
-
-	const FMoverTimeStep TimeStep = bUsingAsyncPhysics ? GetCurrentAsyncMoverTimeStep_Internal() : GetCurrentMoverTimeStep(TickParams.DeltaTimeSeconds);
-
-	// Update movement state machine
+	FMoverDefaultSyncState& InputSyncState = Input.SyncState.SyncStateCollection.FindOrAddMutableDataByType<FMoverDefaultSyncState>();
+	InputSyncState.SetTransforms_WorldSpace(CharacterParticle->GetX(), FRotator(CharacterParticle->GetR()), CharacterParticle->GetV() - LocalGroundVelocity + AIMoveVelocity);
 
-	if (MoverComp->bHasRolledBack)
+	if (!CharacterDefaultInputs->SuggestedMovementMode.IsNone())
 	{
-		MoverComp->ProcessFirstSimTickAfterRollback(TimeStep);
+		MoverComp.QueueNextMode(CharacterDefaultInputs->SuggestedMovementMode);
 	}
 
-	if (InputCmd && !InputCmd->SuggestedMovementMode.IsNone())
-	{
-		MoverComp->QueueNextMode(InputCmd->SuggestedMovementMode);
-	}
+	// Cache the particle XR before calling the parent, because the parent will teleport if the target is too far away and change the particle's transform
+	const FVector CharacterParticleStartingX = CharacterParticle->GetX();
+	const FRotator CharacterParticleStartingR = FRotator(CharacterParticle->GetR());
 
-	// Tick the actual simulation. This is where the proposed moves are queried and executed, affecting change to the moving actor's gameplay state and captured in the output sim state
-	MoverComp->ModeFSM->OnSimulationTick(MoverComp->UpdatedComponent, MoverComp->UpdatedCompAsPrimitive, MoverComp->SimBlackboard.Get(), TickStartData, TimeStep, TickEndData);
+	Super::PerformPreSimulate_Internal(TickParams, Input, Output);
 
-	// Set the output sync state and fill in the movement mode
-	Output.SyncState = TickEndData.SyncState;
 	FMoverDefaultSyncState& OutputSyncState = Output.SyncState.SyncStateCollection.FindOrAddMutableDataByType<FMoverDefaultSyncState>();
 
-	const FName MovementModeAfterTick = MoverComp->ModeFSM->GetCurrentModeName();
-	Output.SyncState.MovementMode = MovementModeAfterTick;
-
-	MoverComp->SimBlackboard->TryGet(CommonBlackboard::LastFloorResult, Output.FloorResult);
-
-	/////////////////////////////////////////////////////////////////////////////////////////////////////////////////////////////
-	// Update physics constraint from output sync state
-
-	FVector TargetDeltaPos = OutputSyncState.GetLocation_WorldSpace() - CharacterParticle->GetX();
-
-	if (TargetDeltaPos.SizeSquared2D() > GPhysicsDrivenMotionDebugParams.TeleportThreshold * GPhysicsDrivenMotionDebugParams.TeleportThreshold)
-	{
-		TeleportParticle(CharacterParticle, OutputSyncState.GetLocation_WorldSpace(), OutputSyncState.GetOrientation_WorldSpace().Quaternion());
-	}
-
-	// Add back the ground velocity that was subtracted to but the movement velocity in local space
+	// Add back the ground velocity that was subtracted to the movement velocity in local space
 	FVector TargetVelocity = OutputSyncState.GetVelocity_WorldSpace() + LocalGroundVelocity;
+	CharacterParticle->SetV(TargetVelocity);
 
-	FRotator DeltaRotation = OutputSyncState.GetOrientation_WorldSpace() - FRotator(CharacterParticle->GetR());
+	FRotator DeltaRotation = OutputSyncState.GetOrientation_WorldSpace() - CharacterParticleStartingR;
 	FRotator Winding, Remainder;
 	DeltaRotation.GetWindingAndRemainder(Winding, Remainder);
 	float TargetDeltaFacing = FMath::DegreesToRadians(Remainder.Yaw);
-
-	CharacterParticle->SetV(TargetVelocity);
 	if (TickParams.DeltaTimeSeconds > UE_SMALL_NUMBER)
 	{
 		FVector AngularVelocity = (TargetDeltaFacing / TickParams.DeltaTimeSeconds) * Chaos::FVec3::ZAxisVector;
-		
+			
 		if (!CommonMovementSettings->bShouldRemainVertical)
 		{
 			FVector PreviousUp = CharacterParticle->GetR().RotateVector(FVector::UpVector).GetSafeNormal();
 			FVector TargetUp = OutputSyncState.GetOrientation_WorldSpace().Quaternion().RotateVector(FVector::UpVector).GetSafeNormal();
 			AngularVelocity += (PreviousUp.Cross(TargetUp).GetSafeNormal() * FMath::Acos(TargetUp.Dot(PreviousUp))) / TickParams.DeltaTimeSeconds;
 		}
-		
+			
 		CharacterParticle->SetW(AngularVelocity);
-	}
-
+	}		
+	
 	// Update the constraint data based on the floor result
 	if (Output.FloorResult.bBlockingHit)
 	{
@@ -1103,6 +373,7 @@ void UMoverNetworkPhysicsLiaisonComponent::OnPreSimulate_Internal(const FPhysics
 			WalkableSlopeCosine = 2.0f;
 		}
 
+		const FVector TargetDeltaPos = OutputSyncState.GetLocation_WorldSpace() - CharacterParticleStartingX;
 		ConstraintHandle->SetData({
 			Output.FloorResult.HitResult.ImpactNormal,
 			TargetDeltaPos,
@@ -1114,7 +385,7 @@ void UMoverNetworkPhysicsLiaisonComponent::OnPreSimulate_Internal(const FPhysics
 	else
 	{
 		ConstraintHandle->SetGroundParticle(nullptr);
-
+		
 		OutputSyncState.SetMovementBase(nullptr);
 
 		ConstraintHandle->SetData({
@@ -1126,82 +397,16 @@ void UMoverNetworkPhysicsLiaisonComponent::OnPreSimulate_Internal(const FPhysics
 			});
 	}
 
-	// Physics can tick multiple times using the same input data from the game thread
-	// so make sure to update it here using the results of this update
-	Input.SyncState = Output.SyncState;
-	InputCmd->SuggestedMovementMode = NAME_None; // Should have already been set this frame
-
-	// This is required for cases where we run a second physics update (including generating
-	// input) before the output is saved on the game thread. We want to make sure the next
-	// physics tick starts with the sync state from the previous tick
-	LatestSyncState = Output.SyncState;
-
-	Output.bIsValid = true;
+	CharacterDefaultInputs->SuggestedMovementMode = NAME_None; // Should have already been set this frame
 }
 
-void UMoverNetworkPhysicsLiaisonComponent::TeleportParticle(Chaos::FGeometryParticleHandle* Particle, const FVector& Position, const FQuat& Rotation) const
+Chaos::FPhysicsObject* UMoverNetworkPhysicsLiaisonComponent::GetControlledPhysicsObject() const
 {
-	if (!Particle)
+	if (Constraint && Constraint->GetCharacterParticleProxy())
 	{
-		return;
-	}
-
-	if (Chaos::FPhysicsSolver* Solver = GetPhysicsSolver())
-	{
-		if (Solver->GetEvolution())
-		{
-			Solver->GetEvolution()->SetParticleTransform(Particle, Position, Rotation, true);
-		}
-	}
-}
-
-void UMoverNetworkPhysicsLiaisonComponent::WakeParticleIfSleeping(Chaos::FGeometryParticleHandle* Particle) const
-{
-	if (Particle)
-	{
-		Chaos::FPBDRigidParticleHandle* Rigid = Particle->CastToRigidParticle();
-		if (Rigid && (Rigid->ObjectState() == Chaos::EObjectStateType::Sleeping))
-		{
-			if (UWorld* World = GetWorld())
-			{
-				if (FPhysScene* Scene = World->GetPhysicsScene())
-				{
-					if (Chaos::FPhysicsSolver* Solver = Scene->GetSolver())
-					{
-						if (Solver->GetEvolution())
-						{
-							Solver->GetEvolution()->SetParticleObjectState(Rigid, Chaos::EObjectStateType::Dynamic);
-						}
-					}
-				}
-			}
-		}
-	}
-}
-
-void UMoverNetworkPhysicsLiaisonComponent::OnContactModification_Internal(const FPhysicsMoverAsyncInput& Input, Chaos::FCollisionContactModifier& Modifier) const
-{
-	Chaos::EnsureIsInPhysicsThreadContext();
-
-	if (!HasValidState())
-	{
-		return;
-	}
-
-	if (MoverComp->MovementModes.Contains(Input.SyncState.MovementMode))
-	{
-		if (const IPhysicsCharacterMovementModeInterface* PhysicsMode = Cast<const IPhysicsCharacterMovementModeInterface>(MoverComp->MovementModes[Input.SyncState.MovementMode]))
-		{
-			Chaos::FCharacterGroundConstraintHandle* ConstraintHandle = Constraint->GetProxy<Chaos::FCharacterGroundConstraintProxy>()->GetPhysicsThreadAPI();
-			if (!ConstraintHandle || !ConstraintHandle->IsEnabled() || !ConstraintHandle->GetCharacterParticle())
-			{
-				return;
-			}
-
-			const FPhysicsMoverSimulationContactModifierParams Params { ConstraintHandle, MoverComp->UpdatedCompAsPrimitive };
-			PhysicsMode->OnContactModification_Internal(Params, Modifier);
-		}
+		return Constraint->GetCharacterParticleProxy()->GetPhysicsObject();
 	}
+	return nullptr;
 }
 
 #undef LOCTEXT_NAMESPACE
diff --git a/Engine/Plugins/Experimental/Mover/Source/Mover/Private/Backends/MoverNetworkPhysicsLiaisonBase.cpp b/Engine/Plugins/Experimental/Mover/Source/Mover/Private/Backends/MoverNetworkPhysicsLiaisonBase.cpp
new file mode 100644
--- /dev/null
+++ b/Engine/Plugins/Experimental/Mover/Source/Mover/Private/Backends/MoverNetworkPhysicsLiaisonBase.cpp
@@ -0,0 +1,875 @@
+// Copyright Epic Games, Inc. All Rights Reserved.
+
+#include "Backends/MoverNetworkPhysicsLiaisonBase.h"
+
+#include "Components/ShapeComponent.h"
+#include "Framework/Threading.h"
+#include "GameFramework/Pawn.h"
+#include "MovementModeStateMachine.h"
+#include "PBDRigidsSolver.h"
+#include "Chaos/PhysicsObjectInterface.h"
+#include "Physics/Experimental/PhysScene_Chaos.h"
+#include "PhysicsEngine/PhysicsObjectExternalInterface.h"
+#include "PhysicsMover/Modes/PhysicsDrivenFallingMode.h"
+#include "PhysicsMover/PhysicsMoverManager.h"
+#include "Runtime/Experimental/Chaos/Private/Chaos/PhysicsObjectInternal.h"
+
+#if WITH_EDITOR
+#include "Misc/DataValidation.h"
+#endif
+
+#include UE_INLINE_GENERATED_CPP_BY_NAME(MoverNetworkPhysicsLiaisonBase)
+
+#define LOCTEXT_NAMESPACE "Mover"
+
+//////////////////////////////////////////////////////////////////////////
+
+extern FPhysicsDrivenMotionDebugParams GPhysicsDrivenMotionDebugParams;
+
+//////////////////////////////////////////////////////////////////////////
+// FNetworkPhysicsMoverInputs
+
+void FNetworkPhysicsMoverInputs::ApplyData(UActorComponent* NetworkComponent) const
+{
+	if (UMoverNetworkPhysicsLiaisonComponentBase* LiaisonComp = Cast<UMoverNetworkPhysicsLiaisonComponentBase>(NetworkComponent))
+	{
+		LiaisonComp->SetCurrentInputData(InputCmdContext);
+	}
+}
+
+void FNetworkPhysicsMoverInputs::BuildData(const UActorComponent* NetworkComponent)
+{
+	if (const UMoverNetworkPhysicsLiaisonComponentBase* LiaisonComp = Cast<const UMoverNetworkPhysicsLiaisonComponentBase>(NetworkComponent))
+	{
+		LiaisonComp->GetCurrentInputData(InputCmdContext);
+	}
+}
+
+bool FNetworkPhysicsMoverInputs::NetSerialize(FArchive& Ar, class UPackageMap* PackageMap, bool& bOutSuccess)
+{
+	SerializeFrames(Ar);
+	if (PackageMap)
+	{
+		InputCmdContext.NetSerialize(FNetSerializeParams(Ar));
+		bOutSuccess = true;
+	}
+	else
+	{
+		bOutSuccess = false;
+	}
+
+	return bOutSuccess;
+}
+
+void FNetworkPhysicsMoverInputs::InterpolateData(const FNetworkPhysicsData& MinData, const FNetworkPhysicsData& MaxData)
+{
+	const FNetworkPhysicsMoverInputs& MinDataInput = static_cast<const FNetworkPhysicsMoverInputs&>(MinData);
+	const FNetworkPhysicsMoverInputs& MaxDataInput = static_cast<const FNetworkPhysicsMoverInputs&>(MaxData);
+
+	const float LerpFactor = (LocalFrame - MinDataInput.LocalFrame) / (MaxDataInput.LocalFrame - MinDataInput.LocalFrame);
+	InputCmdContext.InputCollection.Interpolate(MinDataInput.InputCmdContext.InputCollection, MaxDataInput.InputCmdContext.InputCollection, LerpFactor);
+}
+
+void FNetworkPhysicsMoverInputs::MergeData(const FNetworkPhysicsData& FromData)
+{
+	const FNetworkPhysicsMoverInputs& TypedFrom = static_cast<const FNetworkPhysicsMoverInputs&>(FromData);
+	InputCmdContext.InputCollection.Merge(TypedFrom.InputCmdContext.InputCollection);
+}
+
+void FNetworkPhysicsMoverInputs::ValidateData(const UActorComponent* NetworkComponent)
+{
+	if (NetworkComponent)
+	{
+		if (const UMoverNetworkPhysicsLiaisonComponentBase* LiaisonComp = Cast<UMoverNetworkPhysicsLiaisonComponentBase>(NetworkComponent))
+		{
+			LiaisonComp->ValidateInputData(InputCmdContext);
+		}
+	}
+}
+
+bool FNetworkPhysicsMoverInputs::CompareData(const FNetworkPhysicsData& PredictedData)
+{
+	const FMoverInputCmdContext& PredictedInputCmd = static_cast<const FNetworkPhysicsMoverInputs&>(PredictedData).InputCmdContext;
+	return !PredictedInputCmd.InputCollection.ShouldReconcile(InputCmdContext.InputCollection);
+}
+
+const FString FNetworkPhysicsMoverInputs::DebugData()
+{
+	FAnsiStringBuilderBase StringBuilder;
+	InputCmdContext.ToString(StringBuilder);
+	return FString::Printf(TEXT("FNetworkPhysicsMoverInputs:\n%hs"), StringBuilder.ToString());
+}
+
+//////////////////////////////////////////////////////////////////////////
+// FNetworkPhysicsMoverState
+
+void FNetworkPhysicsMoverState::ApplyData(UActorComponent* NetworkComponent) const
+{
+	if (NetworkComponent)
+	{
+		if (UMoverNetworkPhysicsLiaisonComponentBase* LiaisonComp = Cast<UMoverNetworkPhysicsLiaisonComponentBase>(NetworkComponent))
+		{
+			LiaisonComp->SetCurrentStateData(SyncStateContext);
+		}
+	}
+}
+
+void FNetworkPhysicsMoverState::BuildData(const UActorComponent* NetworkComponent)
+{
+	if (NetworkComponent)
+	{
+		if (const UMoverNetworkPhysicsLiaisonComponentBase* LiaisonComp = Cast<const UMoverNetworkPhysicsLiaisonComponentBase>(NetworkComponent))
+		{
+			LiaisonComp->GetCurrentStateData(SyncStateContext);
+		}
+	}
+}
+
+bool FNetworkPhysicsMoverState::NetSerialize(FArchive& Ar, class UPackageMap* PackageMap, bool& bOutSuccess)
+{
+	SerializeFrames(Ar);
+	if (PackageMap)
+	{
+		FNetSerializeParams Params(Ar);
+		SyncStateContext.NetSerialize(Params);
+		bOutSuccess = true;
+	}
+	else
+	{
+		bOutSuccess = false;
+	}
+
+	return bOutSuccess;
+}
+
+void FNetworkPhysicsMoverState::InterpolateData(const FNetworkPhysicsData& MinData, const FNetworkPhysicsData& MaxData)
+{
+	const FNetworkPhysicsMoverState& MinState = static_cast<const FNetworkPhysicsMoverState&>(MinData);
+	const FNetworkPhysicsMoverState& MaxState = static_cast<const FNetworkPhysicsMoverState&>(MaxData);
+
+	const float LerpFactor = (LocalFrame - MinState.LocalFrame) / (MaxState.LocalFrame - MinState.LocalFrame);
+	SyncStateContext.Interpolate(&MinState.SyncStateContext, &MaxState.SyncStateContext, LerpFactor);
+}
+
+bool FNetworkPhysicsMoverState::CompareData(const FNetworkPhysicsData& PredictedData)
+{
+	const FMoverSyncState& PredictedSyncState = static_cast<const FNetworkPhysicsMoverState&>(PredictedData).SyncStateContext;
+	return !PredictedSyncState.ShouldReconcile(SyncStateContext);
+}
+
+const FString FNetworkPhysicsMoverState::DebugData()
+{
+	FAnsiStringBuilderBase StringBuilder;
+	SyncStateContext.ToString(StringBuilder);
+	return FString::Printf(TEXT("FNetworkPhysicsMoverState:\n%hs"), StringBuilder.ToString());
+}
+
+//////////////////////////////////////////////////////////////////////////
+// UMoverNetworkPhysicsLiaisonComponentBase
+
+void UMoverNetworkPhysicsLiaisonComponentBase::GetCurrentInputData(OUT FMoverInputCmdContext& InputCmd) const
+{
+	InputCmd = LatestInputCmd;
+}
+
+void UMoverNetworkPhysicsLiaisonComponentBase::GetCurrentStateData(OUT FMoverSyncState& SyncState) const
+{
+	SyncState = LatestSyncState;
+}
+
+void UMoverNetworkPhysicsLiaisonComponentBase::SetCurrentInputData(const FMoverInputCmdContext& InputCmd)
+{
+	LatestInputCmd = InputCmd;
+}
+
+void UMoverNetworkPhysicsLiaisonComponentBase::SetCurrentStateData(const FMoverSyncState& SyncState)
+{
+	LatestSyncState = SyncState;
+}
+
+bool UMoverNetworkPhysicsLiaisonComponentBase::ValidateInputData(FMoverInputCmdContext& InputCmd) const
+{
+	// TODO - proper data validation
+	return true;
+}
+
+UMoverNetworkPhysicsLiaisonComponentBase::UMoverNetworkPhysicsLiaisonComponentBase()
+{
+	PrimaryComponentTick.TickGroup = TG_PrePhysics;
+	PrimaryComponentTick.bCanEverTick = true;
+
+	bWantsInitializeComponent = true;
+	bAutoActivate = true;
+
+	if (const Chaos::FPhysicsSolver* Solver = GetPhysicsSolver())
+	{
+		bUsingAsyncPhysics = Solver->IsUsingAsyncResults();
+	}
+
+	if (Chaos::FPhysicsSolverBase::IsNetworkPhysicsPredictionEnabled())
+	{
+		SetIsReplicatedByDefault(true);
+	}
+}
+
+//////////////////////////////////////////////////////////////////////////
+//  UMoverNetworkPhysicsLiaisonComponentBase IMoverBackendLiaisonInterface
+
+float UMoverNetworkPhysicsLiaisonComponentBase::GetCurrentSimTimeMs()
+{
+	// Note: this is implicitly an _External function
+	if (const Chaos::FPhysicsSolver* Solver = GetPhysicsSolver())
+	{
+		return bUsingAsyncPhysics ? Solver->GetAsyncDeltaTime() * GetCurrentSimFrame() * 1000.0f : Solver->GetSolverTime() * 1000.0f;
+	}
+	return 0.0f;
+}
+
+int32 UMoverNetworkPhysicsLiaisonComponentBase::GetCurrentSimFrame()
+{
+	// Note: this is implicitly an _External function
+	if (Chaos::FPhysicsSolver* Solver = GetPhysicsSolver())
+	{
+		return Solver->GetCurrentFrame() + GetNetworkPhysicsTickOffset();
+	}
+	return 0;
+}
+
+#if WITH_EDITOR
+EDataValidationResult UMoverNetworkPhysicsLiaisonComponentBase::ValidateData(FDataValidationContext& Context, const UMoverComponent& ValidationMoverComp) const
+{
+	return EDataValidationResult::Valid;
+}
+#endif // WITH_EDITOR
+
+//////////////////////////////////////////////////////////////////////////
+// UMoverNetworkPhysicsLiaisonComponentBase UObject interface
+
+void UMoverNetworkPhysicsLiaisonComponentBase::OnRegister()
+{
+	Super::OnRegister();
+
+	if (UMoverComponent& MoverComp = GetMoverComponent(); MoverComp.UpdatedCompAsPrimitive)
+	{
+		MoverComp.UpdatedCompAsPrimitive->OnComponentPhysicsStateChanged.AddUniqueDynamic(this, &ThisClass::HandleComponentPhysicsStateChanged);
+	}
+}
+
+void UMoverNetworkPhysicsLiaisonComponentBase::OnUnregister()
+{
+	if (UMoverComponent& MoverComp = GetMoverComponent(); MoverComp.UpdatedCompAsPrimitive)
+	{
+		MoverComp.UpdatedCompAsPrimitive->OnComponentPhysicsStateChanged.RemoveDynamic(this, &ThisClass::HandleComponentPhysicsStateChanged);
+	}
+
+	Super::OnUnregister();
+}
+
+void UMoverNetworkPhysicsLiaisonComponentBase::HandleComponentPhysicsStateChanged(UPrimitiveComponent* ChangedComponent, EComponentPhysicsStateChange StateChange)
+{
+	if (StateChange == EComponentPhysicsStateChange::Created)
+	{
+		if (UMoverComponent& MoverComp = GetMoverComponent(); MoverComp.ModeFSM)
+		{
+			MoverComp.ModeFSM->SetModeImmediately(MoverComp.StartingMovementMode);
+		}
+
+		InitializeCachedSyncState();
+	}
+}
+
+void UMoverNetworkPhysicsLiaisonComponentBase::HandleOwningPawnControllerChanged_Server(APawn* OwnerPawn, AController* OldController, AController* NewController)
+{
+	// Inputs for player-controlled pawns originate on the player's client, all others originate on the server
+	NetworkPhysicsComponent->SetIsRelayingLocalInputs(!OwnerPawn->IsPlayerControlled());
+}
+
+Chaos::FPhysicsSolver* UMoverNetworkPhysicsLiaisonComponentBase::GetPhysicsSolver() const
+{
+	if (UWorld* World = GetWorld())
+	{
+		if (FPhysScene* Scene = World->GetPhysicsScene())
+		{
+			return Scene->GetSolver();
+		}
+	}
+	return nullptr;
+}
+
+void UMoverNetworkPhysicsLiaisonComponentBase::InitializeComponent()
+{
+	Super::InitializeComponent();
+
+	if (UWorld* World = GetWorld(); World && World->IsGameWorld())
+	{
+		UMoverComponent& MoverComp = GetMoverComponent();
+		MoverComp.InitMoverSimulation();
+		MoverComp.ModeFSM->SetModeImmediately(MoverComp.StartingMovementMode);
+	}
+
+	if (Chaos::FPhysicsSolverBase::IsNetworkPhysicsPredictionEnabled() && ensureMsgf(bUsingAsyncPhysics, TEXT("Physics mover requires async physics to work with networked physics.")))
+	{
+		NetworkPhysicsComponent = NewObject<UNetworkPhysicsComponent>(GetOwner(), TEXT("PhysMover_NetworkPhysicsComponent"));
+
+		// This isn't technically a DSO component, but set it net addressable as though it is
+		NetworkPhysicsComponent->SetNetAddressable();
+		NetworkPhysicsComponent->SetIsReplicated(true);
+		NetworkPhysicsComponent->RegisterComponent();
+		if (!NetworkPhysicsComponent->HasBeenInitialized())
+		{
+			NetworkPhysicsComponent->InitializeComponent();
+		}
+		NetworkPhysicsComponent->Activate(true);
+
+		// Register network data for recording and rewind/resim
+		NetworkPhysicsComponent->CreateDataHistory<FNetworkPhysicsMoverTraits>(this);
+
+		if (NetworkPhysicsComponent->HasServerWorld())
+		{
+			if (APawn* PawnOwner = Cast<APawn>(GetOwner()))
+			{
+				// When we're owned by a pawn, keep an eye on whether it's currently player-controlled or not
+				PawnOwner->ReceiveControllerChangedDelegate.AddUniqueDynamic(this, &ThisClass::HandleOwningPawnControllerChanged_Server);
+				HandleOwningPawnControllerChanged_Server(PawnOwner, nullptr, PawnOwner->Controller);
+			}
+			else
+			{
+				// If the owner isn't a pawn, there's no chance of player input happening, so inputs to the PT are always produced on the server
+				NetworkPhysicsComponent->SetIsRelayingLocalInputs(true);
+			}
+		}
+	}
+}
+
+void UMoverNetworkPhysicsLiaisonComponentBase::UninitializeComponent()
+{
+	if (NetworkPhysicsComponent)
+	{
+		NetworkPhysicsComponent->RemoveDataHistory();
+		NetworkPhysicsComponent->DestroyComponent();
+		NetworkPhysicsComponent = nullptr;
+	}
+
+	Super::UninitializeComponent();
+}
+
+bool UMoverNetworkPhysicsLiaisonComponentBase::ShouldCreatePhysicsState() const
+{
+	if (!IsRegistered() || IsBeingDestroyed())
+	{
+		return false;
+	}
+
+	UWorld* World = GetWorld();
+	return World && World->IsGameWorld() && World->GetPhysicsScene() && CanCreatePhysics();
+}
+
+bool UMoverNetworkPhysicsLiaisonComponentBase::HasValidState() const
+{
+	const UMoverComponent& MoverComp = GetMoverComponent();
+	return HasValidPhysicsState() && GetControlledPhysicsObject() && GetPhysicsSolver() && MoverComp.UpdatedCompAsPrimitive && MoverComp.UpdatedComponent
+		&& MoverComp.ModeFSM->IsValidLowLevelFast() && MoverComp.GetSimBlackboard()->IsValidLowLevelFast() && MoverComp.MovementMixer;
+}
+
+bool UMoverNetworkPhysicsLiaisonComponentBase::CanCreatePhysics() const
+{
+	check(GetOwner());
+	FString ActorName = GetOwner()->GetName();
+
+	const UMoverComponent& MoverComp = GetMoverComponent();
+	if (!IsValid(MoverComp.UpdatedComponent))
+	{
+		UE_LOG(LogMover, Warning, TEXT("Can't create physics %s (%s). UpdatedComponent is not set."), *ActorName, *GetPathName());
+		return false;
+	}
+
+	if (!IsValid(MoverComp.UpdatedCompAsPrimitive))
+	{
+		UE_LOG(LogMover, Warning, TEXT("Can't create physics %s (%s). UpdatedComponent is not a PrimitiveComponent."), *ActorName, *GetPathName());
+		return false;
+	}
+
+	return true;
+}
+
+void UMoverNetworkPhysicsLiaisonComponentBase::BeginPlay()
+{
+	Super::BeginPlay();
+
+	// Register with the physics mover manager
+	if (UPhysicsMoverManager* MoverManager = UWorld::GetSubsystem<UPhysicsMoverManager>(GetWorld()))
+	{
+		MoverManager->RegisterPhysicsMoverComponent(this);
+	}
+
+	InitializeCachedSyncState();
+}
+
+void UMoverNetworkPhysicsLiaisonComponentBase::EndPlay(const EEndPlayReason::Type EndPlayReason)
+{
+	// Unregister with the physics mover manager
+	if (UPhysicsMoverManager* MoverManager = UWorld::GetSubsystem<UPhysicsMoverManager>(GetWorld()))
+	{
+		MoverManager->UnregisterPhysicsMoverComponent(this);
+	}
+
+	Super::EndPlay(EndPlayReason);
+}
+
+//////////////////////////////////////////////////////////////////////////
+
+Chaos::FUniqueIdx UMoverNetworkPhysicsLiaisonComponentBase::GetUniqueIdx() const
+{
+	if (const UMoverComponent& MoverComp = GetMoverComponent(); MoverComp.UpdatedCompAsPrimitive)
+	{
+		if (FBodyInstance* BI = MoverComp.UpdatedCompAsPrimitive->GetBodyInstance())
+		{
+			if (FPhysicsActorHandle ActorHandle = BI->ActorHandle)
+			{
+				return ActorHandle->GetGameThreadAPI().UniqueIdx();
+			}
+		}
+	}
+
+	return Chaos::FUniqueIdx();
+}
+
+int32 UMoverNetworkPhysicsLiaisonComponentBase::GetNetworkPhysicsTickOffset() const
+{
+	if (NetworkPhysicsComponent && !NetworkPhysicsComponent->HasServerWorld())
+	{
+		if (Chaos::IsInPhysicsThreadContext())
+		{
+			if (const FAsyncNetworkPhysicsComponent* NPCPhysicsThreadAPI = NetworkPhysicsComponent->GetNetworkPhysicsComponent_Internal())
+			{
+				return NPCPhysicsThreadAPI->GetNetworkPhysicsTickOffset();
+			}
+		}
+		else
+		{
+			const APlayerController* PlayerController = NetworkPhysicsComponent->GetPlayerController();
+			if (!PlayerController)
+			{
+				PlayerController = GetWorld()->GetFirstPlayerController();
+			}
+
+			if (PlayerController)
+			{
+				return PlayerController->GetNetworkPhysicsTickOffset();
+			}	
+		}
+	}
+
+	return 0;
+}
+
+FMoverTimeStep UMoverNetworkPhysicsLiaisonComponentBase::GetCurrentAsyncMoverTimeStep_Internal() const
+{
+	Chaos::EnsureIsInPhysicsThreadContext();
+	ensure(bUsingAsyncPhysics);
+
+	FMoverTimeStep TimeStep;
+	if (Chaos::FPhysicsSolver* Solver = GetPhysicsSolver())
+	{
+		TimeStep.ServerFrame = Solver->GetCurrentFrame() + GetNetworkPhysicsTickOffset();
+		TimeStep.StepMs = Solver->GetAsyncDeltaTime() * 1000.0f;
+		TimeStep.BaseSimTimeMs = TimeStep.ServerFrame * TimeStep.StepMs;
+		TimeStep.bIsResimulating = Solver->GetEvolution()->IsResimming();
+	}
+
+	return TimeStep;
+}
+
+FMoverTimeStep UMoverNetworkPhysicsLiaisonComponentBase::GetCurrentAsyncMoverTimeStep_External() const
+{
+	Chaos::EnsureIsInGameThreadContext();
+	ensure(bUsingAsyncPhysics);
+
+	FMoverTimeStep TimeStep;
+	if (Chaos::FPhysicsSolver* Solver = GetPhysicsSolver())
+	{
+		const int32 Offset = GetNetworkPhysicsTickOffset();
+		TimeStep.ServerFrame = Solver->GetCurrentFrame() + Offset;
+		TimeStep.StepMs = Solver->GetAsyncDeltaTime() * 1000.0f;
+		TimeStep.BaseSimTimeMs = Solver->GetPhysicsResultsTime_External() * 1000.0f + Offset * TimeStep.StepMs;
+		TimeStep.bIsResimulating = Solver->GetEvolution()->IsResimming();
+	}
+
+	return TimeStep;
+}
+
+FMoverTimeStep UMoverNetworkPhysicsLiaisonComponentBase::GetCurrentMoverTimeStep(float DeltaSeconds) const
+{
+	ensure(!bUsingAsyncPhysics);
+
+	FMoverTimeStep TimeStep;
+	if (Chaos::FPhysicsSolver* Solver = GetPhysicsSolver())
+	{
+		TimeStep.ServerFrame = Solver->GetCurrentFrame();
+		TimeStep.StepMs = DeltaSeconds * 1000.0f;
+		TimeStep.BaseSimTimeMs = Solver->GetSolverTime() * 1000.0f;
+		TimeStep.bIsResimulating = Solver->GetEvolution()->IsResimming();
+	}
+
+	return TimeStep;
+}
+
+void UMoverNetworkPhysicsLiaisonComponentBase::InitializeCachedSyncState()
+{
+	SyncStateRecord.Clear();
+
+	FMoverSyncState SyncState;
+	SyncState.MovementMode = GetMoverComponent().StartingMovementMode;
+
+	double SimTime = 0.0;
+	if (Chaos::FPhysicsSolver* Solver = GetPhysicsSolver())
+	{
+		SimTime = Solver->GetPhysicsResultsTime_External();
+	}
+
+	SyncStateRecord.Add(SimTime, SyncState);
+	LatestSyncState = SyncState;
+}
+
+UMoverComponent& UMoverNetworkPhysicsLiaisonComponentBase::GetMoverComponent() const
+{
+	return *GetOuterUMoverComponent();
+}
+
+void UMoverNetworkPhysicsLiaisonComponentBase::ProduceInput_External(int32 PhysicsStep, int32 NumSteps, OUT FPhysicsMoverAsyncInput& Input)
+{
+	Chaos::EnsureIsInGameThreadContext();
+
+	if (HasValidState())
+	{
+		// Setting these denote the input as valid, and means ProcessInput_Internal will be called with it on the PT
+		Input.MoverIdx = GetUniqueIdx();
+		Input.MoverSimulation = this;
+
+		// Propagate last frame's output sync state as the initial input sync state for the upcoming frame
+		Input.SyncState = LatestSyncState;
+		
+		float DeltaSeconds = 0.0f;
+		if (const UWorld* World = GetWorld())
+		{
+			if (bUsingAsyncPhysics)
+			{
+				if (const Chaos::FPhysicsSolver* Solver = GetPhysicsSolver(); Solver && Solver->IsUsingAsyncResults())
+				{
+					DeltaSeconds = Solver->GetAsyncDeltaTime();
+				}
+			}
+			else
+			{
+				DeltaSeconds = World->GetDeltaSeconds();
+			}
+		}
+
+		if (!NetworkPhysicsComponent || NetworkPhysicsComponent->IsLocallyControlled())
+		{
+			PerformProduceInput_External(DeltaSeconds, Input);
+			
+			// This is the net instance responsible for actually producing the input command, so set LatestInputCmd here to be picked up by the NPC for replication
+			// All other net instances will have LatestInputCmd assigned in PreProcessInput by the NPC and applied as the actual InputCmd to process down in ProcessInputs_Internal
+			LatestInputCmd = Input.InputCmd;
+		}
+		
+		UMoverComponent& MoverComp = GetMoverComponent();
+		const FMoverTimeStep MoverTimeStep = bUsingAsyncPhysics ? GetCurrentAsyncMoverTimeStep_External() : GetCurrentMoverTimeStep(DeltaSeconds);
+		MoverComp.CachedLastSimTickTimeStep = MoverTimeStep;
+		
+		//@todo DanH: The input info broadcast here will be empty/incorrect when not locally controlled. That seems like a possible landmine of confusion?
+		MoverComp.OnPreSimulationTick.Broadcast(MoverTimeStep, Input.InputCmd);
+	}
+}
+
+void UMoverNetworkPhysicsLiaisonComponentBase::PerformProduceInput_External(float DeltaTime, FPhysicsMoverAsyncInput& Input)
+{
+	UMoverComponent& MoverComp = GetMoverComponent();
+	const int DeltaTimeMS = FMath::RoundToInt(DeltaTime * 1000.0f);
+	MoverComp.ProduceInput(DeltaTimeMS, &Input.InputCmd);
+}
+
+void UMoverNetworkPhysicsLiaisonComponentBase::ConsumeOutput_External(const FPhysicsMoverAsyncOutput& Output, double OutputTimeInSeconds)
+{
+	Chaos::EnsureIsInGameThreadContext();
+
+	if (Output.bIsValid)
+	{
+		SyncStateRecord.Add(OutputTimeInSeconds, Output.SyncState);
+	}
+}
+
+void UMoverNetworkPhysicsLiaisonComponentBase::PostPhysicsUpdate_External()
+{
+	Chaos::EnsureIsInGameThreadContext();
+
+	if (!HasValidState())
+	{
+		return;
+	}
+
+	UMoverComponent& MoverComp = GetMoverComponent();
+	const float DeltaSeconds = MoverComp.CachedLastSimTickTimeStep.StepMs * 0.001f;
+	const FMoverTimeStep MoverTimeStep = bUsingAsyncPhysics ? GetCurrentAsyncMoverTimeStep_External() : GetCurrentMoverTimeStep(DeltaSeconds);
+
+	FMoverSyncState InterpolatedSyncState;
+	if (bUsingAsyncPhysics)
+	{
+		const float ResultsTime = GetPhysicsSolver()->GetPhysicsResultsTime_External();
+		SyncStateRecord.GetInterpolated(ResultsTime, InterpolatedSyncState);
+	}
+	else
+	{
+		InterpolatedSyncState = SyncStateRecord.GetLatest();
+	}
+	
+	// Physics interactions in the last frame may have caused a change in position or velocity that's different from what a simple lerp would predict,
+	// so stomp the lerped sync state's transform data with that of the actual particle after the last sim frame
+	FMoverDefaultSyncState& TransformSyncState = InterpolatedSyncState.SyncStateCollection.FindOrAddMutableDataByType<FMoverDefaultSyncState>();
+	if (Chaos::FPBDRigidParticle* Particle = GetControlledParticle_External())
+	{
+		TransformSyncState.SetTransforms_WorldSpace(Particle->GetX(), FRotator(Particle->GetR()), Particle->GetV(), TransformSyncState.GetMovementBase(), TransformSyncState.GetMovementBaseBoneName());
+		
+		// Make sure the move direction intent is in base space (the base quat is identity if there's no base, effectively making this a no-op)
+		TransformSyncState.MoveDirectionIntent = TransformSyncState.GetCapturedMovementBaseQuat().UnrotateVector(TransformSyncState.MoveDirectionIntent);
+	}
+
+	const FName CachedMovementMode = MoverComp.GetMovementModeName();
+
+	// The MoverComponent relies on its CachedLastSyncState for a lot of information, so setting it here is what makes the resulting state of the completed sim frame
+	// "real" from the perspective the MoverComp and any objects that call getters on it
+	MoverComp.CachedLastSyncState = InterpolatedSyncState;
+	MoverComp.bHasValidCachedState = true;
+
+	// Note this may return something different now because its dependent on the CachedLastSyncState that we just changed above
+	const FName NextMode = MoverComp.GetMovementModeName();
+	if (CachedMovementMode != NextMode)
+	{
+		//@todo DanH: Is this sufficient to just trigger the event? Should we instead be queueing the next mode and going through AdvanceToNextMode on the FSM?
+		MoverComp.OnMovementModeChanged.Broadcast(CachedMovementMode, NextMode);
+	}
+
+	MoverComp.CachedLastSimTickTimeStep = MoverTimeStep;
+	MoverComp.OnPostSimulationTick.Broadcast(MoverTimeStep);
+}
+
+void UMoverNetworkPhysicsLiaisonComponentBase::ProcessInputs_Internal(int32 PhysicsStep, float DeltaTime, const FPhysicsMoverAsyncInput& Input) const
+{
+	Chaos::EnsureIsInPhysicsThreadContext();
+
+	if (CanProcessInputs_Internal(Input))
+	{
+		if (NetworkPhysicsComponent)
+		{
+			const bool bIsInputGenerator = NetworkPhysicsComponent->IsLocallyControlled();
+			if (!bIsInputGenerator)
+			{
+				// If this instance is not generating input for the owning actor, the NPC will have written our LatestInputCmd (via FNetworkPhysicsMoverInputs::ApplyData)
+				// by now as part of PreProcessInput. Applying it as the frame input locally is the final step of replicating the inputs from the generator to everyone else.
+				Input.InputCmd = LatestInputCmd;
+			}
+
+			// Unlike inputs, there's no case where the server isn't the authority on the sync state. By default, the LatestSyncState is
+			// only set by NPC replication when resimulating (via FNetworkPhysicsMoverState::ApplyData), but that setting can be adjusted.
+			// So to account for all possible configs, always establish the initial sync state of the frame based on LatestSyncState on clients (it'd be redundant on servers).
+			if (!NetworkPhysicsComponent->HasServerWorld())
+			{
+				Chaos::FPBDRigidsEvolution* RigidsEvolution = GetPhysicsSolver()->GetEvolution();
+				const bool bIsSolverResim = RigidsEvolution->IsResimming();
+				const bool bIsFirstResimFrame = RigidsEvolution->IsResetting();
+
+				// Rollback mover state if on the first resimulation frame
+				if (bIsInputGenerator && bIsSolverResim && bIsFirstResimFrame)
+				{
+					UMoverComponent& MoverComp = GetMoverComponent();
+					FMoverAuxStateContext UnusedInvalidAuxState;
+					FMoverAuxStateContext UnusedAuxState;
+					MoverComp.OnSimulationPreRollback(&LatestSyncState, &Input.SyncState, &UnusedInvalidAuxState, &UnusedAuxState);
+					GetCurrentStateData(Input.SyncState);
+					MoverComp.OnSimulationRollback(&Input.SyncState, &UnusedAuxState);
+				}
+				else
+				{
+					GetCurrentStateData(Input.SyncState);
+				}
+			}
+		}
+		
+		PerformProcessInputs_Internal(PhysicsStep, DeltaTime, Input);
+	}
+}
+
+bool UMoverNetworkPhysicsLiaisonComponentBase::CanProcessInputs_Internal(const FPhysicsMoverAsyncInput& Input) const
+{
+	return HasValidState() &&
+		GetPhysicsSolver() &&
+		GetControlledParticle_Internal() &&
+		GetMoverComponent().MovementModes.Contains(Input.SyncState.MovementMode);
+}
+
+void UMoverNetworkPhysicsLiaisonComponentBase::PerformProcessInputs_Internal(int32 PhysicsStep, float DeltaTime, const FPhysicsMoverAsyncInput& Input) const
+{
+	const FMoverTimeStep MoverTimeStep = bUsingAsyncPhysics ? GetCurrentAsyncMoverTimeStep_Internal() : GetCurrentMoverTimeStep(DeltaTime);
+	FMoverAuxStateContext UnusedAuxState;
+	GetMoverComponent().OnPreMovement.Broadcast(MoverTimeStep, Input.InputCmd, Input.SyncState, UnusedAuxState);
+
+	Chaos::FPBDRigidParticleHandle* ParticleHandle = GetControlledParticle_Internal();
+	FMoverDefaultSyncState& InputSyncState = Input.SyncState.SyncStateCollection.FindOrAddMutableDataByType<FMoverDefaultSyncState>();
+	InputSyncState.SetTransforms_WorldSpace(ParticleHandle->GetX(), FRotator(ParticleHandle->GetR()), ParticleHandle->GetV());
+}
+
+void UMoverNetworkPhysicsLiaisonComponentBase::PreSimulate_Internal(const FPhysicsMoverSimulationTickParams& TickParams, const FPhysicsMoverAsyncInput& Input, OUT FPhysicsMoverAsyncOutput& Output) const
+{
+	Chaos::EnsureIsInPhysicsThreadContext();
+
+	// Sync state should carry over to the next sim frame by default unless something modifies it
+	Output.SyncState = Input.SyncState;
+
+	if (CanSimulate_Internal(TickParams, Input))
+	{
+		PerformPreSimulate_Internal(TickParams, Input, Output);
+
+		// Physics can tick multiple times using the same input data from the game thread
+		// so make sure to update it here using the results of this update
+		Input.SyncState = Output.SyncState;
+
+		// This is required for cases where we run a second physics update (including generating
+		// input) before the output is saved on the game thread. We want to make sure the next
+		// physics tick starts with the sync state from the previous tick
+		LatestSyncState = Output.SyncState;
+
+		Output.bIsValid = true;
+	}
+}
+
+bool UMoverNetworkPhysicsLiaisonComponentBase::CanSimulate_Internal(const FPhysicsMoverSimulationTickParams& TickParams, const FPhysicsMoverAsyncInput& Input) const
+{
+	Chaos::EnsureIsInPhysicsThreadContext();
+
+	if (!CanProcessInputs_Internal(Input))
+	{
+		return false;
+	}
+
+	if (!Input.SyncState.SyncStateCollection.FindDataByType<FMoverDefaultSyncState>())
+	{
+		return false;
+	}
+
+	const UMoverComponent& MoverComp = GetMoverComponent();
+	if (!MoverComp.UpdatedCompAsPrimitive->IsSimulatingPhysics())
+	{
+		return false;
+	}
+
+	return true;
+}
+
+void UMoverNetworkPhysicsLiaisonComponentBase::PerformPreSimulate_Internal(const FPhysicsMoverSimulationTickParams& TickParams, const FPhysicsMoverAsyncInput& Input, FPhysicsMoverAsyncOutput& Output) const
+{
+	Chaos::EnsureIsInPhysicsThreadContext();
+
+	///////////////////////////////////////////////////////////////////////////////////////////////////////////////////////////
+	// Update the simulation
+
+	UMoverComponent& MoverComp = GetMoverComponent();
+	const FMoverTimeStep TimeStep = bUsingAsyncPhysics ? GetCurrentAsyncMoverTimeStep_Internal() : GetCurrentMoverTimeStep(TickParams.DeltaTimeSeconds);
+
+	// Update movement state machine
+	if (MoverComp.bHasRolledBack)
+	{
+		MoverComp.ProcessFirstSimTickAfterRollback(TimeStep);
+	}
+
+	//@todo DanH: Invoking the FSM OnSimulationTick can sometimes trigger AdvanceToNextMode to happen on the PT
+	// Tick the actual simulation. This is where the proposed moves are queried and executed, affecting change to the moving actor's gameplay state and captured in the output sim state
+	FMoverTickStartData TickStartData(Input.InputCmd, Input.SyncState, FMoverAuxStateContext());
+	FMoverTickEndData TickEndData(&Input.SyncState, nullptr);
+	MoverComp.ModeFSM->OnSimulationTick(MoverComp.UpdatedComponent, MoverComp.UpdatedCompAsPrimitive, MoverComp.GetSimBlackboard_Mutable(), TickStartData, TimeStep, TickEndData);
+
+	// Set the output sync state and fill in the movement mode
+	Output.SyncState = TickEndData.SyncState;
+
+	FMoverDefaultSyncState& OutputSyncState = Output.SyncState.SyncStateCollection.FindOrAddMutableDataByType<FMoverDefaultSyncState>();
+
+	const FName MovementModeAfterTick = MoverComp.ModeFSM->GetCurrentModeName();
+	Output.SyncState.MovementMode = MovementModeAfterTick;
+
+	MoverComp.GetSimBlackboard()->TryGet(CommonBlackboard::LastFloorResult, Output.FloorResult);
+
+	Chaos::FPBDRigidParticleHandle* ParticleHandle = GetControlledParticle_Internal();
+	FVector TargetDeltaPos = OutputSyncState.GetLocation_WorldSpace() - ParticleHandle->GetX();
+	if (TargetDeltaPos.SizeSquared2D() > GPhysicsDrivenMotionDebugParams.TeleportThreshold * GPhysicsDrivenMotionDebugParams.TeleportThreshold)
+	{
+		GetPhysicsSolver()->GetEvolution()->SetParticleTransform(ParticleHandle, OutputSyncState.GetLocation_WorldSpace(), OutputSyncState.GetOrientation_WorldSpace().Quaternion(), true);
+	}
+
+	ParticleHandle->SetV(OutputSyncState.GetVelocity_WorldSpace());
+
+	//@todo DanH: Does the base need to concern itself with setting W on the particle?
+}
+
+Chaos::FPhysicsObject* UMoverNetworkPhysicsLiaisonComponentBase::GetControlledPhysicsObject() const
+{
+	if (IPhysicsComponent* PhysicsComponent = Cast<IPhysicsComponent>(GetMoverComponent().GetUpdatedComponent()))
+	{
+		return PhysicsComponent->GetPhysicsObjectByName(NAME_None);
+	}
+	return nullptr;
+}
+
+Chaos::FPBDRigidParticle* UMoverNetworkPhysicsLiaisonComponentBase::GetControlledParticle_External() const
+{
+	if (Chaos::FPhysicsObject* ControlledObject = GetControlledPhysicsObject())
+	{
+		return ControlledObject->GetParticle<Chaos::EThreadContext::External>()->CastToRigidParticle();
+	}
+	return nullptr;
+}
+
+Chaos::FPBDRigidParticleHandle* UMoverNetworkPhysicsLiaisonComponentBase::GetControlledParticle_Internal() const
+{
+	if (Chaos::FPhysicsObject* ControlledObject = GetControlledPhysicsObject())
+	{
+		return ControlledObject->GetParticle<Chaos::EThreadContext::Internal>()->CastToRigidParticle();
+	}
+	return nullptr;
+}
+
+void UMoverNetworkPhysicsLiaisonComponentBase::OnContactModification_Internal(const FPhysicsMoverAsyncInput& Input, Chaos::FCollisionContactModifier& Modifier) const
+{
+	Chaos::EnsureIsInPhysicsThreadContext();
+
+	// if (CanModifyContact_Internal(Input, Modifier))
+	// {
+	// 	PerformContactModification_Internal(Input, Modifier);
+	// }
+}
+
+void UMoverNetworkPhysicsLiaisonComponentBase::TeleportParticleBy_Internal(Chaos::FGeometryParticleHandle& Particle, const FVector& PositionDelta, const FQuat& RotationDelta) const
+{
+	const FVector TeleportLocation = Particle.GetX() + PositionDelta;
+	const FQuat TeleportRotation = Particle.GetR() + RotationDelta;
+	if (const Chaos::FPhysicsSolver* Solver = GetPhysicsSolver(); Solver && Solver->GetEvolution())
+	{
+		Solver->GetEvolution()->SetParticleTransform(&Particle, TeleportLocation, TeleportRotation, true);
+	}
+}
+
+void UMoverNetworkPhysicsLiaisonComponentBase::WakeParticleIfSleeping(Chaos::FGeometryParticleHandle* Particle) const
+{
+	if (Particle)
+	{
+		if (Chaos::FPBDRigidParticleHandle* RigidParticle = Particle->CastToRigidParticle();
+			RigidParticle && RigidParticle->ObjectState() == Chaos::EObjectStateType::Sleeping)
+		{
+			if (const Chaos::FPhysicsSolver* Solver = GetPhysicsSolver(); Solver && Solver->GetEvolution())
+			{
+				Solver->GetEvolution()->SetParticleObjectState(RigidParticle, Chaos::EObjectStateType::Dynamic);
+			}
+		}
+	}
+}
+
+#undef LOCTEXT_NAMESPACE
diff --git a/Engine/Plugins/Experimental/Mover/Source/Mover/Private/PhysicsMover/PhysicsMoverManager.cpp b/Engine/Plugins/Experimental/Mover/Source/Mover/Private/PhysicsMover/PhysicsMoverManager.cpp
--- a/Engine/Plugins/Experimental/Mover/Source/Mover/Private/PhysicsMover/PhysicsMoverManager.cpp
+++ b/Engine/Plugins/Experimental/Mover/Source/Mover/Private/PhysicsMover/PhysicsMoverManager.cpp
@@ -2,7 +2,7 @@
 
 #include "PhysicsMover/PhysicsMoverManager.h"
 
-#include "Backends/MoverNetworkPhysicsLiaison.h"
+#include "Backends/MoverNetworkPhysicsLiaisonBase.h"
 #include "Chaos/Character/CharacterGroundConstraint.h"
 #include "PBDRigidsSolver.h"
 #include "Physics/Experimental/PhysScene_Chaos.h"
@@ -74,12 +74,12 @@ bool UPhysicsMoverManager::DoesSupportWorldType(const EWorldType::Type WorldType
 
 //////////////////////////////////////////////////////////////////////////
 
-void UPhysicsMoverManager::RegisterPhysicsMoverComponent(TWeakObjectPtr<UMoverNetworkPhysicsLiaisonComponent> InPhysicsMoverComp)
+void UPhysicsMoverManager::RegisterPhysicsMoverComponent(TWeakObjectPtr<UMoverNetworkPhysicsLiaisonComponentBase> InPhysicsMoverComp)
 {
 	PhysicsMoverComponents.AddUnique(InPhysicsMoverComp);
 }
 
-void UPhysicsMoverManager::UnregisterPhysicsMoverComponent(TWeakObjectPtr<UMoverNetworkPhysicsLiaisonComponent> InPhysicsMoverComp)
+void UPhysicsMoverManager::UnregisterPhysicsMoverComponent(TWeakObjectPtr<UMoverNetworkPhysicsLiaisonComponentBase> InPhysicsMoverComp)
 {
 	PhysicsMoverComponents.Remove(InPhysicsMoverComp);
 }
@@ -89,7 +89,7 @@ void UPhysicsMoverManager::UnregisterPhysicsMoverComponent(TWeakObjectPtr<UMover
 void UPhysicsMoverManager::InjectInputs_External(int32 PhysicsStep, int32 NumSteps)
 {
 	// Clear invalid data before we attempt to use it.
-	PhysicsMoverComponents.RemoveAllSwap([](const TWeakObjectPtr<UMoverNetworkPhysicsLiaisonComponent> MoverComp) {
+	PhysicsMoverComponents.RemoveAllSwap([](const TWeakObjectPtr<UMoverNetworkPhysicsLiaisonComponentBase> MoverComp) {
 		return (MoverComp.IsValid() == false) || (MoverComp->GetUniqueIdx().IsValid() == false);
 		});
 
@@ -98,11 +98,15 @@ void UPhysicsMoverManager::InjectInputs_External(int32 PhysicsStep, int32 NumSte
 	ManagerAsyncInput->Reset();
 	ManagerAsyncInput->AsyncInput.Reserve(PhysicsMoverComponents.Num());
 
-	for (TWeakObjectPtr<UMoverNetworkPhysicsLiaisonComponent> PhysicsMoverComp : PhysicsMoverComponents)
+	for (TWeakObjectPtr<UMoverNetworkPhysicsLiaisonComponentBase> PhysicsMoverComp : PhysicsMoverComponents)
 	{
 		TUniquePtr<FPhysicsMoverAsyncInput> InputData = MakeUnique<FPhysicsMoverAsyncInput>();
 		PhysicsMoverComp->ProduceInput_External(PhysicsStep, NumSteps, *InputData);
-		ManagerAsyncInput->AsyncInput.Add(MoveTemp(InputData));
+		if (InputData->IsValid())
+		{
+			//@todo DanH: either check if it's "valid" or just establish the idx and sim associated with the input up-front
+			ManagerAsyncInput->AsyncInput.Add(MoveTemp(InputData));	
+		}
 	}
 }
 
@@ -111,7 +115,7 @@ void UPhysicsMoverManager::OnPostPhysicsTick(FChaosScene*)
 	while (Chaos::TSimCallbackOutputHandle<FPhysicsMoverManagerAsyncOutput> ManagerAsyncOutput = AsyncCallback->PopFutureOutputData_External())
 	{
 		const double OutputTime = ManagerAsyncOutput->InternalTime;
-		for (TWeakObjectPtr<UMoverNetworkPhysicsLiaisonComponent> PhysicsMoverComp : PhysicsMoverComponents)
+		for (TWeakObjectPtr<UMoverNetworkPhysicsLiaisonComponentBase> PhysicsMoverComp : PhysicsMoverComponents)
 		{
 			Chaos::FUniqueIdx Idx = PhysicsMoverComp->GetUniqueIdx();
 			if (Idx.IsValid())
@@ -124,7 +128,7 @@ void UPhysicsMoverManager::OnPostPhysicsTick(FChaosScene*)
 		}
 	}
 
-	for (TWeakObjectPtr<UMoverNetworkPhysicsLiaisonComponent> PhysicsMoverComp : PhysicsMoverComponents)
+	for (TWeakObjectPtr<UMoverNetworkPhysicsLiaisonComponentBase> PhysicsMoverComp : PhysicsMoverComponents)
 	{
 		PhysicsMoverComp->PostPhysicsUpdate_External();
 	}
diff --git a/Engine/Plugins/Experimental/Mover/Source/Mover/Private/PhysicsMover/PhysicsMoverManagerAsyncCallback.cpp b/Engine/Plugins/Experimental/Mover/Source/Mover/Private/PhysicsMover/PhysicsMoverManagerAsyncCallback.cpp
--- a/Engine/Plugins/Experimental/Mover/Source/Mover/Private/PhysicsMover/PhysicsMoverManagerAsyncCallback.cpp
+++ b/Engine/Plugins/Experimental/Mover/Source/Mover/Private/PhysicsMover/PhysicsMoverManagerAsyncCallback.cpp
@@ -2,7 +2,7 @@
 
 #include "PhysicsMover/PhysicsMoverManagerAsyncCallback.h"
 
-#include "Backends/MoverNetworkPhysicsLiaison.h"
+#include "Backends/MoverNetworkPhysicsLiaisonBase.h"
 #include "Chaos/Character/CharacterGroundConstraint.h"
 #include "Chaos/Framework/Parallel.h"
 
@@ -65,7 +65,7 @@ void FPhysicsMoverManagerAsyncCallback::OnPreSimulate_Internal()
 		if (AsyncInput.IsValid())
 		{
 			TUniquePtr<FPhysicsMoverAsyncOutput> AsyncOutput = MakeUnique<FPhysicsMoverAsyncOutput>();
-			AsyncInput.MoverSimulation->OnPreSimulate_Internal(TickParams, AsyncInput, *AsyncOutput);
+			AsyncInput.MoverSimulation->PreSimulate_Internal(TickParams, AsyncInput, *AsyncOutput);
 			ManagerAsyncOutput.PhysicsMoverToAsyncOutput.Add(AsyncInput.MoverIdx, MoveTemp(AsyncOutput));
 		}
 	};
diff --git a/Engine/Plugins/Experimental/Mover/Source/Mover/Public/Backends/MoverNetworkPhysicsLiaison.h b/Engine/Plugins/Experimental/Mover/Source/Mover/Public/Backends/MoverNetworkPhysicsLiaison.h
--- a/Engine/Plugins/Experimental/Mover/Source/Mover/Public/Backends/MoverNetworkPhysicsLiaison.h
+++ b/Engine/Plugins/Experimental/Mover/Source/Mover/Public/Backends/MoverNetworkPhysicsLiaison.h
@@ -4,249 +4,57 @@
 
 #include "Chaos/Character/CharacterGroundConstraint.h"
 #include "Chaos/GeometryParticlesfwd.h"
-#include "Backends/MoverBackendLiaison.h"
-#include "MoverComponent.h"
-#include "Physics/NetworkPhysicsComponent.h"
+#include "Backends/MoverNetworkPhysicsLiaisonBase.h"
 #include "PhysicsMover/PhysicsMoverSimulationTypes.h"
 #include "MoverNetworkPhysicsLiaison.generated.h"
 
 class UCommonLegacyMovementSettings;
 
-
-//////////////////////////////////////////////////////////////////////////
-// Physics networking
-
-USTRUCT()
-struct MOVER_API FNetworkPhysicsMoverInputs : public FNetworkPhysicsData
-{
-	GENERATED_USTRUCT_BODY()
-
-	UPROPERTY()
-	FMoverInputCmdContext InputCmdContext;
-
-	/**  Apply the data onto the network physics component */
-	virtual void ApplyData(UActorComponent* NetworkComponent) const override;
-
-	/**  Build the data from the network physics component */
-	virtual void BuildData(const UActorComponent* NetworkComponent) override;
-
-	/**  Serialize data function that will be used to transfer the struct across the network */
-	bool NetSerialize(FArchive& Ar, class UPackageMap* Map, bool& bOutSuccess);
-
-	/** Interpolate the data in between two inputs data */
-	virtual void InterpolateData(const FNetworkPhysicsData& MinData, const FNetworkPhysicsData& MaxData) override;
-	
-	/** Merge data into this input */
-	virtual void MergeData(const FNetworkPhysicsData& FromData) override;
-
-	/** Check input data is valid - Input is send from client to server, no need to make sure it's reasonable */
-	virtual void ValidateData(const UActorComponent* NetworkComponent) override;
-
-	/** Return string with debug information */
-	virtual const FString DebugData() override;
-};
-
-template<>
-struct TStructOpsTypeTraits<FNetworkPhysicsMoverInputs> : public TStructOpsTypeTraitsBase2<FNetworkPhysicsMoverInputs>
-{
-	enum
-	{
-		WithNetSerializer = true,
-	};
-};
-
-USTRUCT()
-struct MOVER_API FNetworkPhysicsMoverState : public FNetworkPhysicsData
-{
-	GENERATED_USTRUCT_BODY()
-
-	UPROPERTY()
-	FMoverSyncState SyncStateContext;
-
-	/**  Apply the data onto the network physics component */
-	virtual void ApplyData(UActorComponent* NetworkComponent) const override;
-
-	/**  Build the data from the network physics component */
-	virtual void BuildData(const UActorComponent* NetworkComponent) override;
-
-	/**  Serialize data function that will be used to transfer the struct across the network */
-	bool NetSerialize(FArchive& Ar, class UPackageMap* Map, bool& bOutSuccess);
-
-	/** Interpolate the data in between two inputs data */
-	virtual void InterpolateData(const FNetworkPhysicsData& MinData, const FNetworkPhysicsData& MaxData) override;
-
-	/** Return string with debug information */
-	virtual const FString DebugData() override;
-};
-
-template<>
-struct TStructOpsTypeTraits<FNetworkPhysicsMoverState> : public TStructOpsTypeTraitsBase2<FNetworkPhysicsMoverState>
-{
-	enum
-	{
-		WithNetSerializer = true,
-	};
-};
-
-struct MOVER_API FNetworkPhysicsMoverTraits
-{
-	using InputsType = FNetworkPhysicsMoverInputs;
-	using StatesType = FNetworkPhysicsMoverState;
-};
-
-//////////////////////////////////////////////////////////////////////////
-// FSyncStateRecord
-
-class FSyncStateRecord
-{
-public:
-	struct FData
-	{
-		FMoverSyncState SyncState;
-		double SimTime;
-	};
-
-	void Add(double InSimTime, const FMoverSyncState& InSyncState)
-	{
-		CurrentIndex = (CurrentIndex + 1) % 2;
-		Data[CurrentIndex].SimTime = InSimTime;
-		Data[CurrentIndex].SyncState = InSyncState;
-	}
-
-	const FMoverSyncState& GetLatest() const
-	{
-		return Data[CurrentIndex].SyncState;
-	}
-
-	void GetInterpolated(double SimTime, FMoverSyncState& OutSyncState) const
-	{
-		const uint8 PrevIndex = (CurrentIndex + 1) % 2;
-		const double PrevTime = Data[PrevIndex].SimTime;
-		const double CurrTime = Data[CurrentIndex].SimTime;
-
-		if (FMath::IsNearlyEqual(PrevTime, CurrTime) || (SimTime >= CurrTime))
-		{
-			OutSyncState = Data[CurrentIndex].SyncState;
-		}
-		else if (SimTime <= PrevTime)
-		{
-			OutSyncState = Data[PrevIndex].SyncState;
-		}
-		else
-		{
-			const float Alpha = FMath::Clamp((SimTime - PrevTime) / (CurrTime - PrevTime), 0.0f, 1.0f);
-			OutSyncState.Interpolate(&Data[PrevIndex].SyncState, &Data[CurrentIndex].SyncState, Alpha);
-		}
-	}
-
-	void Clear()
-	{
-		CurrentIndex = 1;
-		Data[0].SyncState.Reset();
-		Data[0].SyncState.SyncStateCollection.Empty();
-		Data[1].SyncState.Reset();
-		Data[1].SyncState.SyncStateCollection.Empty();
-	}
-
-	FData Data[2];
-	uint8 CurrentIndex = 1;
-};
-
+//@todo DanH: Rename to UMoverCharacterPhysicsLiaisonComponent
 //////////////////////////////////////////////////////////////////////////
 // UMoverNetworkPhysicsLiaisonComponent
 
 /**
- * MoverPhysicsLiaisonComponent: This component acts as a middleman between an actor's Mover component and the chaos physics network prediction system.
+ * MoverNetworkPhysicsLiaisonComponent: This component acts as a middleman between an actor's Mover component and the chaos physics network prediction system.
  * This class is set on a Mover component as the "back end".
  */
 UCLASS()
-class MOVER_API UMoverNetworkPhysicsLiaisonComponent : public UActorComponent, public IMoverBackendLiaisonInterface
+class MOVER_API UMoverNetworkPhysicsLiaisonComponent : public UMoverNetworkPhysicsLiaisonComponentBase
 {
 	GENERATED_BODY()
 
 public:
 	UMoverNetworkPhysicsLiaisonComponent();
 
-	// IMoverBackendLiaisonInterface
-	virtual float GetCurrentSimTimeMs() override;
-	virtual int32 GetCurrentSimFrame() override;
-#if WITH_EDITOR
-	virtual EDataValidationResult ValidateData(FDataValidationContext& Context, const UMoverComponent& ValidationMoverComp) const override;
-#endif
-
 	// UObject interface
 	virtual void OnRegister() override;
-	virtual void OnUnregister() override;
-	virtual void InitializeComponent() override;
-	virtual void UninitializeComponent() override;
-	virtual bool ShouldCreatePhysicsState() const override;
 	virtual bool HasValidPhysicsState() const override;
 	virtual void OnCreatePhysicsState() override;
 	virtual void OnDestroyPhysicsState() override;
-	virtual bool CanCreatePhysics() const;
-	virtual void BeginPlay() override;
-	virtual void EndPlay(const EEndPlayReason::Type EndPlayReason) override;
-
-	// Required for Network Physics Rewind/Resim data
-	void GetCurrentInputData(OUT FMoverInputCmdContext& InputCmd) const;
-	void SetCurrentInputData(const FMoverInputCmdContext& InputCmd);
-	void GetCurrentStateData(OUT FMoverSyncState& SyncState) const;
-	void SetCurrentStateData(const FMoverSyncState& SyncState);
-	bool ValidateInputData(FMoverInputCmdContext& InputCmd) const;
-
-	// Used by the manager to uniquely identify the component
-	Chaos::FUniqueIdx GetUniqueIdx() const;
 
-	void ProduceInput_External(int32 PhysicsStep, int32 NumSteps, OUT FPhysicsMoverAsyncInput& Input);
-	void ConsumeOutput_External(const FPhysicsMoverAsyncOutput& Output, const double OutputTimeInSeconds);
-	void PostPhysicsUpdate_External();
-
-	void ProcessInputs_Internal(int32 PhysicsStep, float DeltaTimeInSeconds, const FPhysicsMoverAsyncInput& Input) const;
-	virtual void OnPreSimulate_Internal(const FPhysicsMoverSimulationTickParams& TickParams, const FPhysicsMoverAsyncInput& Input, OUT FPhysicsMoverAsyncOutput& SimOutput);
-	void OnContactModification_Internal(const FPhysicsMoverAsyncInput& Input, Chaos::FCollisionContactModifier& Modifier) const;
+#if WITH_EDITOR
+	virtual EDataValidationResult ValidateData(FDataValidationContext& Context, const UMoverComponent& ValidationMoverComp) const override;
+#endif
+	
+	virtual void ConsumeOutput_External(const FPhysicsMoverAsyncOutput& Output, const double OutputTimeInSeconds) override;
+	virtual void PostPhysicsUpdate_External() override;
+	virtual void OnContactModification_Internal(const FPhysicsMoverAsyncInput& Input, Chaos::FCollisionContactModifier& Modifier) const override;
 
 protected:
-	UFUNCTION()
-	void OnComponentPhysicsStateChanged(UPrimitiveComponent* ChangedComponent, EComponentPhysicsStateChange StateChange);
-
-	bool HasValidState() const;
-
+	virtual bool HasValidState() const override;
+
+	virtual bool CanProcessInputs_Internal(const FPhysicsMoverAsyncInput& Input) const override;
+	virtual void PerformProcessInputs_Internal(int32 PhysicsStep, float DeltaTime, const FPhysicsMoverAsyncInput& Input) const override;
+	virtual bool CanSimulate_Internal(const FPhysicsMoverSimulationTickParams& TickParams, const FPhysicsMoverAsyncInput& Input) const override;
+	virtual void PerformPreSimulate_Internal(const FPhysicsMoverSimulationTickParams& TickParams, const FPhysicsMoverAsyncInput& Input, FPhysicsMoverAsyncOutput& Output) const override;
+	virtual Chaos::FPhysicsObject* GetControlledPhysicsObject() const override;
+	virtual void HandleComponentPhysicsStateChanged(UPrimitiveComponent* ChangedComponent, EComponentPhysicsStateChange StateChange) override;
+	
 	void DestroyConstraint();
 	void SetupConstraint();
 
 	void UpdateConstraintSettings();
-	void TeleportParticle(Chaos::FGeometryParticleHandle* Particle, const FVector& Position, const FQuat& Rotation) const;
-	void WakeParticleIfSleeping(Chaos::FGeometryParticleHandle* Particle) const;
-
-	int32 GetNetworkPhysicsTickOffset() const;
-
-	// Time step on the physics thread when using async physics
-	FMoverTimeStep GetCurrentAsyncMoverTimeStep_Internal() const;
-
-	// Time step on the game thread when using async physics. Uses physics results time
-	FMoverTimeStep GetCurrentAsyncMoverTimeStep_External() const;
-
-	// Time step on the game thread when using non-async physics
-	FMoverTimeStep GetCurrentMoverTimeStep(float DeltaSeconds) const;
-
-	void InitializeCachedSyncState();
-
-	Chaos::FPhysicsSolver* GetPhysicsSolver() const;
-
-	// These are written to by the network input and state data
-	FMoverInputCmdContext LatestInputCmd;
-	FMoverSyncState LatestSyncState;
-
+	
 	TUniquePtr<Chaos::FCharacterGroundConstraint> Constraint;
-	TObjectPtr<UMoverComponent> MoverComp;
 	TObjectPtr<UCommonLegacyMovementSettings> CommonMovementSettings;
-
-	TObjectPtr<UNetworkPhysicsComponent> NetworkPhysicsComponent;
-
-	// This stores the previous two outputs from the physics thread
-	// which can be interpolated to get the sync state at a time
-	// that matches the interpolated physics particle.
-	FSyncStateRecord SyncStateRecord;
-
-	bool bUsingAsyncPhysics = false;
 };
\ No newline at end of file
diff --git a/Engine/Plugins/Experimental/Mover/Source/Mover/Public/Backends/MoverNetworkPhysicsLiaisonBase.h b/Engine/Plugins/Experimental/Mover/Source/Mover/Public/Backends/MoverNetworkPhysicsLiaisonBase.h
new file mode 100644
--- /dev/null
+++ b/Engine/Plugins/Experimental/Mover/Source/Mover/Public/Backends/MoverNetworkPhysicsLiaisonBase.h
@@ -0,0 +1,292 @@
+// Copyright Epic Games, Inc. All Rights Reserved.
+
+#pragma once
+
+#include "Chaos/GeometryParticlesfwd.h"
+#include "Backends/MoverBackendLiaison.h"
+#include "MoverComponent.h"
+#include "Physics/NetworkPhysicsComponent.h"
+#include "PhysicsMover/PhysicsMoverSimulationTypes.h"
+#include "MoverNetworkPhysicsLiaisonBase.generated.h"
+
+//////////////////////////////////////////////////////////////////////////
+// Physics networking
+
+USTRUCT()
+struct MOVER_API FNetworkPhysicsMoverInputs : public FNetworkPhysicsData
+{
+	GENERATED_USTRUCT_BODY()
+
+	UPROPERTY()
+	FMoverInputCmdContext InputCmdContext;
+
+	/**  Apply the data onto the network physics component */
+	virtual void ApplyData(UActorComponent* NetworkComponent) const override;
+
+	/**  Build the data from the network physics component */
+	virtual void BuildData(const UActorComponent* NetworkComponent) override;
+
+	/**  Serialize data function that will be used to transfer the struct across the network */
+	bool NetSerialize(FArchive& Ar, class UPackageMap* Map, bool& bOutSuccess);
+
+	/** Interpolate the data in between two inputs data */
+	virtual void InterpolateData(const FNetworkPhysicsData& MinData, const FNetworkPhysicsData& MaxData) override;
+	
+	/** Merge data into this input */
+	virtual void MergeData(const FNetworkPhysicsData& FromData) override;
+
+	/** Check input data is valid - Input is send from client to server, no need to make sure it's reasonable */
+	virtual void ValidateData(const UActorComponent* NetworkComponent) override;
+
+	virtual bool CompareData(const FNetworkPhysicsData& PredictedData) override;
+
+	/** Return string with debug information */
+	virtual const FString DebugData() override;
+};
+
+template<>
+struct TStructOpsTypeTraits<FNetworkPhysicsMoverInputs> : public TStructOpsTypeTraitsBase2<FNetworkPhysicsMoverInputs>
+{
+	enum
+	{
+		WithNetSerializer = true,
+	};
+};
+
+USTRUCT()
+struct MOVER_API FNetworkPhysicsMoverState : public FNetworkPhysicsData
+{
+	GENERATED_USTRUCT_BODY()
+
+	UPROPERTY()
+	FMoverSyncState SyncStateContext;
+
+	/**  Apply the data onto the network physics component */
+	virtual void ApplyData(UActorComponent* NetworkComponent) const override;
+
+	/**  Build the data from the network physics component */
+	virtual void BuildData(const UActorComponent* NetworkComponent) override;
+
+	/**  Serialize data function that will be used to transfer the struct across the network */
+	bool NetSerialize(FArchive& Ar, class UPackageMap* Map, bool& bOutSuccess);
+
+	/** Interpolate the data in between two inputs data */
+	virtual void InterpolateData(const FNetworkPhysicsData& MinData, const FNetworkPhysicsData& MaxData) override;
+
+	virtual bool CompareData(const FNetworkPhysicsData& PredictedData) override;
+	
+	/** Return string with debug information */
+	virtual const FString DebugData() override;
+};
+
+template<>
+struct TStructOpsTypeTraits<FNetworkPhysicsMoverState> : public TStructOpsTypeTraitsBase2<FNetworkPhysicsMoverState>
+{
+	enum
+	{
+		WithNetSerializer = true,
+	};
+};
+
+struct MOVER_API FNetworkPhysicsMoverTraits
+{
+	using InputsType = FNetworkPhysicsMoverInputs;
+	using StatesType = FNetworkPhysicsMoverState;
+};
+
+//////////////////////////////////////////////////////////////////////////
+// FSyncStateRecord
+
+class FSyncStateRecord
+{
+public:
+	struct FData
+	{
+		FMoverSyncState SyncState;
+		double SimTime;
+	};
+
+	void Add(double InSimTime, const FMoverSyncState& InSyncState)
+	{
+		CurrentIndex = (CurrentIndex + 1) % 2;
+		Data[CurrentIndex].SimTime = InSimTime;
+		Data[CurrentIndex].SyncState = InSyncState;
+	}
+
+	const FMoverSyncState& GetLatest() const
+	{
+		return Data[CurrentIndex].SyncState;
+	}
+
+	void GetInterpolated(double SimTime, FMoverSyncState& OutSyncState) const
+	{
+		const uint8 PrevIndex = (CurrentIndex + 1) % 2;
+		const double PrevTime = Data[PrevIndex].SimTime;
+		const double CurrTime = Data[CurrentIndex].SimTime;
+
+		if (FMath::IsNearlyEqual(PrevTime, CurrTime) || (SimTime >= CurrTime))
+		{
+			OutSyncState = Data[CurrentIndex].SyncState;
+		}
+		else if (SimTime <= PrevTime)
+		{
+			OutSyncState = Data[PrevIndex].SyncState;
+		}
+		else
+		{
+			const float Alpha = FMath::Clamp((SimTime - PrevTime) / (CurrTime - PrevTime), 0.0f, 1.0f);
+			OutSyncState.Interpolate(&Data[PrevIndex].SyncState, &Data[CurrentIndex].SyncState, Alpha);
+		}
+	}
+
+	void Clear()
+	{
+		CurrentIndex = 1;
+		Data[0].SyncState.Reset();
+		Data[0].SyncState.SyncStateCollection.Empty();
+		Data[1].SyncState.Reset();
+		Data[1].SyncState.SyncStateCollection.Empty();
+	}
+
+	FData Data[2];
+	uint8 CurrentIndex = 1;
+};
+
+//////////////////////////////////////////////////////////////////////////
+// UMoverNetworkPhysicsLiaisonComponentBase
+
+/**
+ * Base class for liaison components to act as a middleman between an actor's Mover component and the Chaos physics network prediction system to move
+ * the Mover's updated component on the physics thread (PT).
+ * 
+ * This is accomplished by registering a FPhysicsMoverManagerAsyncCallback and exposing the game thread (GT) and PT methods for the TSimCallback to
+ * virtual methods on this liaison.
+ */
+UCLASS(Abstract, Within = MoverComponent)
+class MOVER_API UMoverNetworkPhysicsLiaisonComponentBase : public UActorComponent, public IMoverBackendLiaisonInterface
+{
+	GENERATED_BODY()
+
+public:
+	UMoverNetworkPhysicsLiaisonComponentBase();
+
+	// IMoverBackendLiaisonInterface
+	virtual float GetCurrentSimTimeMs() override;
+	virtual int32 GetCurrentSimFrame() override;
+#if WITH_EDITOR
+	virtual EDataValidationResult ValidateData(FDataValidationContext& Context, const UMoverComponent& ValidationMoverComp) const override;
+#endif
+
+	// UObject interface
+	virtual void OnRegister() override;
+	virtual void OnUnregister() override;
+	virtual void InitializeComponent() override;
+	virtual void UninitializeComponent() override;
+	virtual bool ShouldCreatePhysicsState() const override;
+	virtual bool CanCreatePhysics() const;
+	virtual void BeginPlay() override;
+	virtual void EndPlay(const EEndPlayReason::Type EndPlayReason) override;
+
+	// Required for Network Physics Rewind/Resim data
+	void GetCurrentInputData(OUT FMoverInputCmdContext& InputCmd) const;
+	void SetCurrentInputData(const FMoverInputCmdContext& InputCmd);
+	void GetCurrentStateData(OUT FMoverSyncState& SyncState) const;
+	void SetCurrentStateData(const FMoverSyncState& SyncState);
+	bool ValidateInputData(FMoverInputCmdContext& InputCmd) const;
+
+	// Used by the manager to uniquely identify the component
+	Chaos::FUniqueIdx GetUniqueIdx() const;
+
+	/** These methods correspond to the methods of the same name on the TSimCallback that we register */
+	// "External" = Called on the Game Thread (GT)
+	void ProduceInput_External(int32 PhysicsStep, int32 NumSteps, OUT FPhysicsMoverAsyncInput& Input);
+	virtual void ConsumeOutput_External(const FPhysicsMoverAsyncOutput& Output, double OutputTimeInSeconds);
+	virtual void PostPhysicsUpdate_External();
+
+	// "Internal" = Called on the Physics Thread (PT)
+	void ProcessInputs_Internal(int32 PhysicsStep, float DeltaTime, const FPhysicsMoverAsyncInput& Input) const;
+	void PreSimulate_Internal(const FPhysicsMoverSimulationTickParams& TickParams, const FPhysicsMoverAsyncInput& Input, OUT FPhysicsMoverAsyncOutput& Output) const;
+	virtual void OnContactModification_Internal(const FPhysicsMoverAsyncInput& Input, Chaos::FCollisionContactModifier& Modifier) const;
+	
+protected:
+	/** Catchall to verify that the state of this object is valid (called from both the PT and the GT) */
+	virtual bool HasValidState() const;
+	
+	virtual void PerformProduceInput_External(float DeltaTime, OUT FPhysicsMoverAsyncInput& Input);
+	
+	/**
+	 * Called at the beginning of PreSimulate_Internal to verify that everything is in order.
+	 * Separated to ensure that the validation occurs and reduce redundant validity checks in PerformProcessInputs_Internal() overrides
+	 */
+	virtual bool CanProcessInputs_Internal(const FPhysicsMoverAsyncInput& Input) const;
+
+	/** Actually do the ProcessInputs work on the PT. CanProcessInputs_Internal is guaranteed to have returned true if this is called. */
+	virtual void PerformProcessInputs_Internal(int32 PhysicsStep, float DeltaTime, const FPhysicsMoverAsyncInput& Input) const;
+
+	/** 
+	 * Called at the beginning of PreSimulate_Internal to verify that everything is in order.
+	 * Default behavior is simply to call CanProcessInput_Internal(). This is defined separately in case the TickParams or OutputState need validation as well.
+	 * Separated to ensure that the validation occurs and reduce redundant validity checks in PerformPreSimulate_Internal() overrides
+	 */
+	virtual bool CanSimulate_Internal(const FPhysicsMoverSimulationTickParams& TickParams, const FPhysicsMoverAsyncInput& Input) const;
+	
+	/** Actually do the PreSimulate work on the PT. CanSimulate_Internal is guaranteed to have returned true if this is called. */
+	virtual void PerformPreSimulate_Internal(const FPhysicsMoverSimulationTickParams& TickParams, const FPhysicsMoverAsyncInput& Input, OUT FPhysicsMoverAsyncOutput& Output) const;
+	
+	/** Override in child classes to return the chaos particle for the component the Mover wants to move. */
+	virtual Chaos::FPhysicsObject* GetControlledPhysicsObject() const;
+
+	/** Convenience getters to skip the read interface boilerplate on the Game (aka External) or Physics (aka Internal) threads */
+	Chaos::FPBDRigidParticle* GetControlledParticle_External() const;
+	Chaos::FPBDRigidParticleHandle* GetControlledParticle_Internal() const;
+	
+	UFUNCTION()
+	virtual void HandleComponentPhysicsStateChanged(UPrimitiveComponent* ChangedComponent, EComponentPhysicsStateChange StateChange);
+
+	UFUNCTION()
+	void HandleOwningPawnControllerChanged_Server(APawn* OwnerPawn, AController* OldController, AController* NewController);
+	
+	Chaos::FPhysicsSolver* GetPhysicsSolver() const;
+	
+	void TeleportParticleBy_Internal(Chaos::FGeometryParticleHandle& Particle, const FVector& PositionDelta, const FQuat& RotationDelta) const;
+	void WakeParticleIfSleeping(Chaos::FGeometryParticleHandle* Particle) const;
+
+	int32 GetNetworkPhysicsTickOffset() const;
+
+	// Time step on the physics thread when using async physics
+	FMoverTimeStep GetCurrentAsyncMoverTimeStep_Internal() const;
+
+	// Time step on the game thread when using async physics. Uses physics results time
+	FMoverTimeStep GetCurrentAsyncMoverTimeStep_External() const;
+
+	// Time step on the game thread when using non-async physics
+	FMoverTimeStep GetCurrentMoverTimeStep(float DeltaSeconds) const;
+	
+	void InitializeCachedSyncState();
+
+	UMoverComponent& GetMoverComponent() const;
+
+	/**
+	 * These two properties are ultimately controlled by the NetworkPhysicsComponent (NPC) via FNetworkPhysicsMoverInputs and FNetworkPhysicsMoverState (at top of file).
+	 * 
+	 * For input, the process responsible for generating input (owning client for pawns, server for anything else) does so, and the NPC replicates
+	 * it to the server (if produced on a client) and then to all clients. For input generated on the server, the physics frame/step in which it's processed on the client
+	 * does not match the server physics step, even if account for the NPC's NetworkFrameOffset. This is because the server is in "the past" relative to the client.
+	 *
+	 * For the sync state, the server is always the authority. By default, the sync state from the server is only applied on the client after a rewind, but this can be
+	 * changed via PhysicsReplicationCVars::ResimulationCVars::bApplySimProxyStateAtRuntime (or the corresponding override in FNetworkPhysicsSettingsNetworkPhysicsComponent).
+	 */
+	FMoverInputCmdContext LatestInputCmd;
+	mutable FMoverSyncState LatestSyncState;
+
+	UPROPERTY(Transient)
+	TObjectPtr<UNetworkPhysicsComponent> NetworkPhysicsComponent;
+
+	// This stores the previous two outputs from the physics thread
+	// which can be interpolated to get the sync state at a time
+	// that matches the interpolated physics particle.
+	FSyncStateRecord SyncStateRecord;
+
+	bool bCachedInputIsValid = false;
+	bool bUsingAsyncPhysics = false;
+};
\ No newline at end of file
diff --git a/Engine/Plugins/Experimental/Mover/Source/Mover/Public/MoverComponent.h b/Engine/Plugins/Experimental/Mover/Source/Mover/Public/MoverComponent.h
--- a/Engine/Plugins/Experimental/Mover/Source/Mover/Public/MoverComponent.h
+++ b/Engine/Plugins/Experimental/Mover/Source/Mover/Public/MoverComponent.h
@@ -605,6 +605,7 @@ private:
 
 	friend class UBaseMovementMode;
 	friend class UMoverNetworkPhysicsLiaisonComponent;
+	friend class UMoverNetworkPhysicsLiaisonComponentBase;
 	friend class UMoverDebugComponent;
 	friend class UBasedMovementUtils;
 };
diff --git a/Engine/Plugins/Experimental/Mover/Source/Mover/Public/PhysicsMover/PhysicsMoverManager.h b/Engine/Plugins/Experimental/Mover/Source/Mover/Public/PhysicsMover/PhysicsMoverManager.h
--- a/Engine/Plugins/Experimental/Mover/Source/Mover/Public/PhysicsMover/PhysicsMoverManager.h
+++ b/Engine/Plugins/Experimental/Mover/Source/Mover/Public/PhysicsMover/PhysicsMoverManager.h
@@ -9,7 +9,7 @@
 //////////////////////////////////////////////////////////////////////////
 
 class FChaosScene;
-class UMoverNetworkPhysicsLiaisonComponent;
+class UMoverNetworkPhysicsLiaisonComponentBase;
 
 UCLASS()
 class MOVER_API UPhysicsMoverManager : public UWorldSubsystem
@@ -24,11 +24,11 @@ public:
 	void InjectInputs_External(int32 PhysicsStep, int32 NumSteps);
 	void OnPostPhysicsTick(FChaosScene* Scene);
 
-	void RegisterPhysicsMoverComponent(TWeakObjectPtr<UMoverNetworkPhysicsLiaisonComponent> InPhysicsMoverComp);
-	void UnregisterPhysicsMoverComponent(TWeakObjectPtr<UMoverNetworkPhysicsLiaisonComponent> InPhysicsMoverComp);
+	void RegisterPhysicsMoverComponent(TWeakObjectPtr<UMoverNetworkPhysicsLiaisonComponentBase> InPhysicsMoverComp);
+	void UnregisterPhysicsMoverComponent(TWeakObjectPtr<UMoverNetworkPhysicsLiaisonComponentBase> InPhysicsMoverComp);
 
 private:
-	TArray<TWeakObjectPtr<UMoverNetworkPhysicsLiaisonComponent>> PhysicsMoverComponents;
+	TArray<TWeakObjectPtr<UMoverNetworkPhysicsLiaisonComponentBase>> PhysicsMoverComponents;
 	FDelegateHandle InjectInputsExternalCallbackHandle;
 	FDelegateHandle PhysScenePostTickCallbackHandle;
 	class FPhysicsMoverManagerAsyncCallback* AsyncCallback = nullptr;
diff --git a/Engine/Plugins/Experimental/Mover/Source/Mover/Public/PhysicsMover/PhysicsMoverSimulationTypes.h b/Engine/Plugins/Experimental/Mover/Source/Mover/Public/PhysicsMover/PhysicsMoverSimulationTypes.h
--- a/Engine/Plugins/Experimental/Mover/Source/Mover/Public/PhysicsMover/PhysicsMoverSimulationTypes.h
+++ b/Engine/Plugins/Experimental/Mover/Source/Mover/Public/PhysicsMover/PhysicsMoverSimulationTypes.h
@@ -56,7 +56,7 @@ struct FPhysicsMoverAsyncInput
 	mutable FMoverInputCmdContext InputCmd;
 	mutable FMoverSyncState SyncState;
 
-	TWeakObjectPtr<class UMoverNetworkPhysicsLiaisonComponent> MoverSimulation;
+	TWeakObjectPtr<class UMoverNetworkPhysicsLiaisonComponentBase> MoverSimulation;
 	Chaos::FUniqueIdx MoverIdx;
 };
 
