[Chaos Mover] Fix nonunity and shipping compile errors

#tests Validated against Win64 shipping

diff --git a/Engine/Plugins/Experimental/Mover/Source/Mover/Private/Backends/MoverNetworkPhysicsLiaison.cpp b/Engine/Plugins/Experimental/Mover/Source/Mover/Private/Backends/MoverNetworkPhysicsLiaison.cpp
--- a/Engine/Plugins/Experimental/Mover/Source/Mover/Private/Backends/MoverNetworkPhysicsLiaison.cpp
+++ b/Engine/Plugins/Experimental/Mover/Source/Mover/Private/Backends/MoverNetworkPhysicsLiaison.cpp
@@ -23,6 +23,8 @@
 
 #define LOCTEXT_NAMESPACE "Mover"
 
+extern FPhysicsDrivenMotionDebugParams GPhysicsDrivenMotionDebugParams;
+
 //////////////////////////////////////////////////////////////////////////
 // UMoverNetworkPhysicsLiaisonComponent
 
diff --git a/Engine/Plugins/Experimental/Mover/Source/Mover/Private/Backends/MoverNetworkPhysicsLiaisonBase.cpp b/Engine/Plugins/Experimental/Mover/Source/Mover/Private/Backends/MoverNetworkPhysicsLiaisonBase.cpp
--- a/Engine/Plugins/Experimental/Mover/Source/Mover/Private/Backends/MoverNetworkPhysicsLiaisonBase.cpp
+++ b/Engine/Plugins/Experimental/Mover/Source/Mover/Private/Backends/MoverNetworkPhysicsLiaisonBase.cpp
@@ -230,7 +230,7 @@ int32 UMoverNetworkPhysicsLiaisonComponentBase::GetCurrentSimFrame()
 	// Note: this is implicitly an _External function
 	if (Chaos::FPhysicsSolver* Solver = GetPhysicsSolver())
 	{
-		return Solver->GetCurrentFrame() + GetNetworkPhysicsTickOffset();
+		return Solver->GetCurrentFrame() + GetNetworkPhysicsTickOffset_External();
 	}
 	return 0;
 }
@@ -434,35 +434,6 @@ Chaos::FUniqueIdx UMoverNetworkPhysicsLiaisonComponentBase::GetUniqueIdx() const
 	return Chaos::FUniqueIdx();
 }
 
-int32 UMoverNetworkPhysicsLiaisonComponentBase::GetNetworkPhysicsTickOffset() const
-{
-	if (NetworkPhysicsComponent && !NetworkPhysicsComponent->HasServerWorld())
-	{
-		if (Chaos::IsInPhysicsThreadContext())
-		{
-			if (const FAsyncNetworkPhysicsComponent* NPCPhysicsThreadAPI = NetworkPhysicsComponent->GetNetworkPhysicsComponent_Internal())
-			{
-				return NPCPhysicsThreadAPI->GetNetworkPhysicsTickOffset();
-			}
-		}
-		else
-		{
-			const APlayerController* PlayerController = NetworkPhysicsComponent->GetPlayerController();
-			if (!PlayerController)
-			{
-				PlayerController = GetWorld()->GetFirstPlayerController();
-			}
-
-			if (PlayerController)
-			{
-				return PlayerController->GetNetworkPhysicsTickOffset();
-			}	
-		}
-	}
-
-	return 0;
-}
-
 FMoverTimeStep UMoverNetworkPhysicsLiaisonComponentBase::GetCurrentAsyncMoverTimeStep_Internal() const
 {
 	Chaos::EnsureIsInPhysicsThreadContext();
@@ -471,7 +442,7 @@ FMoverTimeStep UMoverNetworkPhysicsLiaisonComponentBase::GetCurrentAsyncMoverTim
 	FMoverTimeStep TimeStep;
 	if (Chaos::FPhysicsSolver* Solver = GetPhysicsSolver())
 	{
-		TimeStep.ServerFrame = Solver->GetCurrentFrame() + GetNetworkPhysicsTickOffset();
+		TimeStep.ServerFrame = Solver->GetCurrentFrame() + GetNetworkPhysicsTickOffset_Internal();
 		TimeStep.StepMs = Solver->GetAsyncDeltaTime() * 1000.0f;
 		TimeStep.BaseSimTimeMs = TimeStep.ServerFrame * TimeStep.StepMs;
 		TimeStep.bIsResimulating = Solver->GetEvolution()->IsResimming();
@@ -488,7 +459,7 @@ FMoverTimeStep UMoverNetworkPhysicsLiaisonComponentBase::GetCurrentAsyncMoverTim
 	FMoverTimeStep TimeStep;
 	if (Chaos::FPhysicsSolver* Solver = GetPhysicsSolver())
 	{
-		const int32 Offset = GetNetworkPhysicsTickOffset();
+		const int32 Offset = GetNetworkPhysicsTickOffset_External();
 		TimeStep.ServerFrame = Solver->GetCurrentFrame() + Offset;
 		TimeStep.StepMs = Solver->GetAsyncDeltaTime() * 1000.0f;
 		TimeStep.BaseSimTimeMs = Solver->GetPhysicsResultsTime_External() * 1000.0f + Offset * TimeStep.StepMs;
@@ -872,4 +843,36 @@ void UMoverNetworkPhysicsLiaisonComponentBase::WakeParticleIfSleeping(Chaos::FGe
 	}
 }
 
+int32 UMoverNetworkPhysicsLiaisonComponentBase::GetNetworkPhysicsTickOffset_Internal() const
+{
+	if (NetworkPhysicsComponent && !NetworkPhysicsComponent->HasServerWorld())
+	{
+		if (const FAsyncNetworkPhysicsComponent* NPCPhysicsThreadAPI = NetworkPhysicsComponent->GetNetworkPhysicsComponent_Internal())
+		{
+			return NPCPhysicsThreadAPI->GetNetworkPhysicsTickOffset();
+		}
+	}
+
+	return 0;
+}
+
+int32 UMoverNetworkPhysicsLiaisonComponentBase::GetNetworkPhysicsTickOffset_External() const
+{
+	if (NetworkPhysicsComponent && !NetworkPhysicsComponent->HasServerWorld())
+	{
+		const APlayerController* PlayerController = NetworkPhysicsComponent->GetPlayerController();
+		if (!PlayerController)
+		{
+			PlayerController = GetWorld()->GetFirstPlayerController();
+		}
+
+		if (PlayerController)
+		{
+			return PlayerController->GetNetworkPhysicsTickOffset();
+		}
+	}
+	
+	return 0;
+}
+
 #undef LOCTEXT_NAMESPACE
diff --git a/Engine/Plugins/Experimental/Mover/Source/Mover/Public/Backends/MoverNetworkPhysicsLiaisonBase.h b/Engine/Plugins/Experimental/Mover/Source/Mover/Public/Backends/MoverNetworkPhysicsLiaisonBase.h
--- a/Engine/Plugins/Experimental/Mover/Source/Mover/Public/Backends/MoverNetworkPhysicsLiaisonBase.h
+++ b/Engine/Plugins/Experimental/Mover/Source/Mover/Public/Backends/MoverNetworkPhysicsLiaisonBase.h
@@ -251,7 +251,8 @@ protected:
 	void TeleportParticleBy_Internal(Chaos::FGeometryParticleHandle& Particle, const FVector& PositionDelta, const FQuat& RotationDelta) const;
 	void WakeParticleIfSleeping(Chaos::FGeometryParticleHandle* Particle) const;
 
-	int32 GetNetworkPhysicsTickOffset() const;
+	int32 GetNetworkPhysicsTickOffset_Internal() const;
+	int32 GetNetworkPhysicsTickOffset_External() const;
 
 	// Time step on the physics thread when using async physics
 	FMoverTimeStep GetCurrentAsyncMoverTimeStep_Internal() const;
