[Chaos Mover] Adding pathed physics-based movement to Mover. "Pathed movement" means movement along an explicitly defined path, like a moving platform. The importance of being physics-based is twofold:
1. By being part of the physics simulation and updating movement on the physics thread, interaction is much smoother when using PBCM (physics-based character movement). Trying to platform with a PBCM character on a game-thread-translated platform has issues arising from the different frame rates between PT and GT.
2. The component moved along the path can still be part of the physics sim and respond to forces. This is accomplished with a very stiff joint constraint that "pulls" the component along the path. Alternatively the component can also just be moved kinematically.

High level breakdown of the parts involved:

UPathedPhysicsMoverComponent
- UMoverComponent subclass that expects to work explicitly with a UMoverPathedPhysicsLiaisonComponent as its BackendClass. Longer-term, this should be doable with either the component or the liaison subclass (not both), pending further development of Mover
- Since the liaison is internal to the mover component, the mover comp itself needs to be the public-facing component/API for controlling playback of movement along the defined path

UMoverPathedPhysicsLiaisonComponent
- UMoverNetworkPhysicsLiaisonComponent subclass that does the majority of the work of managing the pathed movement
- Creates the joint, processes path-specific input and sync states, and moves either the joint endpoint or the updated component's particle kinematically along the path
- Expects to work only with UPathedPhysicsMovementModes (defining and filtering by this in the details panel is still a todo)

UPathedPhysicsMovementMode
- Defines a path to follow and how long it should take to traverse
- The actual path is a compound result of one or more patterns defined in the mode

UPathedMovementPatternBase
- Relatively simple building blocks for defining more complex aggregate paths
- One-pattern paths are doable and likely the most common case (point-based like InterpToComponent, spline-based, and elliptical exist atm), but stacking multiple patterns can be useful too
- Ex: Two orthogonal straight lines = sin curve, Circle + straight line = corkscrew, etc.

UPathedPhysicsDebugDrawComponent
- Renders the path for each pattern and the compound result
- Controls a transient dupe of the root mesh component for previewing its location at various points along the path

Still to do:
- The actual joint is still really floaty and loose when used
- The PBCM character still jitters a lot when standing on a pathed platform, even without the joint enabled

diff --git a/Engine/Plugins/Experimental/Mover/Source/Mover/Private/Backends/MoverPathedPhysicsLiaison.cpp b/Engine/Plugins/Experimental/Mover/Source/Mover/Private/Backends/MoverPathedPhysicsLiaison.cpp
new file mode 100644
--- /dev/null
+++ b/Engine/Plugins/Experimental/Mover/Source/Mover/Private/Backends/MoverPathedPhysicsLiaison.cpp
@@ -0,0 +1,426 @@
+// Copyright Epic Games, Inc. All Rights Reserved.
+
+#include "Backends/MoverPathedPhysicsLiaison.h"
+
+#include "Chaos/PhysicsObjectInternalInterface.h"
+#include "PBDRigidsSolver.h"
+#include "TimerManager.h"
+#include "Kismet/KismetMathLibrary.h"
+#include "PhysicsMover/PathedMovement/PathedMovementMode.h"
+#include "PhysicsMover/PathedMovement/PathedPhysicsMoverComponent.h"
+#include "Runtime/Experimental/Chaos/Private/Chaos/PhysicsObjectInternal.h"
+
+#include UE_INLINE_GENERATED_CPP_BY_NAME(MoverPathedPhysicsLiaison)
+
+#define LOCTEXT_NAMESPACE "Mover"
+
+TAutoConsoleVariable<float> CVarPathedPhysicsLatencyDelayMs(
+	TEXT("Mover.PathedPhysics.StartMovingLatencyDelayMs"),
+	200.f,
+	TEXT("How long (in ms) to delay starting pathed movement on the server to give the client(s) time to find out about it"));
+
+
+UMoverPathedPhysicsLiaisonComponent::UMoverPathedPhysicsLiaisonComponent()
+	: PhysicsUserData(&ConstraintInstance)
+{
+	bWantsInitializeComponent = true;
+	
+	if (Chaos::FPhysicsSolverBase::IsNetworkPhysicsPredictionEnabled())
+	{
+		if (AActor* MyActor = GetOwner())
+		{
+			// Regardless of how rewinds are triggered, they should always result in a resimulation
+			MyActor->SetPhysicsReplicationMode(EPhysicsReplicationMode::Resimulation);
+		}
+	}
+}
+
+void UMoverPathedPhysicsLiaisonComponent::InitializeComponent()
+{
+	Super::InitializeComponent();
+	
+	UPathedPhysicsMoverComponent& MoverComp = GetPathedMoverComp();
+	MoverComp.OnMovementModeChanged.AddUniqueDynamic(this, &ThisClass::HandleMovementModeChanged);
+
+	if (const UPathedPhysicsMovementMode* InitialPathedMode = MoverComp.FindMode_Mutable<UPathedPhysicsMovementMode>(MoverComp.StartingMovementMode))
+	{
+		ApplyPathModeConfig(*InitialPathedMode);
+	}
+}
+
+bool UMoverPathedPhysicsLiaisonComponent::HasValidPhysicsState() const
+{
+	return ConstraintHandle.IsValid();
+}
+
+void UMoverPathedPhysicsLiaisonComponent::OnCreatePhysicsState()
+{
+	Super::OnCreatePhysicsState();
+
+	CreateTargetJoint();
+}
+
+void UMoverPathedPhysicsLiaisonComponent::OnDestroyPhysicsState()
+{
+	DestroyTargetJoint();
+
+	Super::OnDestroyPhysicsState();
+}
+
+void UMoverPathedPhysicsLiaisonComponent::ConsumeOutput_External(const FPhysicsMoverAsyncOutput& Output, const double OutputTimeInSeconds)
+{
+	Super::ConsumeOutput_External(Output, OutputTimeInSeconds);
+	
+	if (const FPathedPhysicsMovementState* OutputMoveState = Output.SyncState.SyncStateCollection.FindDataByType<FPathedPhysicsMovementState>())
+	{
+		const bool bWasMoving = IsMoving();
+		const bool bWasJointEnabled = IsJointEnabled();
+
+		Inputs_External = OutputMoveState->MutableProps;		
+
+		if (bWasMoving != IsMoving())
+		{
+			GetPathedMoverComp().NotifyIsMovingChanged(IsMoving());
+		}
+
+		if (bWasJointEnabled != IsJointEnabled() && NetworkPhysicsComponent && !NetworkPhysicsComponent->HasServerWorld())
+		{
+			if (UPathedPhysicsMovementMode* PathedMode = GetMoverComponent().GetActiveMode_Mutable<UPathedPhysicsMovementMode>())
+			{
+				// To have bIsJointEnabled on the mode behave like a normal replicated property, forward changes to it from physics input replication to the mode on the client
+				// We don't listen for changes to the property on the client, so this is done purely so anything external that checks UPathedPhysicsMovementMode::IsUsingJoint()
+				// on the client will still get the correct answer.
+				//@todo DanH: Still not right - this will get rejected because client, and we still do want to respond to the change on the client as well to change the CompareState setting
+				PathedMode->SetUseJointConstraint(IsJointEnabled());
+			}
+		}
+	}
+}
+
+void UMoverPathedPhysicsLiaisonComponent::SetPathOrigin(const FTransform& NewPathOrigin)
+{
+	if (!GetPathOrigin().Equals(NewPathOrigin) && GetOwnerRole() == ROLE_Authority)
+	{
+		Inputs_External.PathOrigin = NewPathOrigin;
+	}
+}
+
+void UMoverPathedPhysicsLiaisonComponent::SetPlaybackDirection(bool bPlayForward)
+{
+	if (bPlayForward == IsInReverse() && GetOwnerRole() == ROLE_Authority)
+	{
+		Inputs_External.bIsInReverse = !bPlayForward;
+	}
+}
+
+void UMoverPathedPhysicsLiaisonComponent::SetIsMoving(bool bShouldMove, float StartDelay)
+{
+	if (bShouldMove != IsMoving() && GetOwnerRole() == ROLE_Authority)
+	{
+		if (!bShouldMove)
+		{
+			Inputs_External.MovementStartFrame = INDEX_NONE;
+		}
+		else if (Chaos::FPhysicsSolver* Solver = GetPhysicsSolver(); ensure(Solver))
+		{
+			if (bUsingAsyncPhysics)
+			{
+				const int32 CurrentFrame = Solver->GetCurrentFrame();
+				const int32 StartDelayFrames = FMath::FloorToInt32(StartDelay / Solver->GetAsyncDeltaTime());
+				const int32 LatencyDelayFrames = FMath::CeilToInt32(CVarPathedPhysicsLatencyDelayMs.GetValueOnGameThread() / (Solver->GetAsyncDeltaTime() * 1000.f));
+
+				Inputs_External.MovementStartFrame = CurrentFrame + StartDelayFrames + LatencyDelayFrames;
+			}
+			else if (const UWorld* World = GetWorld())
+			{
+				// In a standalone game that isn't using async physics, the start delay can't be reliably converted to a physics frame.
+				// So we have to use a world timer instead of delaying the start frame
+				FTimerManager& TimerManager = World->GetTimerManager();
+				if (DelayedStartTimerHandle.IsValid())
+				{
+					TimerManager.ClearTimer(DelayedStartTimerHandle);
+				}
+
+				if (StartDelay > 0.f)
+				{
+					TimerManager.SetTimer(DelayedStartTimerHandle, FTimerDelegate::CreateWeakLambda(this, [this]()
+					{
+						Inputs_External.MovementStartFrame = GetCurrentSimFrame();
+					}), StartDelay, false);
+				}
+				else
+				{
+					Inputs_External.MovementStartFrame = Solver->GetCurrentFrame();
+				}
+			}			
+		}
+	}
+}
+
+void UMoverPathedPhysicsLiaisonComponent::SetPlaybackBehavior(EPathedPhysicsPlaybackBehavior PlaybackBehavior)
+{
+	if (PlaybackBehavior != GetPlaybackBehavior() && GetOwnerRole() == ROLE_Authority)
+	{
+		Inputs_External.PlaybackBehavior = PlaybackBehavior;
+	}
+}
+
+UPathedPhysicsMoverComponent& UMoverPathedPhysicsLiaisonComponent::GetPathedMoverComp() const
+{
+	return *GetOuterUPathedPhysicsMoverComponent();
+}
+
+void UMoverPathedPhysicsLiaisonComponent::HandleComponentPhysicsStateChanged(UPrimitiveComponent* ChangedComponent, EComponentPhysicsStateChange StateChange)
+{
+	Super::HandleComponentPhysicsStateChanged(ChangedComponent, StateChange);
+	if (StateChange == EComponentPhysicsStateChange::Destroyed)
+	{
+		DestroyTargetJoint();
+	}
+	else if (StateChange == EComponentPhysicsStateChange::Created)
+	{
+		CreateTargetJoint();
+	}
+}
+
+void UMoverPathedPhysicsLiaisonComponent::PerformProduceInput_External(float DeltaTime, FPhysicsMoverAsyncInput& Input)
+{
+	Super::PerformProduceInput_External(DeltaTime, Input);
+	
+	FPathedPhysicsMovementInputs& InputState = Input.InputCmd.InputCollection.FindOrAddMutableDataByType<FPathedPhysicsMovementInputs>();
+	InputState.Props = Inputs_External;
+}
+
+bool UMoverPathedPhysicsLiaisonComponent::CanProcessInputs_Internal(const FPhysicsMoverAsyncInput& Input) const
+{
+	if (Super::CanProcessInputs_Internal(Input))
+	{
+		if (Chaos::FJointConstraint* JointConstraint = GetJointConstraint();
+			!JointConstraint ||
+			!JointConstraint->IsValid() ||
+			!JointConstraint->GetProxy<Chaos::FJointConstraintPhysicsProxy>()->GetHandle())
+		{
+			return false;
+		}
+
+		if (!GetMoverComponent().FindMode_Mutable<UPathedPhysicsMovementMode>(Input.SyncState.MovementMode))
+		{
+			return false;
+		}
+		
+		return true;
+	}
+	return false;
+}
+
+void UMoverPathedPhysicsLiaisonComponent::PerformProcessInputs_Internal(int32 PhysicsStep, float DeltaTime, const FPhysicsMoverAsyncInput& Input) const
+{
+	Super::PerformProcessInputs_Internal(PhysicsStep, DeltaTime, Input);
+	
+	const UPathedPhysicsMovementMode* PathedMode = GetMoverComponent().GetActiveMode<UPathedPhysicsMovementMode>();
+	if (const FPathedPhysicsMovementInputs* PathedMovementInputs = Input.InputCmd.InputCollection.FindDataByType<FPathedPhysicsMovementInputs>())
+	{
+		Chaos::FPBDRigidParticleHandle& ParticleHandle = *GetControlledParticle_Internal();
+		Chaos::FJointConstraint& JointConstraint = *GetJointConstraint();
+		
+		FPathedPhysicsMovementState& InputMoveState = Input.SyncState.SyncStateCollection.FindOrAddMutableDataByType<FPathedPhysicsMovementState>();
+		FMutablePathedMovementProperties& CurProps = InputMoveState.MutableProps;
+		const FMutablePathedMovementProperties& InputProps = PathedMovementInputs->Props;
+
+		Chaos::FPBDRigidsEvolution& Evolution = *GetPhysicsSolver()->GetEvolution();
+
+		// Toggle whether the joint is enabled and update the controlled particle to be kinematic/dynamic accordingly
+		if (InputProps.bIsJointEnabled != CurProps.bIsJointEnabled)
+		{
+			if (InputProps.bIsJointEnabled)
+			{
+				Evolution.EnableConstraints(&ParticleHandle);
+				Evolution.SetParticleObjectState(&ParticleHandle, Chaos::EObjectStateType::Dynamic);
+
+				// When the joint becomes newly enabled, teleport the endpoint to last frame's progress (since we don't touch it while it's deactivated)
+				Chaos::FRigidTransform3 TargetLastFrame; 
+				if (!GetPathedMoverComp().GetSimBlackboard()->TryGet(PathBlackboard::TargetRelativeTransform, TargetLastFrame))
+				{
+					if (PathedMode)
+					{
+						TargetLastFrame = PathedMode->CalcTargetRelativeTransform(InputMoveState.CurrentProgress);
+					}
+					else
+					{
+						// If there's no blackboard and there's no active mode, use the origin
+						TargetLastFrame = InputProps.PathOrigin;
+					}
+				}
+				
+				Chaos::FKinematicGeometryParticleHandle* EndpointParticleHandle = GetJointConstraint()->GetKinematicEndPoint()->GetHandle_LowLevel()->CastToKinematicParticle();
+				Evolution.SetParticleTransform(EndpointParticleHandle, TargetLastFrame.GetLocation(), TargetLastFrame.GetRotation(), true);
+			}
+			else
+			{
+				Evolution.DisableConstraints(&ParticleHandle);
+				Evolution.SetParticleObjectState(&ParticleHandle, Chaos::EObjectStateType::Kinematic);
+			}
+		}
+		
+
+		if (InputMoveState.MutableProps.PathOrigin.IsValid())
+		{
+			const FVector OriginLocationDelta = InputMoveState.MutableProps.PathOrigin.GetLocation() - InputProps.PathOrigin.GetLocation();
+			const FQuat OriginRotationDelta = InputMoveState.MutableProps.PathOrigin.GetRotation() - InputProps.PathOrigin.GetRotation();
+			if (!OriginLocationDelta.IsNearlyZero() || !FRotator(OriginRotationDelta).IsNearlyZero())
+			{
+				// Teleport both the controlled particle and the joint endpoint by the change in origin
+				TeleportParticleBy_Internal(ParticleHandle, OriginLocationDelta, OriginRotationDelta);
+				TeleportParticleBy_Internal(*JointConstraint.GetKinematicEndPoint()->GetHandle_LowLevel(), OriginLocationDelta, OriginRotationDelta);
+			}
+		}
+
+		if (PathedMode)
+		{
+			// Let the mode process things at the end before we copy the input props over to the sync state 
+			PathedMode->OnProcessInput_Internal(PhysicsStep, DeltaTime, Input);
+		}
+		
+		CurProps = InputProps;
+	}
+	else if (PathedMode)
+	{
+		// No high-level inputs, but always let the mode have a look too
+		PathedMode->OnProcessInput_Internal(PhysicsStep, DeltaTime, Input);
+	}
+}
+
+bool UMoverPathedPhysicsLiaisonComponent::CanSimulate_Internal(const FPhysicsMoverSimulationTickParams& TickParams, const FPhysicsMoverAsyncInput& Input) const
+{
+	const FPathedPhysicsMovementState& InputMoveState = Input.SyncState.SyncStateCollection.FindOrAddDataByType<FPathedPhysicsMovementState>();
+	if (!InputMoveState.MutableProps.IsMoving() && !InputMoveState.MutableProps.bIsJointEnabled)
+	{
+		// If we're not moving or using the joint, we're completely static and there's nothing to sim
+		return false;
+	}
+	
+	return Super::CanSimulate_Internal(TickParams, Input);
+}
+
+void UMoverPathedPhysicsLiaisonComponent::PerformPreSimulate_Internal(const FPhysicsMoverSimulationTickParams& TickParams, const FPhysicsMoverAsyncInput& Input, FPhysicsMoverAsyncOutput& Output) const
+{
+	Super::PerformPreSimulate_Internal(TickParams, Input, Output);
+	
+	const FPathedPhysicsMovementState& OutputMoveState = Output.SyncState.SyncStateCollection.FindOrAddDataByType<FPathedPhysicsMovementState>();
+
+	// If a path target was specified (i.e. if movement is happening), set it as the kinematic target
+	Chaos::FRigidTransform3 TargetRelativeTransform;
+	if (GetPathedMoverComp().GetSimBlackboard()->TryGet(PathBlackboard::TargetRelativeTransform, TargetRelativeTransform))
+	{
+		Chaos::FPBDRigidsEvolution& Evolution = *GetPhysicsSolver()->GetEvolution();
+		
+		const Chaos::FRigidTransform3 TargetWorldTransform = Chaos::FRigidTransform3::MultiplyNoScale(TargetRelativeTransform, OutputMoveState.MutableProps.PathOrigin);
+		if (OutputMoveState.MutableProps.bIsJointEnabled)
+		{
+			Chaos::FKinematicGeometryParticleHandle* EndpointParticleHandle = GetJointConstraint()->GetKinematicEndPoint()->GetHandle_LowLevel()->CastToKinematicParticle();
+			Evolution.SetParticleKinematicTarget(EndpointParticleHandle, Chaos::FKinematicTarget::MakePositionTarget(TargetWorldTransform));
+		}
+		else
+		{
+			// When we don't want to use the joint, just move the pathed component kinematically
+			Evolution.SetParticleKinematicTarget(GetControlledParticle_Internal(), Chaos::FKinematicTarget::MakePositionTarget(TargetWorldTransform));
+		}
+	}
+}
+
+Chaos::FJointConstraint* UMoverPathedPhysicsLiaisonComponent::GetJointConstraint() const
+{
+	return static_cast<Chaos::FJointConstraint*>(ConstraintHandle.Constraint);
+}
+
+void UMoverPathedPhysicsLiaisonComponent::CreateTargetJoint()
+{
+	const UMoverComponent& MoverComp = GetMoverComponent();
+	if (IPhysicsComponent* PhysicsComponent = Cast<IPhysicsComponent>(MoverComp.GetUpdatedComponent()))
+	{
+		if (Chaos::FPhysicsObject* PhysicsObject = PhysicsComponent->GetPhysicsObjectByName(NAME_None))
+		{
+			const FTransform ComponentWorldTransform = MoverComp.GetUpdatedComponent()->GetComponentTransform();
+
+			// Create the constraint via FChaosEngineInterface directly because it allows jointing a "real" object with a point in space (it creates a dummy particle for us)
+			FPhysicsConstraintHandle Handle = FChaosEngineInterface::CreateConstraint(PhysicsObject, nullptr, FTransform::Identity, FTransform::Identity);
+
+			bool bIsConstraintValid = false;
+			if (Handle.IsValid() && ensure(Handle->IsType(Chaos::EConstraintType::JointConstraintType)))
+			{
+				if (Chaos::FJointConstraint* Constraint = static_cast<Chaos::FJointConstraint*>(Handle.Constraint))
+				{
+					// Since we didn't use the ConstraintInstance to actually create the constraint (it requires both bodies exist, see comment above), link everything up manually
+					ConstraintHandle = Handle;
+					ConstraintInstance.ConstraintHandle = ConstraintHandle;
+					Constraint->SetUserData(&PhysicsUserData/*has a (void*)FConstraintInstanceBase*/);
+					bIsConstraintValid = true;
+
+					if (Chaos::FPBDRigidParticle* EndpointParticle = Constraint->GetPhysicsBodies()[1]->GetParticle<Chaos::EThreadContext::External>()->CastToRigidParticle())
+					{
+						EndpointParticle->SetX(ComponentWorldTransform.GetLocation());
+						EndpointParticle->SetR(ComponentWorldTransform.GetRotation());
+					}
+				}
+			}
+
+			if (!bIsConstraintValid)
+			{
+				FChaosEngineInterface::ReleaseConstraint(Handle);
+			}
+		}
+	}
+}
+
+void UMoverPathedPhysicsLiaisonComponent::DestroyTargetJoint()
+{
+	FChaosEngineInterface::ReleaseConstraint(ConstraintHandle);
+	ConstraintInstance.ConstraintHandle.Reset();
+}
+
+void UMoverPathedPhysicsLiaisonComponent::HandleMovementModeChanged(const FName& OldModeName, const FName& NewModeName)
+{
+	if (const UPathedPhysicsMovementMode* OldMode = GetMoverComponent().FindMode_Mutable<UPathedPhysicsMovementMode>(OldModeName))
+	{
+		OldMode->OnIsUsingJointChanged().RemoveAll(this);
+	}
+	
+	if (const UPathedPhysicsMovementMode* NewMode = GetMoverComponent().FindMode_Mutable<UPathedPhysicsMovementMode>(NewModeName))
+	{
+		ApplyPathModeConfig(*NewMode);
+	}
+}
+
+void UMoverPathedPhysicsLiaisonComponent::ApplyPathModeConfig(const UPathedPhysicsMovementMode& PathedMode)
+{
+	PathedMode.OnIsUsingJointChanged().AddUObject(this, &ThisClass::HandleIsUsingJointChanged);
+	HandleIsUsingJointChanged(PathedMode.IsUsingJoint());
+
+	const EPathedPhysicsPlaybackBehavior DefaultPlaybackBehavior = GetPathedMoverComp().GetDefaultPlaybackBehavior();
+	SetPlaybackBehavior(PathedMode.GetPlaybackBehaviorOverride().Get(DefaultPlaybackBehavior));
+	
+	if (HasValidPhysicsState())
+	{
+		ConstraintInstance.CopyProfilePropertiesFrom(PathedMode.GetConstraintProperties());
+	}
+}
+
+void UMoverPathedPhysicsLiaisonComponent::HandleIsUsingJointChanged(bool bIsUsingJoint)
+{
+	// When using the joint, we need to replicate the actual transform of the controlled component's particle to determine if a resim is needed
+	//	This is easier, cheaper, and accounts for relevancy if we do so via the actor's built-in bReplicateMovement functionality
+	// When not using the joint, however, the component's kinematic target is set each frame based entirely on FPathedPhysicsMovementState::CurrentProgress
+	//	Therefore, we can simply compare the calculated progress between server and client to determine if a resim is needed, and don't need to send any transform data
+	Inputs_External.bIsJointEnabled = bIsUsingJoint;
+	if (AActor* MyActor = GetOwner())
+	{
+		MyActor->SetReplicatingMovement(bIsUsingJoint);
+	}
+	
+	if (NetworkPhysicsComponent && !NetworkPhysicsComponent->HasServerWorld())
+	{
+		NetworkPhysicsComponent->SetCompareStateToTriggerRewind(!bIsUsingJoint);
+	}
+}
+
+#undef LOCTEXT_NAMESPACE
diff --git a/Engine/Plugins/Experimental/Mover/Source/Mover/Private/PhysicsMover/PathedMovement/ArcRotationPathPattern.cpp b/Engine/Plugins/Experimental/Mover/Source/Mover/Private/PhysicsMover/PathedMovement/ArcRotationPathPattern.cpp
new file mode 100644
--- /dev/null
+++ b/Engine/Plugins/Experimental/Mover/Source/Mover/Private/PhysicsMover/PathedMovement/ArcRotationPathPattern.cpp
@@ -0,0 +1,60 @@
+// Copyright Epic Games, Inc. All Rights Reserved.
+
+#include "PhysicsMover/PathedMovement/ArcRotationPathPattern.h"
+
+#include "DebugRenderSceneProxy.h"
+#include "PhysicsMover/PathedMovement/PathedPhysicsDebugDrawComponent.h"
+#include "PhysicsMover/PathedMovement/PathedPhysicsMoverComponent.h"
+
+FVector CalcArcPoint(const FTransform& PathOriginTransform, const FVector& RotationAxis, float ArcRadius, float ArcAngleDeg)
+{
+	FTransform ArcPointTransform(FQuat(RotationAxis, FMath::DegreesToRadians(ArcAngleDeg)));
+	ArcPointTransform.Accumulate(PathOriginTransform);
+	return ArcPointTransform.TransformPositionNoScale(FVector(ArcRadius, 0.f, 0.f));
+}
+
+void UArcRotationPattern::AppendDebugDrawElements(UPathedPhysicsDebugDrawComponent& DebugDrawComp, FBoxSphereBounds::Builder& InOutDebugBoundsBuilder)
+{
+	constexpr float BoundsDrawRadius = 500.f;
+	const FTransform& PathOriginTransform = GetPathedMoverComp().GetPathOriginTransform();
+
+	// Draw dashed lines for the min and max bounds of the arc
+	const FVector ArcMinEndpoint = CalcArcPoint(PathOriginTransform, RotationAxis, BoundsDrawRadius, 0.f);
+	InOutDebugBoundsBuilder += ArcMinEndpoint;
+
+	const FVector ArcMaxEndpoint = CalcArcPoint(PathOriginTransform, RotationAxis, BoundsDrawRadius, RotationArcAngle);
+	InOutDebugBoundsBuilder += ArcMinEndpoint;
+	
+	const FVector OriginLoc = PathOriginTransform.GetLocation();
+	DebugDrawComp.DebugDashedLines.Emplace(OriginLoc, ArcMinEndpoint, PatternDebugDrawColor, 5.f);
+	DebugDrawComp.DebugDashedLines.Emplace(OriginLoc, ArcMaxEndpoint, PatternDebugDrawColor, 5.f);
+
+	// Draw a curve indicating the rotation between the bounds
+	constexpr float RotationArrowDrawRadius = BoundsDrawRadius * 0.75f;
+	
+	FVector PrevArrowPoint = CalcArcPoint(PathOriginTransform, RotationAxis, RotationArrowDrawRadius, 0.f);
+	const int32 NumCurvePoints = FMath::CeilToInt32(FMath::Lerp(3.f, 20.f, RotationArcAngle / 360.f));
+	for (int32 PointIdx = 1; PointIdx <= NumCurvePoints; ++PointIdx)
+	{
+		const float Angle = FMath::Lerp(0.f, RotationArcAngle, float(PointIdx) / NumCurvePoints);
+		const FVector ArrowPoint = CalcArcPoint(PathOriginTransform, RotationAxis, RotationArrowDrawRadius, Angle);
+
+		DebugDrawComp.DebugLines.Emplace(PrevArrowPoint, ArrowPoint, PatternDebugDrawColor, 2.f);
+		if (PointIdx == 1 && PerLoopBehavior == EPathedPhysicsPlaybackBehavior::ThereAndBack)
+		{
+			DebugDrawComp.DebugArrowLines.Emplace(ArrowPoint, PrevArrowPoint, PatternDebugDrawColor, 5.f);
+		}
+		else if (PointIdx == NumCurvePoints)
+		{
+			DebugDrawComp.DebugArrowLines.Emplace(PrevArrowPoint, ArrowPoint, PatternDebugDrawColor, 5.f);
+		}
+
+		PrevArrowPoint = ArrowPoint;
+	}
+}
+
+FTransform UArcRotationPattern::CalcUnmaskedTargetRelativeTransform(float PatternProgress, const FTransform& CurTargetTransform) const
+{
+	const float TargetAngle = RotationArcAngle * PatternProgress;
+	return FTransform(FQuat(RotationAxis, FMath::DegreesToRadians(TargetAngle)));
+}
diff --git a/Engine/Plugins/Experimental/Mover/Source/Mover/Private/PhysicsMover/PathedMovement/EllipticalMovementPathPattern.cpp b/Engine/Plugins/Experimental/Mover/Source/Mover/Private/PhysicsMover/PathedMovement/EllipticalMovementPathPattern.cpp
new file mode 100644
--- /dev/null
+++ b/Engine/Plugins/Experimental/Mover/Source/Mover/Private/PhysicsMover/PathedMovement/EllipticalMovementPathPattern.cpp
@@ -0,0 +1,30 @@
+// Copyright Epic Games, Inc. All Rights Reserved.
+
+#include "PhysicsMover/PathedMovement/EllipticalMovementPathPattern.h"
+
+#include "HAL/IConsoleManager.h"
+#include "PhysicsMover/PathedMovement/PathedMovementMode.h"
+
+#include UE_INLINE_GENERATED_CPP_BY_NAME(EllipticalMovementPathPattern)
+
+void UEllipticalMovementPathPattern::InitializePattern()
+{
+	Super::InitializePattern();
+}
+
+FTransform UEllipticalMovementPathPattern::CalcUnmaskedTargetRelativeTransform(float PatternProgress, const FTransform& CurTargetTransform) const
+{
+	// Calc the ellipse point
+	const float Angle = PatternProgress * FMath::DegreesToRadians(UsableArcAngle);
+	FVector2f TargetLocationXY(RadiusX * FMath::Cos(Angle), RadiusY * FMath::Sin(Angle));
+
+	// Scooch the center over so 0 progress means no movement
+	TargetLocationXY.X -= RadiusX;
+
+	// Spin about the path origin as desired
+	TargetLocationXY = TargetLocationXY.GetRotated(OriginAngle);
+
+	// Rotate the ellipse from the XY plane to whatever 3D plane is desired
+	const FVector TargetLocation = EllipsePlaneRotation.RotateVector(FVector(TargetLocationXY.X, TargetLocationXY.Y, 0.f));
+	return FTransform(TargetLocation);
+}
diff --git a/Engine/Plugins/Experimental/Mover/Source/Mover/Private/PhysicsMover/PathedMovement/LookAtRotationPathPattern.cpp b/Engine/Plugins/Experimental/Mover/Source/Mover/Private/PhysicsMover/PathedMovement/LookAtRotationPathPattern.cpp
new file mode 100644
--- /dev/null
+++ b/Engine/Plugins/Experimental/Mover/Source/Mover/Private/PhysicsMover/PathedMovement/LookAtRotationPathPattern.cpp
@@ -0,0 +1,37 @@
+// Copyright Epic Games, Inc. All Rights Reserved.
+
+#include "PhysicsMover/PathedMovement/LookAtRotationPathPattern.h"
+
+#include "DebugRenderSceneProxy.h"
+#include "Kismet/KismetMathLibrary.h"
+#include "PhysicsMover/PathedMovement/PathedMovementMode.h"
+#include "PhysicsMover/PathedMovement/PathedPhysicsDebugDrawComponent.h"
+#include "PhysicsMover/PathedMovement/PathedPhysicsMoverComponent.h"
+
+
+void ULookAtRotationPattern::AppendDebugDrawElements(UPathedPhysicsDebugDrawComponent& DebugDrawComp, FBoxSphereBounds::Builder& InOutDebugBoundsBuilder)
+{
+	const FVector LookAtLocation = GetPathedMoverComp().GetPathOriginTransform().TransformPositionNoScale(RelativeLookAtLocation);
+	DebugDrawComp.DebugStars.Emplace(LookAtLocation, PatternDebugDrawColor, 20.f);
+	InOutDebugBoundsBuilder += LookAtLocation;
+}
+
+FTransform ULookAtRotationPattern::CalcUnmaskedTargetRelativeTransform(float PatternProgress, const FTransform& CurTargetTransform) const
+{
+	//@todo DanH: So this is a modifier that has nothing to do with our progress... is that ok as a pattern or should it be something/somewhere else?
+	//		Maybe it needs to be on the mode itself? Depending on how far the component is lagging behind the target, the actual rotation is gonna be different, and we might want to force it?
+	
+	const FRotator LookAtRotation = UKismetMathLibrary::FindLookAtRotation(CurTargetTransform.GetLocation(), RelativeLookAtLocation);
+	return FTransform(LookAtRotation);
+}
+
+void ULookAtRotationPattern::SetRelativeLookAtLocation(const FVector& RelativeLookAt)
+{
+	RelativeLookAtLocation = RelativeLookAt;
+}
+
+void ULookAtRotationPattern::SetLookAtLocation(const FVector& WorldLookAt)
+{
+	const FVector RelativeLookAt = GetPathedMoverComp().GetPathOriginTransform().InverseTransformPosition(WorldLookAt);
+	SetRelativeLookAtLocation(RelativeLookAt);
+}
\ No newline at end of file
diff --git a/Engine/Plugins/Experimental/Mover/Source/Mover/Private/PhysicsMover/PathedMovement/PathedMovementMode.cpp b/Engine/Plugins/Experimental/Mover/Source/Mover/Private/PhysicsMover/PathedMovement/PathedMovementMode.cpp
new file mode 100644
--- /dev/null
+++ b/Engine/Plugins/Experimental/Mover/Source/Mover/Private/PhysicsMover/PathedMovement/PathedMovementMode.cpp
@@ -0,0 +1,347 @@
+// Copyright Epic Games, Inc. All Rights Reserved.
+
+#include "PhysicsMover/PathedMovement/PathedMovementMode.h"
+#include "PhysicsMover/PathedMovement/PathedMovementPatternBase.h"
+#include "MoverComponent.h"
+#include "Backends/MoverPathedPhysicsLiaison.h"
+#include "Framework/Threading.h"
+#include "Kismet/KismetMathLibrary.h"
+#include "PhysicsMover/PathedMovement/PathedPhysicsMoverComponent.h"
+
+#include UE_INLINE_GENERATED_CPP_BY_NAME(PathedMovementMode)
+
+FAutoConsoleVariable CVarEnablePathedPhysicsMovementDebug(
+	TEXT("Mover.PathedPhysicsMovement.EnableDebug"),
+	false,
+	TEXT("True to enable a log firehose of position update debug info")
+);
+
+UPathedPhysicsMovementMode::UPathedPhysicsMovementMode()
+{
+	//@todo DanH: These settings result in a joint that isn't nearly rigid/string/stiff enough to be the default.
+	//	Need to better understand the knobs as cranking up values ad hoc doesn't seem to have much effect.
+	
+	JointConstraintProperties.bLinearBreakable = false;
+	JointConstraintProperties.bAngularBreakable = false;
+	JointConstraintProperties.bDisableCollision = true;
+
+	// Free the linear joints and set the max allowed distance from the target
+	static constexpr float DefaultMaxLinearDistance = 100.f;
+	JointConstraintProperties.LinearLimit.XMotion = LCM_Free;
+	JointConstraintProperties.LinearLimit.YMotion = LCM_Free;
+	JointConstraintProperties.LinearLimit.ZMotion = LCM_Free;
+	JointConstraintProperties.LinearLimit.Limit = DefaultMaxLinearDistance;
+
+	// Linear drive config
+	static constexpr float LinearDriveStiffness = 750.f;
+	static const float LinearDriveCriticalDamping = 2 * FMath::Sqrt(LinearDriveStiffness);
+
+	FConstraintDrive DefaultLinearDriveConstraint;
+	DefaultLinearDriveConstraint.bEnablePositionDrive = true;
+	DefaultLinearDriveConstraint.bEnableVelocityDrive = true;
+	DefaultLinearDriveConstraint.Stiffness = LinearDriveStiffness;
+	DefaultLinearDriveConstraint.Damping = LinearDriveCriticalDamping;
+	DefaultLinearDriveConstraint.MaxForce = 100.f;
+
+	JointConstraintProperties.LinearDrive.XDrive = DefaultLinearDriveConstraint;
+	JointConstraintProperties.LinearDrive.YDrive = DefaultLinearDriveConstraint;
+	JointConstraintProperties.LinearDrive.ZDrive = DefaultLinearDriveConstraint;
+
+	// Free the angular joints and set degree limits relative to the target
+	static constexpr float DefaultAngularLimitDegrees = 15.f;
+	JointConstraintProperties.ConeLimit.Swing1Motion = ACM_Free;
+	JointConstraintProperties.ConeLimit.Swing1LimitDegrees = DefaultAngularLimitDegrees;
+	JointConstraintProperties.ConeLimit.Swing2Motion = ACM_Free;
+	JointConstraintProperties.ConeLimit.Swing2LimitDegrees = DefaultAngularLimitDegrees;
+	JointConstraintProperties.TwistLimit.TwistMotion = ACM_Free;
+	JointConstraintProperties.TwistLimit.TwistLimitDegrees = DefaultAngularLimitDegrees;
+	
+	// Angular drive config
+	static constexpr float AngularDriveStiffness = 1500.f;
+	static const float AngularDriveCriticalDamping = 2 * FMath::Sqrt(AngularDriveStiffness);
+
+	JointConstraintProperties.AngularDrive.AngularDriveMode = EAngularDriveMode::SLERP;
+	JointConstraintProperties.AngularDrive.SlerpDrive.bEnablePositionDrive = true;
+	JointConstraintProperties.AngularDrive.SlerpDrive.bEnableVelocityDrive = true;
+	JointConstraintProperties.AngularDrive.SlerpDrive.Stiffness = AngularDriveStiffness;
+	JointConstraintProperties.AngularDrive.SlerpDrive.Damping = AngularDriveCriticalDamping;
+	JointConstraintProperties.AngularDrive.SlerpDrive.MaxForce = 100.f;
+}
+
+#if WITH_EDITOR
+void UPathedPhysicsMovementMode::PostEditChangeChainProperty(FPropertyChangedChainEvent& PropertyChangedChainEvent)
+{
+	FPatternPostChangeMovementModeHelper::ProcessPostEditChangeChain(*this, PropertyChangedChainEvent);
+}
+#endif
+
+void UPathedPhysicsMovementMode::OnSimulationTick(const FSimulationTickParams& Params, FMoverTickEndData& OutputState)
+{
+	Chaos::EnsureIsInPhysicsThreadContext();
+	
+	OnPreSimulate_Internal(Params, OutputState);
+}
+
+void UPathedPhysicsMovementMode::OnProcessInput_Internal(int32 PhysicsStep, float DeltaTime, const FPhysicsMoverAsyncInput& Input) const
+{
+	if (const FPathedPhysicsMovementInputs* Inputs = Input.InputCmd.InputCollection.FindDataByType<FPathedPhysicsMovementInputs>())
+	{
+		CachedInputs = Inputs->Props;
+
+		const FPathedPhysicsMovementState& StartingState = Input.SyncState.SyncStateCollection.FindOrAddMutableDataByType<FPathedPhysicsMovementState>();
+		const FMutablePathedMovementProperties& CurProps = StartingState.MutableProps;
+		
+		if (CachedInputs.IsMoving() != CurProps.IsMoving())
+		{
+			FPathedPhysicsMovementState& StartMoveState = Input.SyncState.SyncStateCollection.FindOrAddMutableDataByType<FPathedPhysicsMovementState>();
+			const float PlaybackDuration = GetPlaybackDuration();
+			if (StartMoveState.MutableProps.bIsInReverse && StartMoveState.LastStopPlaybackTime == 0.f)
+			{
+				// When starting play in reverse, jump to the end first (otherwise it'll be immediately done)
+				StartMoveState.LastStopPlaybackTime = PlaybackDuration;
+			}
+			else if (!StartMoveState.MutableProps.bIsInReverse && StartMoveState.LastStopPlaybackTime == PlaybackDuration)
+			{
+				// Similarly, when we're starting at the end and want to play forward, jump to the start
+				StartMoveState.LastStopPlaybackTime = 0.f;
+			}
+		}
+	}
+}
+
+void UPathedPhysicsMovementMode::OnPreSimulate_Internal(const FSimulationTickParams& Params, FMoverTickEndData& OutputState) const
+{
+	// The output state is initialized from the input state 
+	FPathedPhysicsMovementState& OutputMoveState = OutputState.SyncState.SyncStateCollection.FindOrAddMutableDataByType<FPathedPhysicsMovementState>();
+
+	if (!OutputMoveState.MutableProps.IsMoving() || OutputMoveState.MutableProps.MovementStartFrame > Params.TimeStep.ServerFrame)
+	{
+		// Either we're not moving or we shouldn't start until a future frame
+		return;
+	}
+	
+	// How long (in s) it's been since we started moving (negative when going in reverse)
+	const float TimeDirection = OutputMoveState.MutableProps.bIsInReverse ? -1.f : 1.f;
+	const float TimeSinceStart = (Params.TimeStep.ServerFrame - OutputMoveState.MutableProps.MovementStartFrame + 1) * Params.TimeStep.StepMs * TimeDirection * 0.001f;
+
+	// Offset the time since starting by wherever we started playback in the first place
+	//	ex: If we last stopped at 3s and it's been 1s, the actual playback time should be 4s
+	//	ex: In reverse, it's been -1s, giving the expected playback time of 2s
+	const float TimeSinceLastStop = TimeSinceStart + OutputMoveState.LastStopPlaybackTime;
+	
+	bool bReachesEndOfPlayback = false;
+	const float PlaybackTime = GetBoundedPlaybackTime(TimeSinceLastStop, bReachesEndOfPlayback);
+	
+	if (bReachesEndOfPlayback)
+	{
+		OutputMoveState.MutableProps.MovementStartFrame = INDEX_NONE;
+		OutputMoveState.LastStopPlaybackTime = PlaybackTime;
+	}
+	
+	const float BoundedTimeAlongPath = GetBoundedTimeAlongPath_Unsafe(PlaybackTime);
+	
+	const float Alpha = FMath::Clamp(UKismetMathLibrary::SafeDivide(BoundedTimeAlongPath , GetPathDuration()), 0.f, 1.f);
+	const float ProgressAmt = FAlphaBlend::AlphaToBlendOption(Alpha, Easing, CustomEasingCurve);
+
+	const Chaos::FRigidTransform3 TargetRelativeTransform = CalcTargetRelativeTransform(ProgressAmt);
+	GetPathedMoverComp().GetSimBlackboard_Mutable()->Set(PathBlackboard::TargetRelativeTransform, TargetRelativeTransform);
+	
+	if (CVarEnablePathedPhysicsMovementDebug->GetBool())
+	{
+		const Chaos::FRigidTransform3 TargetWorldTransform = Chaos::FRigidTransform3::MultiplyNoScale(TargetRelativeTransform, OutputMoveState.MutableProps.PathOrigin);
+		UE_LOG(LogMover, VeryVerbose, TEXT("Pathed mover [%s (%s)]: Setting target:\n\tRelative [%s]\n\tWorld [%s]"),
+			*GetPathedMoverComp().GetOwner()->GetName(),
+			GetPathedMoverComp().GetOwner()->GetNetMode() == NM_DedicatedServer ? TEXT("Server") : TEXT("Client"),
+			*TargetRelativeTransform.ToString(),
+			*TargetWorldTransform.ToString());
+	}
+}
+
+void UPathedPhysicsMovementMode::InitializePath()
+{
+	for (UPathedMovementPatternBase* Pattern : PathPatterns)
+	{
+		if (Pattern)
+		{
+			Pattern->InitializePattern();
+		}
+	}
+}
+
+UPathedMovementPatternBase* UPathedPhysicsMovementMode::BP_FindPattern(TSubclassOf<UPathedMovementPatternBase> PatternType) const
+{
+	for (UPathedMovementPatternBase* Pattern : PathPatterns)
+	{
+		if (Pattern && Pattern->IsA(PatternType))
+		{
+			return Pattern;
+		}
+	}
+	
+	return nullptr;
+}
+
+void UPathedPhysicsMovementMode::SetPathDuration_BeginPlayOnly(float NewDuration)
+{
+	if (ensure(!GetPathedMoverComp().HasBegunPlay()))
+	{
+		OneWayTripDuration = FMath::Max(NewDuration, 0.f);
+	}
+}
+
+void UPathedPhysicsMovementMode::SetUseJointConstraint(bool bUseJoint)
+{
+	if (bUseJointConstraint != bUseJoint)
+	{
+		bUseJointConstraint = bUseJoint;
+		OnIsUsingJointChanged().Broadcast(bUseJoint);
+	}
+}
+
+UPathedPhysicsMoverComponent& UPathedPhysicsMovementMode::GetPathedMoverComp() const
+{
+	return *GetOuterUPathedPhysicsMoverComponent();
+}
+
+Chaos::FRigidTransform3 UPathedPhysicsMovementMode::CalcTargetRelativeTransform(float ProgressAmt) const
+{
+	FTransform TargetRelativeTransform = FTransform::Identity;
+	
+	for (UPathedMovementPatternBase* PathPattern : PathPatterns)
+	{
+		if (PathPattern)
+		{
+			const FTransform PatternTransform = PathPattern->CalcTargetRelativeTransform(ProgressAmt, TargetRelativeTransform);
+        	TargetRelativeTransform.Accumulate(PatternTransform);
+		}
+	}
+
+	return TargetRelativeTransform;
+}
+
+float UPathedPhysicsMovementMode::GetPlaybackDuration() const
+{
+	// By doubling the playback time for ping-pongs, we can treat them the same as normal monodirectional playback when bounding
+	return CachedInputs.IsPingPonging() ? GetPathDuration() * 2 : GetPathDuration();
+}
+
+float UPathedPhysicsMovementMode::GetBoundedPlaybackTime(float PlaybackTime, bool& bOutReachesEndOfPlayback) const
+{
+	bOutReachesEndOfPlayback = false;
+	const float PlaybackDuration = GetPlaybackDuration();;	
+	
+	// FMod gets mad if you send it a tiny mod factor, so just treat 0 duration as infinite
+	if (ensure(!FMath::IsNearlyZero(PlaybackDuration)))
+	{
+		const bool bIsLooping = CachedInputs.IsLooping();
+		if (CachedInputs.bIsInReverse)
+		{
+			if (bIsLooping)
+			{
+				if (PlaybackTime < 0.f)
+				{
+					// Roll a negative time back around to starting at max (i.e. in a 4s path, -1s input should become 3s)
+					PlaybackTime = PlaybackDuration + FMath::Fmod(PlaybackTime, PlaybackDuration);
+				}
+				while (PlaybackTime < 0.f)
+				{
+					PlaybackTime += PlaybackDuration;
+				}
+			}
+			else if (PlaybackTime <= 0.f)
+			{
+				// This is a one-shot that has passed 0, so it's all done
+				PlaybackTime = 0.f;
+				bOutReachesEndOfPlayback = true;
+			}
+		}
+		else if (PlaybackTime > PlaybackDuration)
+		{
+			if (bIsLooping)
+			{
+				// Loop the completed run back to the beginning
+				PlaybackTime = FMath::Fmod(PlaybackTime, PlaybackDuration);
+			}
+			else
+			{
+				// Reached the end, all done
+				PlaybackTime = PlaybackDuration;
+				bOutReachesEndOfPlayback = true;
+			}
+		}
+	}
+
+	return PlaybackTime;
+}
+
+float UPathedPhysicsMovementMode::GetBoundedTimeAlongPath(float Time) const
+{
+	const float PathDuration = GetPathDuration();
+	const EPathedPhysicsPlaybackBehavior PlaybackBehavior = CachedInputs.PlaybackBehavior;
+	const bool bIsPingPonging = PlaybackBehavior == EPathedPhysicsPlaybackBehavior::ThereAndBack || PlaybackBehavior == EPathedPhysicsPlaybackBehavior::PingPong;
+
+	if (Time < 0.f ||
+		(bIsPingPonging && Time > 2 * PathDuration) ||
+		(!bIsPingPonging && Time > PathDuration))
+	{
+		bool bThrowaway = false;
+		Time = GetBoundedPlaybackTime(Time, bThrowaway);
+	}
+
+	return GetBoundedTimeAlongPath_Unsafe(Time);
+}
+
+float UPathedPhysicsMovementMode::GetBoundedTimeAlongPath_Unsafe(float BoundedPlaybackTime) const
+{
+	// TimeAlongPath only has the potential to differ from PlaybackTime when doing ping-pong-style movement
+	if (CachedInputs.IsPingPonging())
+	{
+		const float PathDuration = GetPathDuration();
+		if (BoundedPlaybackTime > PathDuration)
+		{
+			// Ex: We're 8s into a ping-pong of a 5s path. That means we did 5s there and are 3s into the trip back, with 2s to go ==> 10 - 8 = 2
+			const float TimeAlongPath = PathDuration * 2 - BoundedPlaybackTime;
+			return ensure(TimeAlongPath >= 0.f) ? TimeAlongPath : 0.f;
+		}
+	}
+	
+	return BoundedPlaybackTime;
+}
+
+#if WITH_EDITOR
+void FPatternPostChangeMovementModeHelper::ProcessPostEditChangeChain(UPathedPhysicsMovementMode& Mode, FPropertyChangedChainEvent& PropertyChangedChainEvent)
+{
+	if (Mode.PathPatterns.Num() > 1 && PropertyChangedChainEvent.GetPropertyName().IsEqual(GET_MEMBER_NAME_CHECKED(UPathedPhysicsMovementMode, PathPatterns)))
+	{
+		const FProperty* ActuallyChangedProperty = PropertyChangedChainEvent.PropertyChain.GetTail()->GetValue();
+		if (PropertyChangedChainEvent.ChangeType == EPropertyChangeType::ArrayRemove || 
+			ActuallyChangedProperty->GetFName().IsEqual(GET_MEMBER_NAME_CHECKED(UPathedMovementPatternBase, bStartAfterPreviousPattern)) ||
+			ActuallyChangedProperty->GetFName().IsEqual(GET_MEMBER_NAME_CHECKED(UPathedMovementPatternBase, StartAtPathProgress)) ||
+			ActuallyChangedProperty->GetFName().IsEqual(GET_MEMBER_NAME_CHECKED(UPathedMovementPatternBase, EndAtPathProgress)))
+		{
+			// Don't bother figuring out exactly who or what changed at this point - just run through and make sure all the Start/End info is sound
+			float PrevPatternEndProgress = 0.f;
+			for (UPathedMovementPatternBase* Pattern : Mode.PathPatterns)
+			{
+				if (Pattern)
+				{
+					if (Pattern->bStartAfterPreviousPattern)
+					{
+						// It's possible this pattern wasn't the one that changed, so just make sure it's part of the transaction
+						Pattern->Modify();
+						Pattern->StartAtPathProgress = PrevPatternEndProgress;
+					}
+					
+					if (Pattern->StartAtPathProgress > Pattern->EndAtPathProgress)
+					{
+						Pattern->Modify();
+						Pattern->EndAtPathProgress = Pattern->StartAtPathProgress;
+					}
+					
+					PrevPatternEndProgress = Pattern->EndAtPathProgress;
+				}
+			}
+		}
+	}
+}
+#endif
\ No newline at end of file
diff --git a/Engine/Plugins/Experimental/Mover/Source/Mover/Private/PhysicsMover/PathedMovement/PathedMovementPatternBase.cpp b/Engine/Plugins/Experimental/Mover/Source/Mover/Private/PhysicsMover/PathedMovement/PathedMovementPatternBase.cpp
new file mode 100644
--- /dev/null
+++ b/Engine/Plugins/Experimental/Mover/Source/Mover/Private/PhysicsMover/PathedMovement/PathedMovementPatternBase.cpp
@@ -0,0 +1,115 @@
+// Copyright Epic Games, Inc. All Rights Reserved.
+
+#include "PhysicsMover/PathedMovement/PathedMovementPatternBase.h"
+
+#include "HAL/IConsoleManager.h"
+#include "PhysicsMover/PathedMovement/PathedMovementMode.h"
+
+#include UE_INLINE_GENERATED_CPP_BY_NAME(PathedMovementPatternBase)
+
+void UPathedMovementPatternBase::InitializePattern()
+{
+}
+
+void UPathedMovementPatternBase::ProduceInputs_External(OUT FPhysicsMoverAsyncInput& Input)
+{
+	//@todo DanH: Push something with the total duration and playback behavior
+}
+
+template <typename T>
+T ApplyAxisMask(const T& Unmasked, EPatternAxisMaskFlags Flags, float MaskedValue = 0.f)
+{
+	T MaskedResult = Unmasked;
+	if (EnumHasAnyFlags(Flags, EPatternAxisMaskFlags::X))
+	{
+		MaskedResult.X = MaskedValue;
+	}
+	if (EnumHasAnyFlags(Flags, EPatternAxisMaskFlags::Y))
+	{
+		MaskedResult.Y = MaskedValue;
+	}
+	if (EnumHasAnyFlags(Flags, EPatternAxisMaskFlags::Z))
+	{
+		MaskedResult.Z = MaskedValue;
+	}
+	return MaskedResult;
+}
+
+template <>
+FRotator ApplyAxisMask(const FRotator& Unmasked, EPatternAxisMaskFlags Flags, float MaskedValue)
+{
+	FRotator MaskedResult = Unmasked;
+	if (EnumHasAnyFlags(Flags, EPatternAxisMaskFlags::X))
+	{
+		MaskedResult.Roll = MaskedValue;
+	}
+	if (EnumHasAnyFlags(Flags, EPatternAxisMaskFlags::Y))
+	{
+		MaskedResult.Pitch = MaskedValue;
+	}
+	if (EnumHasAnyFlags(Flags, EPatternAxisMaskFlags::Z))
+	{
+		MaskedResult.Yaw = MaskedValue;
+	}
+	return MaskedResult;
+}
+
+FTransform UPathedMovementPatternBase::CalcTargetRelativeTransform(float OverallPathProgress, const FTransform& CurTargetTransform) const
+{
+	//@todo DanH: When doing debug drawing, when a pattern is shorter and/or loops more, it loses resolution. Can't just always calculate as though OverallPathProgress == PatternProgress
+	//		though because we still depend on the calculated result. Maybe I shouldn't bother with that and just let the aggregate one calc on its own?
+	
+	if (NumLoopsPerPath <= 0 || StartAtPathProgress >= EndAtPathProgress || OverallPathProgress <= StartAtPathProgress)
+	{
+		return FTransform();
+	}
+	
+	// How far into the current loop of this specific pattern are we?
+	const float PathProgressSinceStart = FMath::Min(EndAtPathProgress, OverallPathProgress) - StartAtPathProgress;
+	const float PathProgressPerPatternLoop = (EndAtPathProgress - StartAtPathProgress) / NumLoopsPerPath;
+	float CurLoopPathProgress = FMath::Fmod(PathProgressSinceStart, PathProgressPerPatternLoop);
+	if (CurLoopPathProgress == 0.f && PathProgressSinceStart > 0.f)
+	{
+		// Treat progress that matches the per-loop span exactly as 100%, not 0%
+		CurLoopPathProgress = PathProgressPerPatternLoop;
+	}
+
+	// If each loop is there and back, the actual progress needs to flip after completing half the span
+	if (PerLoopBehavior == EPathedPhysicsPlaybackBehavior::ThereAndBack)
+	{
+		// ThereAndBack means we actually progress twice as fast as a OneShot
+		CurLoopPathProgress *= 2.f;
+		
+		const float ReversePathProgress = CurLoopPathProgress - PathProgressPerPatternLoop;
+		if (ReversePathProgress > 0.f)
+		{
+			CurLoopPathProgress = PathProgressPerPatternLoop - ReversePathProgress;
+		}
+	}
+
+	// Now convert from overall path progress to progress for this pattern
+	const float PatternProgress = CurLoopPathProgress / PathProgressPerPatternLoop;
+	const FTransform UnmaskedTarget = CalcUnmaskedTargetRelativeTransform(PatternProgress, CurTargetTransform);
+
+	if (bOrientComponentToPath)
+	{
+		//@todo DanH: Rotate so that the forward vector matches the translation vector direction
+	}
+
+	const FVector TargetLocation = ApplyAxisMask(UnmaskedTarget.GetLocation(), (EPatternAxisMaskFlags)TranslationMasks);
+	const FRotator TargetRotation = ApplyAxisMask(UnmaskedTarget.GetRotation().Rotator(), (EPatternAxisMaskFlags)RotationMasks);
+	const FVector TargetScale = ApplyAxisMask(UnmaskedTarget.GetScale3D(), (EPatternAxisMaskFlags)ScaleMasks, 1.f);
+	return FTransform(TargetRotation, TargetLocation, TargetScale);
+}
+
+UPathedPhysicsMovementMode& UPathedMovementPatternBase::GetMovementMode() const
+{
+	UPathedPhysicsMovementMode* OuterMode = GetOuterUPathedPhysicsMovementMode();
+	check(OuterMode);
+	return *OuterMode;
+}
+
+UPathedPhysicsMoverComponent& UPathedMovementPatternBase::GetPathedMoverComp() const
+{
+	return GetMovementMode().GetPathedMoverComp();
+}
\ No newline at end of file
diff --git a/Engine/Plugins/Experimental/Mover/Source/Mover/Private/PhysicsMover/PathedMovement/PathedMovementTypes.cpp b/Engine/Plugins/Experimental/Mover/Source/Mover/Private/PhysicsMover/PathedMovement/PathedMovementTypes.cpp
new file mode 100644
--- /dev/null
+++ b/Engine/Plugins/Experimental/Mover/Source/Mover/Private/PhysicsMover/PathedMovement/PathedMovementTypes.cpp
@@ -0,0 +1,134 @@
+// Copyright Epic Games, Inc. All Rights Reserved.
+
+#include "PhysicsMover/PathedMovement/PathedMovementTypes.h"
+
+#include UE_INLINE_GENERATED_CPP_BY_NAME(PathedMovementTypes)
+
+bool FMutablePathedMovementProperties::IsMoving() const
+{
+	return MovementStartFrame >= 0;
+}
+
+bool FMutablePathedMovementProperties::IsLooping() const
+{
+	return PlaybackBehavior == EPathedPhysicsPlaybackBehavior::Looping || PlaybackBehavior == EPathedPhysicsPlaybackBehavior::PingPong;
+}
+
+bool FMutablePathedMovementProperties::IsPingPonging() const
+{
+	return PlaybackBehavior == EPathedPhysicsPlaybackBehavior::ThereAndBack || PlaybackBehavior == EPathedPhysicsPlaybackBehavior::PingPong;
+}
+
+void FMutablePathedMovementProperties::NetSerialize(FArchive& Ar)
+{
+	Ar << MovementStartFrame;
+	Ar << bIsInReverse;
+	Ar << bIsJointEnabled;
+	Ar << PlaybackBehavior;
+	Ar << PathOrigin;
+}
+
+void FMutablePathedMovementProperties::ToString(FAnsiStringBuilderBase& Out) const
+{
+	Out.Appendf("MovementStartFrame: %d | ", MovementStartFrame);
+	Out.Appendf("bIsInReverse: %hs | ", *LexToString(bIsInReverse));
+	Out.Appendf("bIsJointEnabled: %hs | ", *LexToString(bIsJointEnabled));
+	Out.Appendf("PlaybackBehavior: %hs | ", *StaticEnum<EPathedPhysicsPlaybackBehavior>()->GetValueAsString(PlaybackBehavior));
+	Out.Appendf("PathOrigin: %hs\n", *PathOrigin.ToHumanReadableString());
+}
+
+bool FMutablePathedMovementProperties::operator==(const FMutablePathedMovementProperties& Other) const
+{
+	return MovementStartFrame == Other.MovementStartFrame
+		&& bIsInReverse == Other.bIsInReverse
+		&& bIsJointEnabled == Other.bIsJointEnabled
+		&& PlaybackBehavior == Other.PlaybackBehavior
+		&& PathOrigin.Equals(Other.PathOrigin);
+}
+
+bool FPathedPhysicsMovementInputs::NetSerialize(FArchive& Ar, UPackageMap* Map, bool& bOutSuccess)
+{
+	Super::NetSerialize(Ar, Map, bOutSuccess);
+
+	Props.NetSerialize(Ar);
+	
+	bOutSuccess = true;
+	return true;
+}
+
+void FPathedPhysicsMovementInputs::ToString(FAnsiStringBuilderBase& Out) const
+{
+	Props.ToString(Out);
+}
+
+bool FPathedPhysicsMovementInputs::ShouldReconcile(const FMoverDataStructBase& AuthorityState) const
+{
+	const FPathedPhysicsMovementInputs& TypedAuthority = static_cast<const FPathedPhysicsMovementInputs&>(AuthorityState);
+	return Props != TypedAuthority.Props;
+}
+
+void FPathedPhysicsMovementInputs::Interpolate(const FMoverDataStructBase& From, const FMoverDataStructBase& To, float Pct)
+{
+	if (Pct < .5f)
+	{
+		*this = static_cast<const FPathedPhysicsMovementInputs&>(From);
+	}
+	else
+	{
+		*this = static_cast<const FPathedPhysicsMovementInputs&>(To);
+	}
+}
+
+void FPathedPhysicsMovementInputs::Merge(const FMoverDataStructBase& From)
+{
+	const FPathedPhysicsMovementInputs& TypedFrom = static_cast<const FPathedPhysicsMovementInputs&>(From);
+	//@todo DanH: What, if anything, do I actually want to merge on the mutable props?
+}
+
+bool FPathedPhysicsMovementState::NetSerialize(FArchive& Ar, UPackageMap* Map, bool& bOutSuccess)
+{
+	Super::NetSerialize(Ar, Map, bOutSuccess);
+
+	Ar << LastStopPlaybackTime;
+	Ar << CurrentProgress;
+	MutableProps.NetSerialize(Ar);
+	
+	bOutSuccess = true;
+	return true;
+}
+
+void FPathedPhysicsMovementState::ToString(FAnsiStringBuilderBase& Out) const
+{
+	FMoverDataStructBase::ToString(Out);
+
+	Out.Appendf("PlaybackTime: %.2f | ", LastStopPlaybackTime);
+	Out.Appendf("CurrentProgress: %.2f | ", CurrentProgress);
+	MutableProps.ToString(Out);
+}
+
+bool FPathedPhysicsMovementState::ShouldReconcile(const FMoverDataStructBase& AuthorityState) const
+{
+	const FPathedPhysicsMovementState& TypedAuthority = static_cast<const FPathedPhysicsMovementState&>(AuthorityState);
+
+	static constexpr float PlaybackTimeErrorTolerance = 0.05f;
+	static constexpr float ProgressErrorTolerance = 0.05f;
+
+	if (MutableProps != TypedAuthority.MutableProps ||
+		!FMath::IsNearlyEqual(LastStopPlaybackTime, TypedAuthority.LastStopPlaybackTime, PlaybackTimeErrorTolerance) ||
+		!FMath::IsNearlyEqual(CurrentProgress, TypedAuthority.CurrentProgress, ProgressErrorTolerance))
+	{
+		return true;
+	}
+
+	return false;
+}
+
+void FPathedPhysicsMovementState::Interpolate(const FMoverDataStructBase& From, const FMoverDataStructBase& To, float Pct)
+{
+	const FPathedPhysicsMovementState& TypedFrom = static_cast<const FPathedPhysicsMovementState&>(From);
+	const FPathedPhysicsMovementState& TypedTo = static_cast<const FPathedPhysicsMovementState&>(To);
+
+	MutableProps = TypedTo.MutableProps;
+	LastStopPlaybackTime = FMath::Lerp(TypedFrom.LastStopPlaybackTime, TypedTo.LastStopPlaybackTime, Pct);
+	CurrentProgress = FMath::Lerp(TypedFrom.CurrentProgress, TypedTo.CurrentProgress, Pct);
+}
\ No newline at end of file
diff --git a/Engine/Plugins/Experimental/Mover/Source/Mover/Private/PhysicsMover/PathedMovement/PathedPhysicsDebugDrawComponent.cpp b/Engine/Plugins/Experimental/Mover/Source/Mover/Private/PhysicsMover/PathedMovement/PathedPhysicsDebugDrawComponent.cpp
new file mode 100644
--- /dev/null
+++ b/Engine/Plugins/Experimental/Mover/Source/Mover/Private/PhysicsMover/PathedMovement/PathedPhysicsDebugDrawComponent.cpp
@@ -0,0 +1,334 @@
+// Copyright Epic Games, Inc. All Rights Reserved.
+
+#include "PhysicsMover/PathedMovement/PathedPhysicsDebugDrawComponent.h"
+
+#include "StaticMeshSceneProxy.h"
+#include "Components/MeshComponent.h"
+#include "Components/StaticMeshComponent.h"
+#include "Engine/World.h"
+#include "PhysicsMover/PathedMovement/PathedMovementMode.h"
+#include "PhysicsMover/PathedMovement/PathedMovementPatternBase.h"
+#include "PhysicsMover/PathedMovement/PathedPhysicsMoverComponent.h"
+
+TAutoConsoleVariable<int32> CVarPathedPhysicsTotalDebugDrawSteps(
+	TEXT("Mover.PathedPhysics.DebugDraw.TotalNumSteps"),
+	32,
+	TEXT("How many steps/lines to draw for each debug drawn path"));
+	
+TAutoConsoleVariable<int32> CVarPathedPhysicsDisplayedDebugDrawSteps(
+	TEXT("Mover.PathedPhysics.DebugDraw.DisplayedSteps"),
+	0,
+	TEXT("Of the total number of steps in the path debug draw (see Mover.PathedPhysics.DebugDraw.TotalNumSteps), how many should actually get drawn? If <= 0, all steps are drawn."));
+
+class FPathedPhysicsDebugRenderSceneProxy : public FDebugRenderSceneProxy
+{
+public:
+	explicit FPathedPhysicsDebugRenderSceneProxy(const UPrimitiveComponent& InComponent)
+		: FDebugRenderSceneProxy(&InComponent)
+	{}
+
+	virtual SIZE_T GetTypeHash() const override 
+	{
+		static size_t UniquePointer;
+		return reinterpret_cast<size_t>(&UniquePointer);
+	}
+	
+	virtual FPrimitiveViewRelevance GetViewRelevance(const FSceneView* View) const override
+	{
+		FPrimitiveViewRelevance Result;
+		Result.bDrawRelevance = true; ///*bDrawDebug && *//*!IsSelected() &&*/ IsShown(View) /*&& View->Family->EngineShowFlags.Splines*/;
+		Result.bDynamicRelevance = true;
+		// Result.bShadowRelevance = IsShadowCast(View);
+		Result.bEditorPrimitiveRelevance = true; // UseEditorCompositing(View);
+		return Result;
+	}
+
+	virtual uint32 GetMemoryFootprint(void) const override { return sizeof(*this) + GetAllocatedSize(); }
+	
+	//@todo DanH: Add a show flag for these
+	// virtual FPrimitiveViewRelevance GetViewRelevance(const FSceneView* View) const override;
+	// virtual uint32 GetMemoryFootprint() const override;
+};
+
+UPathedPhysicsDebugDrawComponent::UPathedPhysicsDebugDrawComponent()
+{
+	// bIsEditorOnly = true;
+}
+
+void UPathedPhysicsDebugDrawComponent::OnRegister()
+{
+	// DestroyProgressPreviewMeshComp();
+	UpdatePreviewMeshComp();
+
+#if WITH_EDITOR
+	FCoreUObjectDelegates::OnObjectPropertyChanged.AddUObject(this, &ThisClass::HandleObjectPropertyChanged);
+#endif
+
+	Super::OnRegister();
+}
+
+void UPathedPhysicsDebugDrawComponent::OnUnregister()
+{
+	Super::OnUnregister();
+
+#if WITH_EDITOR
+	FCoreUObjectDelegates::OnObjectPropertyChanged.RemoveAll(this);
+#endif
+
+	if (ProgressPreviewMeshComp)
+	{
+		ProgressPreviewMeshComp->UnregisterComponent();
+	}
+}
+
+FBoxSphereBounds UPathedPhysicsDebugDrawComponent::CalcBounds(const FTransform& LocalToWorld) const
+{
+	return Bounds;
+}
+
+void UPathedPhysicsDebugDrawComponent::CreateRenderState_Concurrent(FRegisterComponentContext* Context)
+{
+	// Calculate everything for the scene proxy here so that our bounds are ready early enough (they get checked just before the scene proxy gets added)
+	FBoxSphereBounds::Builder DebugBoundsBuilder;
+	DebugBoundsBuilder += GetOwner()->GetTransform().GetLocation();
+	
+	DebugLines.Reset();
+	DebugDashedLines.Reset();
+	DebugArrowLines.Reset();
+	DebugSpheres.Reset();
+	DebugStars.Reset();
+
+	const UPathedPhysicsMoverComponent& PathedMoverComp = *GetTypedOuter<UPathedPhysicsMoverComponent>();
+	const UPathedPhysicsMovementMode* PathedMovementMode = PathedMoverComp.FindMode_Mutable<UPathedPhysicsMovementMode>(PathedMoverComp.StartingMovementMode);
+	if (PathedMovementMode && !PathedMovementMode->PathPatterns.IsEmpty() && (PathedMovementMode->bDebugDrawAggregatePath || PathedMovementMode->bAllowPatternDebugDrawing))
+	{
+		int32 NumPatternsToDebugDraw = 0;
+		TArray<const UPathedMovementPatternBase*> ValidPatterns;
+		for (UPathedMovementPatternBase* PathPattern : PathedMovementMode->PathPatterns)
+		{
+			if (PathPattern)
+			{
+				ValidPatterns.Emplace(PathPattern);
+				if (PathPattern->DebugDrawUsingStepSamples())
+				{
+					++NumPatternsToDebugDraw;
+				}
+
+				if (PathedMovementMode->bAllowPatternDebugDrawing)
+				{
+					PathPattern->AppendDebugDrawElements(*this, DebugBoundsBuilder);
+				}
+			}
+		}
+
+		if (NumPatternsToDebugDraw > 0 || PathedMovementMode->bDebugDrawAggregatePath)
+		{
+			const FTransform LocalToWorld = GetOwner()->GetTransform();
+			FVector PreviousStepAggregateLocation = LocalToWorld.GetLocation();
+			TArray<FVector> PreviousStepPatternLocations;
+
+			int32 NumCurvesToDraw = PathedMovementMode->bDebugDrawAggregatePath ? 1 : 0;
+			if (PathedMovementMode->bAllowPatternDebugDrawing)
+			{
+				PreviousStepPatternLocations.SetNumZeroed(ValidPatterns.Num());
+				PreviousStepPatternLocations[0] = PreviousStepAggregateLocation;
+				// for (FVector& PatternLoc : PreviousStepPatternLocations)
+				// {
+				// 	PatternLoc = PreviousStepAggregateLocation;
+				// }
+				NumCurvesToDraw += NumPatternsToDebugDraw;
+			}
+			
+			const float NumSteps = CVarPathedPhysicsTotalDebugDrawSteps.GetValueOnGameThread();
+			DebugLines.Reserve(DebugLines.Num() + NumSteps * NumCurvesToDraw);
+			const int32 NumStepsToDraw = CVarPathedPhysicsDisplayedDebugDrawSteps.GetValueOnGameThread() > 0 ? CVarPathedPhysicsDisplayedDebugDrawSteps.GetValueOnGameThread() : NumSteps;
+			for (int32 Step = 1; Step <= NumStepsToDraw; ++Step)
+			{
+				FTransform TargetRelativeTransform = FTransform::Identity;
+				const float ProgressAmt = Step / NumSteps;
+
+				for (int32 PatternIdx = 0; PatternIdx < ValidPatterns.Num(); ++PatternIdx)
+				{
+					const UPathedMovementPatternBase& Pattern = *ValidPatterns[PatternIdx];
+					const bool bDrawSampledPattern = Pattern.DebugDrawUsingStepSamples() && PathedMovementMode->bAllowPatternDebugDrawing;
+					if (bDrawSampledPattern || PathedMovementMode->bDebugDrawAggregatePath)
+					{
+						// Effectively the same thing as UPathedPhysicsMovementMode::CalcTargetRelativeTransform, done manually to be able to draw each pattern as well
+						const FTransform PatternRelativeTransform = Pattern.CalcTargetRelativeTransform(ProgressAmt, TargetRelativeTransform);
+						TargetRelativeTransform.Accumulate(PatternRelativeTransform);
+
+						if (bDrawSampledPattern)
+						{
+							const FVector PatternLocation = Chaos::FRigidTransform3::MultiplyNoScale(PatternRelativeTransform, LocalToWorld).GetLocation();
+							DebugLines.Emplace(PreviousStepPatternLocations[PatternIdx], PatternLocation, Pattern.PatternDebugDrawColor, 1.f);
+
+							PreviousStepPatternLocations[PatternIdx] = PatternLocation;
+							DebugBoundsBuilder += PatternLocation;
+						}
+					}
+				}
+
+				// Aggregate Path
+				if (PathedMovementMode->bDebugDrawAggregatePath)
+				{
+					const FVector StepLocation = Chaos::FRigidTransform3::MultiplyNoScale(TargetRelativeTransform, LocalToWorld).GetLocation();
+					DebugLines.Emplace(PreviousStepAggregateLocation, StepLocation, PathedMovementMode->PathDebugDrawColor, 2.f);
+					PreviousStepAggregateLocation = StepLocation;
+					DebugBoundsBuilder += StepLocation;
+				}
+			}
+		}	
+	}
+
+	Bounds = DebugBoundsBuilder;
+	Bounds = Bounds.ExpandBy(25.);
+
+	Super::CreateRenderState_Concurrent(Context);
+}
+
+FDebugRenderSceneProxy* UPathedPhysicsDebugDrawComponent::CreateDebugSceneProxy()
+{
+	if (DebugLines.IsEmpty() &&
+		DebugDashedLines.IsEmpty() &&
+		DebugArrowLines.IsEmpty() &&
+		DebugSpheres.IsEmpty() &&
+		DebugStars.IsEmpty())
+	{
+		return nullptr;
+	}
+
+	FPathedPhysicsDebugRenderSceneProxy* DebugSceneProxy = new FPathedPhysicsDebugRenderSceneProxy(*this);
+    DebugSceneProxy->DrawType = FDebugRenderSceneProxy::SolidAndWireMeshes;
+    // SceneProxy->ViewFlagName = FString();
+    // SceneProxy->ViewFlagIndex = 0;
+    // SceneProxy->ViewFlagName = TEXT("PathedPhysics");
+    DebugSceneProxy->ViewFlagName = TEXT("Splines");
+    DebugSceneProxy->ViewFlagIndex = uint32(FEngineShowFlags::FindIndexByName(*DebugSceneProxy->ViewFlagName));
+	
+	DebugSceneProxy->Lines = DebugLines;
+	DebugSceneProxy->DashedLines = DebugDashedLines;
+	DebugSceneProxy->ArrowLines = DebugArrowLines;
+	DebugSceneProxy->Spheres = DebugSpheres;
+	DebugSceneProxy->Stars = DebugStars;
+	
+	return DebugSceneProxy;
+}
+
+void UPathedPhysicsDebugDrawComponent::UpdatePreviewMeshComp(bool bForce)
+{
+	UMeshComponent* MeshRoot = GetOwnerMeshRoot();
+	if (MeshRoot && GetWorld() && !GetWorld()->IsGameWorld())
+	{
+		bool bDoUpdate = false;
+		
+		const UPathedPhysicsMoverComponent& PathedMoverComp = *GetTypedOuter<UPathedPhysicsMoverComponent>();
+    	if (ProgressPreviewMeshComp)
+    	{
+		    if (ProgressPreviewMeshComp->GetClass() == MeshRoot->GetClass() && bForce)
+		    {
+		    	bDoUpdate = true;
+
+		    	// Copy all the properties on the root mesh comp (except attachments)
+		    	for (TFieldIterator<FProperty> PropIter(MeshRoot->GetClass()); PropIter; ++PropIter)
+		    	{
+				    if (PropIter->GetFName() == TEXT("AttachChildren") || PropIter->GetFName() == TEXT("AttachParent"))
+				    {
+					    continue;
+				    }
+		    		PropIter->CopyCompleteValue_InContainer(ProgressPreviewMeshComp, MeshRoot);
+		    	}
+		    }
+		    else
+		    {
+			    DestroyProgressPreviewMeshComp();
+		    }
+    	}
+
+		if (!ProgressPreviewMeshComp)
+    	{
+    		ProgressPreviewMeshComp = DuplicateObject<UMeshComponent>(MeshRoot, this, TEXT("ProgressPreviewMeshComp"));
+    		// ProgressPreviewMeshComp->SetupAttachment(this);
+			ProgressPreviewMeshComp->bAutoRegister = false;
+
+			bDoUpdate = true;
+    	}
+
+		ensure(ProgressPreviewMeshComp->GetNumChildrenComponents() == 0);
+
+		if (!ProgressPreviewMeshComp->IsRegistered())
+		{
+			ProgressPreviewMeshComp->RegisterComponent();
+		}
+
+		if (bDoUpdate)
+		{
+			//@todo DanH: Double-clicking the preview comp in the level editor can still select the component - how do I disable that?
+			ProgressPreviewMeshComp->CreationMethod = EComponentCreationMethod::UserConstructionScript;
+			ProgressPreviewMeshComp->SetSimulatePhysics(false);
+			ProgressPreviewMeshComp->SetCollisionEnabled(ECollisionEnabled::NoCollision);
+			ProgressPreviewMeshComp->bIsEditorOnly = true;
+			ProgressPreviewMeshComp->SetHiddenInGame(true);
+
+			const UPathedPhysicsMovementMode* PathedMovementMode = PathedMoverComp.FindMode_Mutable<UPathedPhysicsMovementMode>(PathedMoverComp.StartingMovementMode);
+			if (PathedMovementMode && PathedMoverComp.bDisplayProgressPreviewMesh)
+			{
+				ProgressPreviewMeshComp->SetVisibility(true);
+			
+				const FTransform PreviewMeshRelativeTransform = PathedMovementMode->CalcTargetRelativeTransform(PathedMoverComp.PreviewMeshProgress);
+
+				// Calculate the location without scale, then tack on the source mesh's scale
+				FTransform PreviewMeshWorldTransform = Chaos::FRigidTransform3::MultiplyNoScale(PreviewMeshRelativeTransform, MeshRoot->GetComponentTransform());
+				PreviewMeshWorldTransform.SetScale3D(MeshRoot->GetComponentScale());
+				ProgressPreviewMeshComp->SetWorldTransform(PreviewMeshWorldTransform);
+
+				if (PathedMoverComp.ProgressPreviewMeshMaterial)
+				{
+					const int32 NumMaterials = ProgressPreviewMeshComp->GetNumMaterials();
+					for (int MaterialIdx = 0; MaterialIdx < NumMaterials; ++MaterialIdx)
+					{
+						ProgressPreviewMeshComp->SetMaterial(MaterialIdx, PathedMoverComp.ProgressPreviewMeshMaterial);
+					}
+				}
+			}
+			else
+			{
+				ProgressPreviewMeshComp->SetVisibility(false);
+			}
+		}
+	}
+	else if (ProgressPreviewMeshComp)
+	{
+		DestroyProgressPreviewMeshComp();
+	}
+}
+
+void UPathedPhysicsDebugDrawComponent::DestroyProgressPreviewMeshComp()
+{
+	if (ProgressPreviewMeshComp)
+	{
+		ProgressPreviewMeshComp->DestroyComponent();
+		ProgressPreviewMeshComp = nullptr;
+	}
+}
+
+UMeshComponent* UPathedPhysicsDebugDrawComponent::GetOwnerMeshRoot() const
+{
+	if (AActor* Owner = GetOwner())
+	{
+		return Cast<UMeshComponent>(Owner->GetRootComponent());
+	}
+	return nullptr;
+}
+
+#if WITH_EDITOR
+void UPathedPhysicsDebugDrawComponent::HandleObjectPropertyChanged(UObject* Object, FPropertyChangedEvent& PropertyChangedEvent)
+{
+	if (UMeshComponent* MeshComponent = Cast<UMeshComponent>(Object);
+		MeshComponent && MeshComponent == GetOwner()->GetRootComponent())
+	{
+		UpdatePreviewMeshComp();		
+	}
+}
+#endif
+
+#undef LOCTEXT_NAMESPACE
diff --git a/Engine/Plugins/Experimental/Mover/Source/Mover/Private/PhysicsMover/PathedMovement/PathedPhysicsMoverComponent.cpp b/Engine/Plugins/Experimental/Mover/Source/Mover/Private/PhysicsMover/PathedMovement/PathedPhysicsMoverComponent.cpp
new file mode 100644
--- /dev/null
+++ b/Engine/Plugins/Experimental/Mover/Source/Mover/Private/PhysicsMover/PathedMovement/PathedPhysicsMoverComponent.cpp
@@ -0,0 +1,151 @@
+// Copyright Epic Games, Inc. All Rights Reserved.
+
+#include "PhysicsMover/PathedMovement/PathedPhysicsMoverComponent.h"
+
+#include "Backends/MoverPathedPhysicsLiaison.h"
+#include "PhysicsMover/PathedMovement/PathedMovementMode.h"
+#include "PhysicsMover/PathedMovement/PathedPhysicsDebugDrawComponent.h"
+
+#include UE_INLINE_GENERATED_CPP_BY_NAME(PathedPhysicsMoverComponent)
+
+UPathedPhysicsMoverComponent::UPathedPhysicsMoverComponent()
+{
+	MovementModes.Add(TEXT("Default"), CreateDefaultSubobject<UPathedPhysicsMovementMode>(TEXT("DefaultPath")));
+	StartingMovementMode = TEXT("Default");
+	
+	BackendClass = UMoverPathedPhysicsLiaisonComponent::StaticClass();
+	bSupportsKinematicBasedMovement = false;
+
+#if ENABLE_DRAW_DEBUG
+	DebugDrawComp = CreateDefaultSubobject<UPathedPhysicsDebugDrawComponent>(TEXT("PathDebugDraw")/*, true*/);
+	
+	// Since the debug draw comp relies on data in this component, and this component will often need to be re-registered following the construction script,
+	// we want registration of the debug draw comp to match (and follow) the registration state of this component
+	DebugDrawComp->bAutoRegister = false;
+	
+	// DebugDrawComp->SetFlags(RF_Transactional);
+	// if (GetOwner())
+	// {
+	// 	DebugDrawComp->SetupAttachment(GetOwner()->GetDefaultAttachComponent());
+	// }
+#endif
+}
+
+void UPathedPhysicsMoverComponent::OnRegister()
+{
+	Super::OnRegister();
+
+	for (TPair<FName, TObjectPtr<UBaseMovementMode>>& NameModePair : MovementModes)
+	{
+		if (UPathedPhysicsMovementMode* PathedMode = Cast<UPathedPhysicsMovementMode>(NameModePair.Value))
+		{
+			PathedMode->InitializePath();
+		}
+	}
+	
+#if ENABLE_DRAW_DEBUG
+	DebugDrawComp->RegisterComponent();
+#endif
+}
+
+void UPathedPhysicsMoverComponent::OnUnregister()
+{
+	Super::OnUnregister();
+
+#if ENABLE_DRAW_DEBUG
+	DebugDrawComp->UnregisterComponent();
+#endif
+}
+
+void UPathedPhysicsMoverComponent::InitializeComponent()
+{
+	Super::InitializeComponent();
+
+	//@todo DanH: Log error about invalid backend class (better would be to make it impossible to choose the wrong one)
+	PathedPhysicsLiaison = Cast<UMoverPathedPhysicsLiaisonComponent>(BackendLiaisonComp.GetObject());
+
+	if (PathedPhysicsLiaison)
+	{
+		PathedPhysicsLiaison->SetPathOrigin(GetUpdatedComponentTransform());
+		PathedPhysicsLiaison->SetIsMoving(bAutoMoveOnSpawn, MovementStartDelay);
+		PathedPhysicsLiaison->SetPlaybackBehavior(DefaultPlaybackBehavior);
+	}
+}
+
+void UPathedPhysicsMoverComponent::UninitializeComponent()
+{
+	Super::UninitializeComponent();
+	PathedPhysicsLiaison = nullptr;
+}
+
+#if WITH_EDITOR
+void UPathedPhysicsMoverComponent::PostEditChangeProperty(struct FPropertyChangedEvent& PropertyChangedEvent)
+{
+	Super::PostEditChangeProperty(PropertyChangedEvent);
+
+#if ENABLE_DRAW_DEBUG
+	// DebugDrawComp->MarkRenderStateDirty();
+#endif
+}
+#endif
+
+bool UPathedPhysicsMoverComponent::IsInReverse() const
+{
+	return PathedPhysicsLiaison && PathedPhysicsLiaison->IsInReverse();
+}
+
+void UPathedPhysicsMoverComponent::SetPlaybackDirection(bool bPlayForward)
+{
+	if (PathedPhysicsLiaison)
+	{
+		PathedPhysicsLiaison->SetPlaybackDirection(bPlayForward);
+	}
+}
+
+bool UPathedPhysicsMoverComponent::IsMoving() const
+{
+	return PathedPhysicsLiaison && PathedPhysicsLiaison->IsMoving();
+}
+
+void UPathedPhysicsMoverComponent::SetIsMoving(bool bShouldMove)
+{
+	if (PathedPhysicsLiaison)
+	{
+		PathedPhysicsLiaison->SetIsMoving(bShouldMove);
+	}
+}
+
+// EPathedPhysicsPlaybackBehavior UPathedPhysicsMoverComponent::GetPlaybackBehavior() const
+// {
+// 	return PathedPhysicsLiaison ? PathedPhysicsLiaison->GetPlaybackBehavior() : EPathedPhysicsPlaybackBehavior::OneShot;
+// }
+
+void UPathedPhysicsMoverComponent::SetDefaultPlaybackBehavior(EPathedPhysicsPlaybackBehavior PlaybackBehavior)
+{
+	//@todo DanH: Change the default behavior here, but don't necessarily change the behavior if the mode has overridden it
+	DefaultPlaybackBehavior = PlaybackBehavior;
+}
+
+bool UPathedPhysicsMoverComponent::IsJointEnabled() const
+{
+	return PathedPhysicsLiaison && PathedPhysicsLiaison->IsJointEnabled();
+}
+
+void UPathedPhysicsMoverComponent::SetPathOriginTransform(const FTransform& NewPathOrigin)
+{
+	if (PathedPhysicsLiaison)
+	{
+		PathedPhysicsLiaison->SetPathOrigin(NewPathOrigin);
+	}
+}
+
+const FTransform& UPathedPhysicsMoverComponent::GetPathOriginTransform() const
+{
+	// The liaison isn't established in editor worlds, where it's safe to assume the actor location is the path origin
+	return PathedPhysicsLiaison ? PathedPhysicsLiaison->GetPathOrigin() : GetOwner()->GetActorTransform();
+}
+
+void UPathedPhysicsMoverComponent::NotifyIsMovingChanged(bool bIsMoving)
+{
+	//@todo DanH: Fire an event about it
+}
diff --git a/Engine/Plugins/Experimental/Mover/Source/Mover/Private/PhysicsMover/PathedMovement/PointMovementPathPattern.cpp b/Engine/Plugins/Experimental/Mover/Source/Mover/Private/PhysicsMover/PathedMovement/PointMovementPathPattern.cpp
new file mode 100644
--- /dev/null
+++ b/Engine/Plugins/Experimental/Mover/Source/Mover/Private/PhysicsMover/PathedMovement/PointMovementPathPattern.cpp
@@ -0,0 +1,127 @@
+// Copyright Epic Games, Inc. All Rights Reserved.
+
+#include "PhysicsMover/PathedMovement/PointMovementPathPattern.h"
+
+#include "DebugRenderSceneProxy.h"
+#include "Kismet/KismetMathLibrary.h"
+#include "PhysicsMover/PathedMovement/PathedPhysicsDebugDrawComponent.h"
+#include "PhysicsMover/PathedMovement/PathedPhysicsMoverComponent.h"
+
+void UPointMovementPathPattern::InitializePattern()
+{
+	Super::InitializePattern();
+
+	RefreshAssignedPointProgress(true);
+}
+
+#if WITH_EDITOR
+void UPointMovementPathPattern::PostEditChangeProperty(struct FPropertyChangedEvent& PropertyChangedEvent)
+{
+	Super::PostEditChangeProperty(PropertyChangedEvent);
+
+	RefreshAssignedPointProgress(true);
+}
+#endif
+
+void UPointMovementPathPattern::AppendDebugDrawElements(UPathedPhysicsDebugDrawComponent& DebugDrawComp, FBoxSphereBounds::Builder& InOutDebugBoundsBuilder)
+{
+	if (!PathPoints.IsEmpty())
+	{
+		const FTransform& PathOriginTransform = GetPathedMoverComp().GetPathOriginTransform();
+		FVector PrevPointLocation = PathOriginTransform.GetLocation();
+		for (const FPointMovementPathPoint& PathPoint : PathPoints)
+		{
+			DebugDrawComp.DebugSpheres.Emplace(4.f, PrevPointLocation, PatternDebugDrawColor);
+
+			FVector PointLocation = PathPoint.WorldLoc;
+			if (PathPoint.Basis == EPointMovementLocationBasis::PreviousPoint)
+			{
+				PointLocation = PrevPointLocation + PathPoint.Location;	
+			}
+			else if (PathPoint.Basis == EPointMovementLocationBasis::PathOrigin)
+			{
+				PointLocation = PathOriginTransform.TransformPositionNoScale(PathPoint.Location);
+			}
+			
+			DebugDrawComp.DebugLines.Emplace(PrevPointLocation, PointLocation, PatternDebugDrawColor, 1.f);
+			InOutDebugBoundsBuilder += PointLocation;
+			
+			PrevPointLocation = PointLocation;
+		}
+
+		DebugDrawComp.DebugSpheres.Emplace(4.f, PrevPointLocation, PatternDebugDrawColor);
+	}
+}
+
+FTransform UPointMovementPathPattern::CalcUnmaskedTargetRelativeTransform(float PatternProgress, const FTransform& CurTargetTransform) const
+{
+	// RefreshAssignedPointProgress();
+	
+	// The destination point is the first point at or equal to the target progress
+	int32 DestPointIdx = 0;
+	for (; DestPointIdx < PathPoints.Num() && PathPoints[DestPointIdx].Progress < PatternProgress; ++DestPointIdx) {}
+
+	if (PathPoints.IsValidIndex(DestPointIdx))
+	{
+		const FTransform& PathOrigin = GetPathedMoverComp().GetPathOriginTransform();
+		
+		const FPointMovementPathPoint& NextPoint = PathPoints[DestPointIdx];
+		FPointMovementPathPoint PrevPoint;
+		if (DestPointIdx > 0)
+		{
+			PrevPoint = PathPoints[DestPointIdx - 1];
+		}
+		else
+		{
+			PrevPoint.WorldLoc = PathOrigin.GetLocation();
+		}
+
+		const float ProgressSinceLastPoint = PatternProgress - PrevPoint.Progress;
+		const float Alpha = ProgressSinceLastPoint / (NextPoint.Progress - PrevPoint.Progress);
+		
+		// Point locations are calculated in world space, but path targets need to be provided as relative transforms
+		const FVector TargetWorldLoc = UKismetMathLibrary::VLerp(PrevPoint.WorldLoc, NextPoint.WorldLoc, Alpha);
+		const FVector TargetRelativeLoc = PathOrigin.InverseTransformPositionNoScale(TargetWorldLoc);
+
+		return FTransform(TargetRelativeLoc);
+	}
+	return FTransform::Identity;
+}
+
+void UPointMovementPathPattern::RefreshAssignedPointProgress(bool bForceRefresh) const
+{
+	if (!bHasAssignedPointProgress || bForceRefresh)
+	{
+		TotalPathDistance = 0.f;
+		
+		const FTransform& RootWorldTransform = GetPathedMoverComp().GetPathOriginTransform();
+		FVector PrevPointLocation = RootWorldTransform.GetLocation();
+
+		// Run through and establish point locations and distance
+		for (const FPointMovementPathPoint& PathPoint : PathPoints)
+		{
+			if (PathPoint.Basis == EPointMovementLocationBasis::PreviousPoint)
+			{
+				PathPoint.WorldLoc = PrevPointLocation + PathPoint.Location;	
+			}
+			else if (PathPoint.Basis == EPointMovementLocationBasis::PathOrigin)
+			{
+				PathPoint.WorldLoc = RootWorldTransform.TransformPositionNoScale(PathPoint.Location);
+			}
+			
+			const float DistanceFromPrevious = (PathPoint.WorldLoc - PrevPointLocation).Size();
+			TotalPathDistance += DistanceFromPrevious;
+			PathPoint.DistanceFromStart = TotalPathDistance;
+
+			PrevPointLocation = PathPoint.WorldLoc;
+		}
+
+		// Go through again and assign progress based on distance
+		for (const FPointMovementPathPoint& PathPoint : PathPoints)
+		{
+			PathPoint.Progress = UKismetMathLibrary::SafeDivide(PathPoint.DistanceFromStart, TotalPathDistance);
+		}
+		
+		bHasAssignedPointProgress = true;
+	}
+}
\ No newline at end of file
diff --git a/Engine/Plugins/Experimental/Mover/Source/Mover/Private/PhysicsMover/PathedMovement/SplineMovementPathPattern.cpp b/Engine/Plugins/Experimental/Mover/Source/Mover/Private/PhysicsMover/PathedMovement/SplineMovementPathPattern.cpp
new file mode 100644
--- /dev/null
+++ b/Engine/Plugins/Experimental/Mover/Source/Mover/Private/PhysicsMover/PathedMovement/SplineMovementPathPattern.cpp
@@ -0,0 +1,66 @@
+// Copyright Epic Games, Inc. All Rights Reserved.
+
+#include "PhysicsMover/PathedMovement/SplineMovementPathPattern.h"
+
+#include "Components/SplineComponent.h"
+#include "PhysicsMover/PathedMovement/PathedPhysicsMoverComponent.h"
+
+void USplineMovementPathPattern::InitializePattern()
+{
+	Super::InitializePattern();
+
+	if (!SplineComp)
+	{
+		if (AActor* OwningActor = GetPathedMoverComp().GetOwner(); ensure(OwningActor))
+		{
+			SplineComp = Cast<USplineComponent>(SplineComponentRef.GetComponent(OwningActor));
+			if (!SplineComp)
+			{
+				SplineComp = OwningActor->FindComponentByClass<USplineComponent>();
+			}
+		}
+	}
+}
+
+//TOptional<float> USplineMovementPathPattern::GetPatternDuration() const
+//{
+//	TOptional<float> Duration = Super::GetPatternDuration();
+//	if (bUseSplineDuration)
+//	{
+//		if (const USplineComponent* SplineComponent = GetSpline())
+//		{
+//			Duration = SplineComponent->Duration;
+//		}
+//	}
+//	
+//	if (const float UsableRange = UpperBound - LowerBound; bOffsetsModifyDuration && UsableRange < 1.f && UsableRange > 0.f)
+//	{
+//		const float BaseDuration = Duration.Get(GetMovementMode().GetDefaultPathDuration());
+//		return BaseDuration * UsableRange;
+//	}
+//	return Duration;
+//}
+
+FTransform USplineMovementPathPattern::CalcUnmaskedTargetRelativeTransform(float PatternProgress, const FTransform& CurTargetTransform) const
+{
+	if (SplineComp)
+	{
+		const float SplineLength = SplineComp->GetSplineLength();
+
+		const bool bHasValidBounds = HasValidBounds();
+		const float SplineStartDist = bHasValidBounds ? SplineLength * LowerBound : 0.f;
+		const float SplineEndDist = bHasValidBounds ? SplineLength * UpperBound : 1.f;
+
+		const float DistanceAlongSpline = FMath::Lerp(SplineStartDist, SplineEndDist, PatternProgress);
+		FTransform TargetToSpline = SplineComp->GetTransformAtDistanceAlongSpline(DistanceAlongSpline, ESplineCoordinateSpace::Local, bApplySplineScaling);
+
+		if (!bOrientComponentToPath)
+		{
+			TargetToSpline.SetRotation(FQuat::Identity);
+		}
+		
+		// Intentionally not accounting for any offset between the path origin and the spline's transform - zero progress should always equate to the path origin
+		return TargetToSpline;
+	}
+	return FTransform::Identity;
+}
\ No newline at end of file
diff --git a/Engine/Plugins/Experimental/Mover/Source/Mover/Public/Backends/MoverPathedPhysicsLiaison.h b/Engine/Plugins/Experimental/Mover/Source/Mover/Public/Backends/MoverPathedPhysicsLiaison.h
new file mode 100644
--- /dev/null
+++ b/Engine/Plugins/Experimental/Mover/Source/Mover/Public/Backends/MoverPathedPhysicsLiaison.h
@@ -0,0 +1,80 @@
+// Copyright Epic Games, Inc. All Rights Reserved.
+
+#pragma once
+
+#include "Backends/MoverNetworkPhysicsLiaisonBase.h"
+#include "PhysicsEngine/ConstraintInstance.h"
+#include "PhysicsMover/PhysicsMoverSimulationTypes.h"
+#include "PhysicsMover/PathedMovement/PathedMovementTypes.h"
+
+#include "MoverPathedPhysicsLiaison.generated.h"
+
+class UPathedPhysicsMovementMode;
+class UPathedPhysicsMoverComponent;
+/**
+ * Liaison that works with path-following movement modes.
+ * Establishes a joint constraint between the actual component to update and wherever it should be along the path.
+ * This allows the platform's velocity and rotation to be affected by physics as it moves along the path, and the joint pulls it back where it should be (think spring).
+ * Only compatible with UPathedPhysicsMovementModes, which are responsible for dictating how "loose" the joint between the platform and its ideal position is.
+ */
+UCLASS(Within = PathedPhysicsMoverComponent)
+class MOVER_API UMoverPathedPhysicsLiaisonComponent : public UMoverNetworkPhysicsLiaisonComponentBase
+{
+	GENERATED_BODY()
+
+public:
+	UMoverPathedPhysicsLiaisonComponent();
+
+	virtual void InitializeComponent() override;
+	virtual bool HasValidPhysicsState() const override;
+	virtual void OnCreatePhysicsState() override;
+	virtual void OnDestroyPhysicsState() override;
+	
+	virtual void ConsumeOutput_External(const FPhysicsMoverAsyncOutput& Output, const double OutputTimeInSeconds) override;
+
+	const FTransform& GetPathOrigin() const { return Inputs_External.PathOrigin; }
+	void SetPathOrigin(const FTransform& NewPathOrigin);
+
+	bool IsInReverse() const { return Inputs_External.bIsInReverse; }
+	void SetPlaybackDirection (bool bPlayForward);
+	
+	bool IsMoving() const { return Inputs_External.IsMoving(); }
+	void SetIsMoving(bool bShouldMove, float StartDelay = 0.f);
+
+	bool IsJointEnabled() const { return Inputs_External.bIsJointEnabled; }
+
+	EPathedPhysicsPlaybackBehavior GetPlaybackBehavior() const { return Inputs_External.PlaybackBehavior; }
+	void SetPlaybackBehavior(EPathedPhysicsPlaybackBehavior PlaybackBehavior);
+
+	UPathedPhysicsMoverComponent& GetPathedMoverComp() const;
+	
+protected:
+	virtual void HandleComponentPhysicsStateChanged(UPrimitiveComponent* ChangedComponent, EComponentPhysicsStateChange StateChange) override;
+
+	virtual void PerformProduceInput_External(float DeltaTime, OUT FPhysicsMoverAsyncInput& Input) override;
+	virtual bool CanProcessInputs_Internal(const FPhysicsMoverAsyncInput& Input) const override;
+	virtual void PerformProcessInputs_Internal(int32 PhysicsStep, float DeltaTime, const FPhysicsMoverAsyncInput& Input) const override;
+	virtual bool CanSimulate_Internal(const FPhysicsMoverSimulationTickParams& TickParams, const FPhysicsMoverAsyncInput& Input) const override;
+	virtual void PerformPreSimulate_Internal(const FPhysicsMoverSimulationTickParams& TickParams, const FPhysicsMoverAsyncInput& Input, FPhysicsMoverAsyncOutput& Output) const override;
+
+	virtual void ApplyPathModeConfig(const UPathedPhysicsMovementMode& PathedMode);
+	
+	Chaos::FJointConstraint* GetJointConstraint() const;
+	
+	void CreateTargetJoint();
+	void DestroyTargetJoint();
+	
+protected:
+	FPhysicsUserData PhysicsUserData;
+	FConstraintInstance ConstraintInstance;
+	FPhysicsConstraintHandle ConstraintHandle;
+
+private:
+	UFUNCTION()
+	void HandleMovementModeChanged(const FName& OldModeName, const FName& NewModeName);
+	void HandleIsUsingJointChanged(bool bIsUsingJoint);
+	
+	FMutablePathedMovementProperties Inputs_External;
+
+	FTimerHandle DelayedStartTimerHandle;
+};
diff --git a/Engine/Plugins/Experimental/Mover/Source/Mover/Public/PhysicsMover/PathedMovement/ArcRotationPathPattern.h b/Engine/Plugins/Experimental/Mover/Source/Mover/Public/PhysicsMover/PathedMovement/ArcRotationPathPattern.h
new file mode 100644
--- /dev/null
+++ b/Engine/Plugins/Experimental/Mover/Source/Mover/Public/PhysicsMover/PathedMovement/ArcRotationPathPattern.h
@@ -0,0 +1,31 @@
+// Copyright Epic Games, Inc. All Rights Reserved.
+
+#pragma once
+
+#include "CoreMinimal.h"
+#include "PathedMovementPatternBase.h"
+
+#include "ArcRotationPathPattern.generated.h"
+
+class UPathedPhysicsDebugDrawComponent;
+
+UCLASS()
+class UArcRotationPattern : public UPathedMovementPatternBase
+{
+	GENERATED_BODY()
+
+public:
+protected:
+	virtual bool DebugDrawUsingStepSamples() const override { return false; }
+	virtual void AppendDebugDrawElements(UPathedPhysicsDebugDrawComponent& DebugDrawComp, FBoxSphereBounds::Builder& InOutDebugBoundsBuilder) override;
+	
+	virtual FTransform CalcUnmaskedTargetRelativeTransform(float PatternProgress, const FTransform& CurTargetTransform) const override;
+
+	/** The size of the arc angle to rotate along */
+	UPROPERTY(EditAnywhere, Category = EllipticalPathPattern, meta = (Units = Degrees, UIMin = 0, ClampMin = 0, UIMax = 360, ClampMax = 360))
+	float RotationArcAngle = 360.f;
+
+	/** The axis to rotate about */
+	UPROPERTY(EditAnywhere, Category = EllipticalPathPattern, meta = (UIMin = 0, ClampMin = 0, UIMax = 1.f, ClampMax = 1.f))
+	FVector RotationAxis = FVector::UpVector;
+};
\ No newline at end of file
diff --git a/Engine/Plugins/Experimental/Mover/Source/Mover/Public/PhysicsMover/PathedMovement/EllipticalMovementPathPattern.h b/Engine/Plugins/Experimental/Mover/Source/Mover/Public/PhysicsMover/PathedMovement/EllipticalMovementPathPattern.h
new file mode 100644
--- /dev/null
+++ b/Engine/Plugins/Experimental/Mover/Source/Mover/Public/PhysicsMover/PathedMovement/EllipticalMovementPathPattern.h
@@ -0,0 +1,39 @@
+// Copyright Epic Games, Inc. All Rights Reserved.
+
+#pragma once
+
+#include "CoreMinimal.h"
+#include "PathedMovementPatternBase.h"
+
+#include "EllipticalMovementPathPattern.generated.h"
+
+UCLASS()
+class UEllipticalMovementPathPattern : public UPathedMovementPatternBase
+{
+	GENERATED_BODY()
+	
+public:
+	virtual void InitializePattern() override;
+
+protected:
+	virtual FTransform CalcUnmaskedTargetRelativeTransform(float PatternProgress, const FTransform& CurTargetTransform) const override;
+
+	/** Radius of the ellipse along the x axis */
+	UPROPERTY(EditAnywhere, Category = EllipticalPathPattern)
+	float RadiusX = 100.f;
+	
+	UPROPERTY(EditAnywhere, Category = EllipticalPathPattern)
+	float RadiusY = 100.f;
+
+	/** The angle between the path origin (i.e. the initial location of the component/particle) and the center of the ellipse */
+	UPROPERTY(EditAnywhere, Category = EllipticalPathPattern, meta = (Units = Degrees, UIMin = 0, ClampMin = 0, UIMax = 360, ClampMax = 360))
+	float OriginAngle = 0.f;
+
+	/** The amount of the ellipsoid arc to actually use for the path (where 360 is the entire ellipse) */
+	UPROPERTY(EditAnywhere, Category = EllipticalPathPattern, meta = (Units = Degrees, UIMin = 0, ClampMin = 0, UIMax = 360, ClampMax = 360))
+	float UsableArcAngle = 360.f;
+	
+	/** The world rotation of the plane the ellipse is on */
+	UPROPERTY(EditAnywhere, Category = EllipticalPathPattern)
+	FRotator EllipsePlaneRotation = FRotator::ZeroRotator;
+};
\ No newline at end of file
diff --git a/Engine/Plugins/Experimental/Mover/Source/Mover/Public/PhysicsMover/PathedMovement/LookAtRotationPathPattern.h b/Engine/Plugins/Experimental/Mover/Source/Mover/Public/PhysicsMover/PathedMovement/LookAtRotationPathPattern.h
new file mode 100644
--- /dev/null
+++ b/Engine/Plugins/Experimental/Mover/Source/Mover/Public/PhysicsMover/PathedMovement/LookAtRotationPathPattern.h
@@ -0,0 +1,34 @@
+// Copyright Epic Games, Inc. All Rights Reserved.
+
+#pragma once
+
+#include "CoreMinimal.h"
+#include "PathedMovementPatternBase.h"
+
+#include "LookAtRotationPathPattern.generated.h"
+
+/** Stare at a single fixed point at all times (pairs nicely with ellipse if you want to always look at the center) */
+UCLASS()
+class ULookAtRotationPattern : public UPathedMovementPatternBase
+{
+	GENERATED_BODY()
+
+public:
+	UFUNCTION(BlueprintCallable, Category = LookAtPattern)
+	void SetRelativeLookAtLocation(const FVector& RelativeLookAt);
+
+	UFUNCTION(BlueprintCallable, Category = LookAtPattern)
+	void SetLookAtLocation(const FVector& LookAt);
+
+	virtual bool DebugDrawUsingStepSamples() const override { return false; }
+	virtual void AppendDebugDrawElements(UPathedPhysicsDebugDrawComponent& DebugDrawComp, FBoxSphereBounds::Builder& InOutDebugBoundsBuilder) override;
+	
+protected:
+	virtual FTransform CalcUnmaskedTargetRelativeTransform(float PatternProgress, const FTransform& CurTargetTransform) const override;
+	
+	/** The location relative to the path origin to always look at while moving along the path */
+	UPROPERTY(EditAnywhere, Category = LookAtPattern)
+	FVector RelativeLookAtLocation = FVector::ZeroVector;
+
+	//@todo DanH: Option to target an actor (or component maybe?)
+};
diff --git a/Engine/Plugins/Experimental/Mover/Source/Mover/Public/PhysicsMover/PathedMovement/PathedMovementMode.h b/Engine/Plugins/Experimental/Mover/Source/Mover/Public/PhysicsMover/PathedMovement/PathedMovementMode.h
new file mode 100644
--- /dev/null
+++ b/Engine/Plugins/Experimental/Mover/Source/Mover/Public/PhysicsMover/PathedMovement/PathedMovementMode.h
@@ -0,0 +1,160 @@
+// Copyright Epic Games, Inc. All Rights Reserved.
+
+#pragma once
+
+#include "CoreMinimal.h"
+#include "MovementMode.h"
+#include "MoverTypes.h"
+#include "AlphaBlend.h"
+#include "PhysicsEngine/ConstraintInstance.h"
+#include "PathedMovementTypes.h"
+
+#include "PathedMovementMode.generated.h"
+
+class UPathedPhysicsMoverComponent;
+class UMoverPathedPhysicsLiaisonComponent;
+class UPathedMovementPatternBase;
+class FSceneView;
+struct FPhysicsMoverAsyncOutput;
+struct FPhysicsMoverAsyncInput;
+
+/**
+ * Base for any movement mode that proceeds along a defined path within a specified amount of time
+ * Only functional when used in concert with a PathedPhysicsMoverComponent and MoverPathedPhysicsLiaisonComponent
+ */
+UCLASS(Within = PathedPhysicsMoverComponent)
+class UPathedPhysicsMovementMode : public UBaseMovementMode
+{
+	GENERATED_BODY()
+
+public:
+	UPathedPhysicsMovementMode();
+
+#if WITH_EDITOR
+	virtual void PostEditChangeChainProperty(struct FPropertyChangedChainEvent& PropertyChangedChainEvent) override;
+#endif
+	
+	virtual void OnSimulationTick(const FSimulationTickParams& Params, FMoverTickEndData& OutputState) override final;
+
+	virtual void OnProcessInput_Internal(int32 PhysicsStep, float DeltaTime, const FPhysicsMoverAsyncInput& Input) const;
+	virtual void OnPreSimulate_Internal(const FSimulationTickParams& Params, FMoverTickEndData& OutputState) const;
+
+	/** Called when the owning UPathedPhysicsMoverComponent is registered (i.e. at both edit and game time) */
+	virtual void InitializePath();
+	
+	/** Finds the first path pattern on this mode of the given type, if there is one */
+	UFUNCTION(BlueprintPure, Category = PathedMovementMode, DisplayName = "Find Pattern", meta=(DeterminesOutputType = PatternType))
+	UPathedMovementPatternBase* BP_FindPattern(TSubclassOf<UPathedMovementPatternBase> PatternType) const;
+	
+	/**
+	 * Set the time it takes to travel one way along this path from start to end.
+	 * Only valid to call in BeginPlay or the Construction Script, ignored otherwise.
+	 * Once play has begun, the path duration can be modified indirectly via SetMovementTimeDilation on the owning PathedMoverComponent
+	 */
+	UFUNCTION(BlueprintCallable, Category = PathedMovementMode, DisplayName = "Set Path Duration (BeginPlay only)")
+	void SetPathDuration_BeginPlayOnly(float NewDuration);
+	float GetPathDuration() const { return OneWayTripDuration; }
+
+	UFUNCTION(BlueprintCallable, Category = PathedMovementMode)
+	void SetUseJointConstraint(bool bUseJoint);
+	bool IsUsingJoint() const { return bUseJointConstraint; }
+
+	//@todo DanH: Options for auto-calculations of joint constraint params based on gravity, friction, mass, etc.
+
+	Chaos::FRigidTransform3 CalcTargetRelativeTransform(float ProgressAmt) const;
+	
+	const FConstraintProfileProperties& GetConstraintProperties() const { return JointConstraintProperties; }
+	UPathedPhysicsMoverComponent& GetPathedMoverComp() const;
+	const FMutablePathedMovementProperties& GetCachedInputs() const { return CachedInputs; }
+
+	DECLARE_MULTICAST_DELEGATE_OneParam(FOnIsUsingJointChanged, bool);
+	FOnIsUsingJointChanged& OnIsUsingJointChanged() const { return OnIsUsingJointChangedEvent; }
+
+	const TOptional<EPathedPhysicsPlaybackBehavior>& GetPlaybackBehaviorOverride() const { return PlaybackBehaviorOverride; }
+	
+protected:
+	float GetPlaybackDuration() const;
+	float GetBoundedPlaybackTime(float TotalPlayTime, bool& bOutReachesEndOfPlayback) const;
+	float GetBoundedTimeAlongPath(float Time) const;
+	float GetBoundedTimeAlongPath_Unsafe(float BoundedPlaybackTime) const;
+	
+private:
+	/**
+	 * True to have the updated component be able to respond to physics inputs by being "pulled" along the path using a joint constraint
+	 * configured using JointConstraintProperties below.
+	 * 
+	 * If false, no joint will be used, and the affected component will be moved kinematically (i.e. independent of other physics effects and treated as though it has
+	 * infinite mass in the physics sim) and never deviate from the path.
+	 *
+	 * Note that the movement is calculated on the physics thread either way, so it will correctly interact with other bodies in the sim.
+	 *
+	 * Ex: Platform moving along a path in a straight line. If this is enabled, the platform could wobble, dip, shake, stall, etc. in response to physics forces/impulses
+	 *		(depending on the joint and physics properties of the platform). If this is disabled, the platform will rigidly move along the path and ignore all physics.
+	 */
+	UPROPERTY(EditAnywhere, BlueprintReadOnly, Category = "Pathed Movement Mode", meta = (AllowPrivateAccess))
+	bool bUseJointConstraint = true;
+
+	/**
+	 * Configuration of the joint constraint that (if enabled) pulls the updated component along the path
+	 * The default settings are for a very stiff joint that allows just a little bit of "play"
+	 */
+	UPROPERTY(EditAnywhere, Category = "Pathed Movement Mode", meta = (EditCondition = bUseJointConstraint))
+	FConstraintProfileProperties JointConstraintProperties;
+
+	/** Optional playback behavior override for this mode specifically */
+	UPROPERTY(EditAnywhere, BlueprintReadOnly, Category = "Pathed Movement Mode", meta = (AllowPrivateAccess))
+	TOptional<EPathedPhysicsPlaybackBehavior> PlaybackBehaviorOverride;
+
+	/** Total duration in seconds of a one-way trip from start to end. */
+	UPROPERTY(EditAnywhere, BlueprintReadOnly, Category = "Pathed Movement Mode", meta = (ClampMin=0, UIMin=0, Units=Seconds, AllowPrivateAccess))
+	float OneWayTripDuration = 5.f;
+
+	//@todo DanH: Try to think of a way to make this less order-dependent/error-prone
+	/**
+	 * The set of patterns to follow that combine to form the path the component will follow
+	 *
+	 * NOTE: Order can matter here!! The patterns are evaluated each physics frame in order starting from the front.
+	 *		 Each successive pattern is aware of the aggregate results from the previous patterns.
+	 *		 In most cases, this doesn't matter, as the pattern calculates the target transform against the static path origin point.
+	 *		 However, some patterns necessarily determine their target based on the calculated target to that point
+	 *		 For example, the LookAtRotationPattern calculates the look-at rotation target based on the target location
+	 */
+	UPROPERTY(EditAnywhere, Instanced, BlueprintReadOnly, Category = "Pathed Movement Mode", meta = (AllowPrivateAccess))
+	TArray<TObjectPtr<UPathedMovementPatternBase>> PathPatterns;
+	
+	/** The kind of easing to apply when traveling along the path */
+	UPROPERTY(EditAnywhere, BlueprintReadOnly, Category = "Pathed Movement Mode", meta = (AllowPrivateAccess))
+	EAlphaBlendOption Easing = EAlphaBlendOption::Linear;
+
+	/** If using a custom ease, this is the curve that will be used. If blank, will fall back to standard linear interpolation. */
+	UPROPERTY(EditAnywhere, BlueprintReadOnly, Category = "Pathed Movement Mode", meta = (EditCondition = "Easing==EAlphaBlendOption::Custom", AllowPrivateAccess))
+	TObjectPtr<UCurveFloat> CustomEasingCurve = nullptr;
+
+	mutable FMutablePathedMovementProperties CachedInputs;
+	mutable FOnIsUsingJointChanged OnIsUsingJointChangedEvent;
+	
+	/** True to draw debug lines for the aggregate path resulting from all the patterns in the path (i.e. the path that will actually be followed in game) */
+	UPROPERTY(EditAnywhere, Category = "Pathed Movement Mode|Debug")
+	bool bDebugDrawAggregatePath = true;
+
+	/** True to allow individual patterns to draw debug representations. Can be toggled individually on each pattern as well, this is a master switch for convenience. */
+	UPROPERTY(EditAnywhere, Category = "Pathed Movement Mode|Debug")
+	bool bAllowPatternDebugDrawing = true;
+
+	/** The color used when drawing the aggregate path */
+	UPROPERTY(EditAnywhere, Category = "Pathed Movement Mode|Debug", meta = (EditCondition = bDebugDrawAggregatePath))
+	FColor PathDebugDrawColor = FColor::Blue;
+	
+	friend class UPathedPhysicsDebugDrawComponent;
+	friend class FPatternPostChangeMovementModeHelper;
+};
+
+#if WITH_EDITOR
+// Silly little bit of indirection to let UPathedPhysicsMovementMode::PostEditChangeChainProperty act like it's a friend of UPathedMovementPatternBase without needing
+// to fully declare this class before the pattern (i.e. having PathedMovementPatternBase.h have to include this whole header)
+class FPatternPostChangeMovementModeHelper
+{
+	static void ProcessPostEditChangeChain(UPathedPhysicsMovementMode& Mode, FPropertyChangedChainEvent& PropertyChangedChainEvent);
+	friend class UPathedPhysicsMovementMode;
+};
+#endif
\ No newline at end of file
diff --git a/Engine/Plugins/Experimental/Mover/Source/Mover/Public/PhysicsMover/PathedMovement/PathedMovementPatternBase.h b/Engine/Plugins/Experimental/Mover/Source/Mover/Public/PhysicsMover/PathedMovement/PathedMovementPatternBase.h
new file mode 100644
--- /dev/null
+++ b/Engine/Plugins/Experimental/Mover/Source/Mover/Public/PhysicsMover/PathedMovement/PathedMovementPatternBase.h
@@ -0,0 +1,137 @@
+// Copyright Epic Games, Inc. All Rights Reserved.
+
+#pragma once
+
+#include "CoreMinimal.h"
+#include "PathedMovementTypes.h"
+
+#include "PathedMovementPatternBase.generated.h"
+
+class UPathedPhysicsMoverComponent;
+struct FPhysicsMoverAsyncInput;
+class UCurveFloat;
+enum class EAlphaBlendOption : uint8;
+class UPathedPhysicsMovementMode;
+class UPathedPhysicsDebugDrawComponent;
+
+//@todo DanH: What's the Flags specifier for? Does that do everything now or do I still need the Bitflags meta?
+UENUM(Flags, meta = (Bitflags))
+enum class EPatternAxisMaskFlags : uint8
+{
+	None = 0 UMETA(Hidden),
+
+	X = 1 UMETA(DisplayName = "X / Roll"),
+	Y = 1 << 1 UMETA(DisplayName = "Y / Pitch"),
+	Z = 1 << 2 UMETA(DisplayName = "Z / Yaw)"),
+
+	All = X | Y | Z UMETA(Hidden)
+};
+ENUM_CLASS_FLAGS(EPatternAxisMaskFlags);
+
+UCLASS(Abstract, Within = PathedPhysicsMovementMode, DefaultToInstanced, Blueprintable, BlueprintType, EditInlineNew)
+class UPathedMovementPatternBase : public UObject
+{
+	GENERATED_BODY()
+
+public:
+	virtual void InitializePattern();
+	virtual void ProduceInputs_External(OUT FPhysicsMoverAsyncInput& Input);
+	
+	FTransform CalcTargetRelativeTransform(float OverallPathProgress, const FTransform& CurTargetTransform) const;
+
+	UPathedPhysicsMovementMode& GetMovementMode() const;
+	UPathedPhysicsMoverComponent& GetPathedMoverComp() const;
+
+	virtual bool DebugDrawUsingStepSamples() const { return bDebugDrawPattern; }
+	virtual void AppendDebugDrawElements(UPathedPhysicsDebugDrawComponent& DebugDrawComp, FBoxSphereBounds::Builder& InOutDebugBoundsBuilder) {}
+
+#if WITH_EDITOR
+	//@todo DanH: Pipe this through (I wonder if there should be a UPackage::IsDataValid() override that auto-pipes it to all UObjects in the package?)
+	// virtual EDataValidationResult IsDataValid(class FDataValidationContext& Context) const override;
+
+	friend class FPatternPostChangeMovementModeHelper;
+#endif
+	
+protected:
+	/**
+	* Where patterns calculate their target transforms at the given progress, optionally relative to the current aggregate target.
+	* Axis masking is applied to this result before it's actually used.
+	*/
+	virtual FTransform CalcUnmaskedTargetRelativeTransform(float PatternProgress, const FTransform& CurTargetTransform) const { return FTransform::Identity; }
+
+	//
+	// NOTE: Making any of these properties mutable at game time will require either replicating them (easier, but not as good) or pushing them as inputs to the PT (harder, more correct)
+	//
+
+	/** Along which axes is this pattern disallowed from modifying the translation/location of the updated component? */
+	UPROPERTY(EditAnywhere, Category = "Path Pattern", meta = (Bitmask, BitmaskEnum = "/Script/Mover.EPatternAxisMaskFlags"))
+	uint8 TranslationMasks = (uint8)EPatternAxisMaskFlags::None;
+
+	/** Along which axes is this pattern disallowed from modifying the rotation of the updated component? */
+	UPROPERTY(EditAnywhere, Category = "Path Pattern", meta = (Bitmask, BitmaskEnum = "/Script/Mover.EPatternAxisMaskFlags"))
+	uint8 RotationMasks = (uint8)EPatternAxisMaskFlags::None;
+	
+	/** Along which axes is this pattern disallowed from modifying the scale of the updated component? */
+	UPROPERTY(EditAnywhere, Category = "Path Pattern", meta = (Bitmask, BitmaskEnum = "/Script/Mover.EPatternAxisMaskFlags"))
+	uint8 ScaleMasks = (uint8)EPatternAxisMaskFlags::None;
+
+	/**
+	 * If true, this pattern will not begin to take effect until the previous pattern has completed.
+	 * Note: If true and the previous pattern's EndAtPathProgress is 1, this pattern will never start.
+	 */
+	UPROPERTY(EditAnywhere, Category = "Path Pattern")
+	bool bStartAfterPreviousPattern = false;
+
+	// /**
+	//  * The desired duration (in overall path progress) for this pattern once it begins.
+	//  * If this desired duration is longer than the progress remaining when the pattern starts, it will be clamped to fill what remains
+	//  */
+	// UPROPERTY(EditAnywhere, Category = "Path Pattern", meta = (EditCondition = "bStartAfterPreviousPattern", UIMin = 0.f, ClampMin = 0.f, UIMax = 1.f, ClampMax = 1.f))
+	// float DesiredPatternDuration = 1.f;
+
+	/** The overall path progress when this pattern should begin, where 0 is the start of the path and 1 is the end. Must be less than EndAtProgress. */
+	UPROPERTY(EditAnywhere, Category = "Path Pattern", meta = (EditCondition = "!bStartAfterPreviousPattern", UIMin = 0.f, ClampMin = 0.f, UIMax = 1.f, ClampMax = 1.f))
+	float StartAtPathProgress = 0.f;
+
+	/** The overall path progress when this pattern should complete, where 0 is the start of the path and 1 is the end. Must be greater than StartAtProgress. */
+	UPROPERTY(EditAnywhere, Category = "Path Pattern", meta = (UIMin = 0.f, ClampMin = 0.f, UIMax = 1.f, ClampMax = 1.f))
+	float EndAtPathProgress = 1.f;
+
+	//@todo DanH: Maybe this should be a float and allow partial loops?
+	/**
+	 * The number of loops to complete within the active span of this pattern (i.e. between StartAtProgress and EndAtProgress)
+	 * on a single run along the full aggregate path. Setting to 0 effectively disables this pattern.
+	 */
+	UPROPERTY(EditAnywhere, Category = "Path Pattern", meta = (UIMin = 0, ClampMin = 0))
+	int32 NumLoopsPerPath = 1;
+
+	/** Playback behavior per loop of this pattern */
+	UPROPERTY(EditAnywhere, Category = "Path Pattern", meta = (InvalidEnumValues = "Looping,PingPong"))
+	EPathedPhysicsPlaybackBehavior PerLoopBehavior = EPathedPhysicsPlaybackBehavior::OneShot;
+	
+	/** 
+	 * If true, the component will be rotated to face in the direction of this pattern's motion.
+	 * To have the component face in the direction of the aggregate path, enable this on all movement patterns.
+	 */
+	UPROPERTY(EditAnywhere, Category = "Path Pattern")
+	bool bOrientComponentToPath = false;
+
+	/** The kind of easing to apply when traveling along the path */
+	UPROPERTY(EditAnywhere, Category = "Path Pattern")
+	EAlphaBlendOption Easing;
+
+	/** If using a custom ease, this is the curve that will be used. If blank, will fall back to standard linear interpolation. */
+	UPROPERTY(EditAnywhere, Category = "Path Pattern", meta = (EditCondition = "Easing==EAlphaBlendOption::Custom"))
+	TObjectPtr<UCurveFloat> CustomEasingCurve = nullptr;
+
+	/** True to draw debug lines for this specific pattern in editor views */
+	UPROPERTY(EditAnywhere, Category = "Path Pattern|Debug")
+	bool bDebugDrawPattern = true;
+
+	//@todo DanH: More central setting for color per pattern index maybe?
+	/** The color used for debug draws of this pattern */
+	UPROPERTY(EditAnywhere, Category = "Path Pattern|Debug", meta = (EditCondition = bDrawDebug))
+	FColor PatternDebugDrawColor;
+
+	friend class UPathedPhysicsDebugDrawComponent;
+};
diff --git a/Engine/Plugins/Experimental/Mover/Source/Mover/Public/PhysicsMover/PathedMovement/PathedMovementTypes.h b/Engine/Plugins/Experimental/Mover/Source/Mover/Public/PhysicsMover/PathedMovement/PathedMovementTypes.h
new file mode 100644
--- /dev/null
+++ b/Engine/Plugins/Experimental/Mover/Source/Mover/Public/PhysicsMover/PathedMovement/PathedMovementTypes.h
@@ -0,0 +1,114 @@
+// Copyright Epic Games, Inc. All Rights Reserved.
+
+#pragma once
+
+#include "CoreMinimal.h"
+#include "MoverTypes.h"
+#include "Chaos/Core.h"
+
+#include "PathedMovementTypes.generated.h"
+
+namespace PathBlackboard
+{
+	const FName TargetRelativeTransform = TEXT("TargetRelativeTransform");
+}
+
+UENUM()
+enum class EPathedPhysicsPlaybackBehavior : uint8
+{
+	/** Progresses from 0 -> 1 and then stops */
+	OneShot,
+	/** Progresses from 0 -> 1 -> 0 and then stops (basically a one-shot ping-pong) */
+	ThereAndBack,
+	/** Progresses from 0 -> 1, then resetting to 0, repeating infinitely */
+	Looping,
+	/** Progresses from 0 -> 1 -> 0 -> 1 (and so on) infinitely */
+	PingPong
+};
+
+/**
+ * Properties that can change during game-time and affect path movement calculations.
+ * Split into its own struct because both the input and sync state structs want to hold onto these. 
+ */
+USTRUCT()
+struct FMutablePathedMovementProperties
+{
+	GENERATED_BODY()
+
+	/**
+	 * The server frame of the simulation when movement should/did actually begin
+	 * Tracked this way so that we can set a future frame as the actual movement start frame to account for network latency, allowing clients and server to start at the same time irl.
+	 * A value < 0 indicates that no movement is happening along the path.
+	 */
+	UPROPERTY()
+	int32 MovementStartFrame = INDEX_NONE;
+
+	UPROPERTY()
+	bool bIsInReverse = false;
+
+	UPROPERTY()
+	bool bIsJointEnabled = false;
+
+	UPROPERTY()
+	EPathedPhysicsPlaybackBehavior PlaybackBehavior = EPathedPhysicsPlaybackBehavior::OneShot;
+
+	UPROPERTY()
+	FTransform PathOrigin = FTransform(NoInit);
+
+	bool IsMoving() const;
+	bool IsLooping() const;
+	bool IsPingPonging() const;
+	
+	void NetSerialize(FArchive& Ar);
+	void ToString(FAnsiStringBuilderBase& Out) const;
+	
+	bool operator==(const FMutablePathedMovementProperties& Other) const;
+	bool operator!=(const FMutablePathedMovementProperties& Other) const { return !operator==(Other); }
+};
+
+/**
+ * Frame inputs for all pathed movement
+ * Any property that can change during game time must be delivered in this way, only immutable properties can be safely
+ * referenced directly on the object they come from (usually the movement mode)
+ */
+USTRUCT()
+struct FPathedPhysicsMovementInputs : public FMoverDataStructBase
+{
+	GENERATED_BODY()
+
+	UPROPERTY()
+	FMutablePathedMovementProperties Props;
+
+	virtual bool NetSerialize(FArchive& Ar, UPackageMap* Map, bool& bOutSuccess) override;
+	virtual FMoverDataStructBase* Clone() const override { return new FPathedPhysicsMovementInputs(*this); }
+	virtual UScriptStruct* GetScriptStruct() const override { return StaticStruct(); }
+	virtual void ToString(FAnsiStringBuilderBase& Out) const override;
+	virtual bool ShouldReconcile(const FMoverDataStructBase& AuthorityState) const override;
+	virtual void Interpolate(const FMoverDataStructBase& From, const FMoverDataStructBase& To, float Pct) override;
+	virtual void Merge(const FMoverDataStructBase& From) override;
+};
+
+/** State about the pathed movement that persists from frame to frame */
+USTRUCT()
+struct FPathedPhysicsMovementState : public FMoverDataStructBase
+{
+	GENERATED_BODY()
+	
+	/** The (bounded) playback time along the path where we were the last time movement stopped for any reason */
+	UPROPERTY()
+	float LastStopPlaybackTime = 0.f;
+
+	UPROPERTY()
+	float CurrentProgress = 0.f;
+	
+	/** The most recent mutable properties received from the input */
+	UPROPERTY()
+	FMutablePathedMovementProperties MutableProps;
+
+	virtual bool NetSerialize(FArchive& Ar, UPackageMap* Map, bool& bOutSuccess) override;
+	virtual FMoverDataStructBase* Clone() const override { return new FPathedPhysicsMovementState(*this); }
+	virtual UScriptStruct* GetScriptStruct() const override { return StaticStruct(); }
+	virtual void ToString(FAnsiStringBuilderBase& Out) const override;
+	virtual bool ShouldReconcile(const FMoverDataStructBase& AuthorityState) const override;
+	virtual void Interpolate(const FMoverDataStructBase& From, const FMoverDataStructBase& To, float Pct) override;
+};
\ No newline at end of file
diff --git a/Engine/Plugins/Experimental/Mover/Source/Mover/Public/PhysicsMover/PathedMovement/PathedPhysicsDebugDrawComponent.h b/Engine/Plugins/Experimental/Mover/Source/Mover/Public/PhysicsMover/PathedMovement/PathedPhysicsDebugDrawComponent.h
new file mode 100644
--- /dev/null
+++ b/Engine/Plugins/Experimental/Mover/Source/Mover/Public/PhysicsMover/PathedMovement/PathedPhysicsDebugDrawComponent.h
@@ -0,0 +1,45 @@
+// Copyright Epic Games, Inc. All Rights Reserved.
+
+#pragma once
+
+#include "Debug/DebugDrawComponent.h"
+
+#include "PathedPhysicsDebugDrawComponent.generated.h"
+
+class UMeshComponent;
+
+//@todo DanH: Ideally clicking anything drawn by this would select the mover comp in the editor. I think I need a whole FComponentVisualizer set up for that though?
+UCLASS()
+class UPathedPhysicsDebugDrawComponent : public UDebugDrawComponent
+{
+	GENERATED_BODY()
+
+public:
+	UPathedPhysicsDebugDrawComponent();
+
+	virtual void OnRegister() override;
+	virtual void OnUnregister() override;
+
+	virtual FBoxSphereBounds CalcBounds(const FTransform& LocalToWorld) const override;
+
+	TArray<FDebugRenderSceneProxy::FDebugLine> DebugLines;
+	TArray<FDebugRenderSceneProxy::FDashedLine> DebugDashedLines;
+	TArray<FDebugRenderSceneProxy::FArrowLine> DebugArrowLines;
+	TArray<FDebugRenderSceneProxy::FWireStar> DebugStars;
+	TArray<FDebugRenderSceneProxy::FSphere> DebugSpheres;
+	
+protected:
+	virtual void CreateRenderState_Concurrent(FRegisterComponentContext* Context) override;
+	virtual FDebugRenderSceneProxy* CreateDebugSceneProxy() override;
+
+	void UpdatePreviewMeshComp(bool bForce = false);
+	void DestroyProgressPreviewMeshComp();
+	UMeshComponent* GetOwnerMeshRoot() const;
+
+#if WITH_EDITOR
+	void HandleObjectPropertyChanged(UObject* Object, FPropertyChangedEvent& PropertyChangedEvent);
+#endif
+
+	UPROPERTY(Transient)
+	TObjectPtr<UMeshComponent> ProgressPreviewMeshComp = nullptr;
+};
\ No newline at end of file
diff --git a/Engine/Plugins/Experimental/Mover/Source/Mover/Public/PhysicsMover/PathedMovement/PathedPhysicsMoverComponent.h b/Engine/Plugins/Experimental/Mover/Source/Mover/Public/PhysicsMover/PathedMovement/PathedPhysicsMoverComponent.h
new file mode 100644
--- /dev/null
+++ b/Engine/Plugins/Experimental/Mover/Source/Mover/Public/PhysicsMover/PathedMovement/PathedPhysicsMoverComponent.h
@@ -0,0 +1,110 @@
+// Copyright Epic Games, Inc. All Rights Reserved.
+
+#pragma once
+
+#include "MoverComponent.h"
+#include "PathedMovementTypes.h"
+#include "PathedPhysicsMoverComponent.generated.h"
+
+class UPathedPhysicsDebugDrawComponent;
+class UMoverPathedPhysicsLiaisonComponent;
+
+UCLASS(BlueprintType, meta = (BlueprintSpawnableComponent))
+class MOVER_API UPathedPhysicsMoverComponent : public UMoverComponent
+{
+	GENERATED_BODY()
+
+public:
+	UPathedPhysicsMoverComponent();
+
+	virtual void OnRegister() override;
+	virtual void OnUnregister() override;
+	virtual void InitializeComponent() override;
+	virtual void UninitializeComponent() override;
+
+#if WITH_EDITOR
+	virtual void PostEditChangeProperty(struct FPropertyChangedEvent& PropertyChangedEvent) override;
+#endif
+	
+	UFUNCTION(BlueprintCallable, Category = PathedMover)
+	bool IsMoving() const;
+
+	UFUNCTION(BlueprintCallable, Category = PathedMover)
+	void SetIsMoving(bool bShouldMove);
+
+	UFUNCTION(BlueprintCallable, Category = PathedMover)
+	bool IsInReverse() const;
+
+	UFUNCTION(BlueprintCallable, Category = PathedMover)
+	void SetPlaybackDirection(bool bPlayForward);
+
+	UFUNCTION(BlueprintCallable, Category = PathedMover)
+	EPathedPhysicsPlaybackBehavior GetDefaultPlaybackBehavior() const { return DefaultPlaybackBehavior; }
+
+	UFUNCTION(BlueprintCallable, Category = PathedMover)
+	void SetDefaultPlaybackBehavior(EPathedPhysicsPlaybackBehavior PlaybackBehavior);
+
+	UFUNCTION(BlueprintCallable, Category = PathedMover)
+	bool IsJointEnabled() const;
+
+	bool ShouldMoveOnSpawn() const { return bAutoMoveOnSpawn; }
+
+	//@todo DanH: SetPathProgress, GetPathProgress 
+	//@todo DanH: GetTimeAlongPath (setting this is impossible for nonlinear easing)
+	
+	/**
+	 * Changes path origin (i.e. the transform at the very beginning), causing the whole path to shift accordlingly.
+	 * The updated component is immediately teleported by the delta between the old and new origins, retaining all other physics state
+	 * If a more gradual origin shift is desired, it's up to the caller to set this with a small lerp delta each frame.
+	 */
+	UFUNCTION(BlueprintCallable, Category = PathedMover)
+	void SetPathOriginTransform(const FTransform& NewPathOrigin);
+
+	UFUNCTION(BlueprintCallable, Category = PathedMover)
+	const FTransform& GetPathOriginTransform() const;
+
+	//@todo DanH: See if Core will let me move DEFINE_PASSKEY into the engine
+	//DEFINE_PASSKEY(UMoverPathedPhysicsLiaisonComponent, bool, IsMoving);
+	//void NotifyIsMovingChanged(FIsMovingPasskey bIsMoving);
+	void NotifyIsMovingChanged(bool bIsMoving);
+
+protected:
+	/** Whether to start moving along the path immediately upon beginning play */
+	UPROPERTY(EditAnywhere, Category = "Pathed Movement")
+	bool bAutoMoveOnSpawn = true;
+
+	/** Default follow behavior for the path as a whole */
+	UPROPERTY(EditAnywhere, BlueprintReadWrite, Category = "Pathed Movement")
+	EPathedPhysicsPlaybackBehavior DefaultPlaybackBehavior = EPathedPhysicsPlaybackBehavior::OneShot;
+
+	/** How long after being told to StartMoving() to delay before actually moving. Generally used to stagger the start times of a handful of autoplaying instances. */
+	UPROPERTY(EditAnywhere, BlueprintReadOnly, Category = "Pathed Movement")
+	float MovementStartDelay = 0.f;
+
+	/**  */
+	
+	//@todo DanH: TimeDilation
+	
+	//@todo DanH: Allow disabling the joint at this level regardless of what the movement mode wants?
+
+	UPROPERTY(Transient)
+	TObjectPtr<UMoverPathedPhysicsLiaisonComponent> PathedPhysicsLiaison;
+
+protected:
+	friend class UPathedPhysicsDebugDrawComponent;
+
+	/** When true and the root component of the actor is a mesh, a duplicate mesh will be shown at PreviewMeshProgress along the starting/default path */
+	UPROPERTY(EditAnywhere, Category = "Pathed Movement|Debug")
+	bool bDisplayProgressPreviewMesh = true;
+
+	/** How far along the starting/default path to preview the controlled mesh */
+	UPROPERTY(EditAnywhere, Category = "Pathed Movement|Debug", meta = (EditCondition = bDisplayProgressPreviewMesh, UIMin = 0, ClampMin = 0, UIMax = 1, ClampMax = 1))
+	float PreviewMeshProgress = 1.f;
+
+	/** The material to apply to the preview mesh displayed along the path at PreviewMeshProgress (leave empty for an exact duplicate) */
+	UPROPERTY(EditAnywhere, Category = "Pathed Movement|Debug", meta = (EditCondition = bDisplayProgressPreviewMesh))
+	TObjectPtr<UMaterialInterface> ProgressPreviewMeshMaterial;
+	
+	UPROPERTY(Transient)
+	TObjectPtr<UPathedPhysicsDebugDrawComponent> DebugDrawComp;
+};
diff --git a/Engine/Plugins/Experimental/Mover/Source/Mover/Public/PhysicsMover/PathedMovement/PointMovementPathPattern.h b/Engine/Plugins/Experimental/Mover/Source/Mover/Public/PhysicsMover/PathedMovement/PointMovementPathPattern.h
new file mode 100644
--- /dev/null
+++ b/Engine/Plugins/Experimental/Mover/Source/Mover/Public/PhysicsMover/PathedMovement/PointMovementPathPattern.h
@@ -0,0 +1,85 @@
+// Copyright Epic Games, Inc. All Rights Reserved.
+
+#pragma once
+
+#include "CoreMinimal.h"
+#include "PathedMovementPatternBase.h"
+
+#include "PointMovementPathPattern.generated.h"
+
+UENUM()
+enum class EPointMovementLocationBasis : uint8
+{
+	/** The location is relative to the previous point in the path */
+	PreviousPoint,
+	/** The location is relative to the origin of the path */
+	PathOrigin,
+	/** The location is relative to the world origin */
+	World
+};
+
+USTRUCT(BlueprintType)
+struct FPointMovementPathPoint
+{
+	GENERATED_BODY()
+
+public:
+	FPointMovementPathPoint() {}
+
+	//@todo DanH: Same options for control as FSplinePathPatternPointData - might want a shared base
+	
+	/** Basis that the location of this point is relative to */
+	UPROPERTY(EditAnywhere, BlueprintReadOnly, Category = Location)
+	EPointMovementLocationBasis Basis = EPointMovementLocationBasis::PreviousPoint;
+	
+	/** The location to move toward, relative to the path origin (i.e. root component location by default) */
+	UPROPERTY(EditAnywhere, BlueprintReadOnly, Category = Location, meta = (MakeEditWidget = true))
+	FVector Location = FVector::Zero();
+
+	/**
+	 * The progress along the path that this point corresponds to.
+	 * Read-only, calculated automatically based on number of points, distances between them, and movement rules for the point 
+	 */
+	// UPROPERTY(VisibleAnywhere, BlueprintReadOnly, Transient, Category = Location)
+	mutable float Progress = 0.f;
+	
+	/** Distance from the start of the path to this point. Read-only. */
+	// UPROPERTY(VisibleAnywhere, BlueprintReadOnly, Transient, Category = Location)
+	mutable float DistanceFromStart = 0.f;
+
+	/** The world  */
+	mutable FVector WorldLoc = FVector::Zero();
+};
+
+/** Movement pattern that moves between explicitly defined points */
+UCLASS()
+class UPointMovementPathPattern : public UPathedMovementPatternBase
+{
+	GENERATED_BODY()
+
+public:
+	virtual void InitializePattern() override;
+
+#if WITH_EDITOR
+	virtual void PostEditChangeProperty(struct FPropertyChangedEvent& PropertyChangedEvent) override;
+#endif
+	
+	virtual bool DebugDrawUsingStepSamples() const override { return false; }
+	virtual void AppendDebugDrawElements(UPathedPhysicsDebugDrawComponent& DebugDrawComp, FBoxSphereBounds::Builder& InOutDebugBoundsBuilder) override;
+	
+protected:
+	virtual FTransform CalcUnmaskedTargetRelativeTransform(float PatternProgress, const FTransform& CurTargetTransform) const override;
+
+	//@todo DanH: We need to force a refresh whenever the path origin changes
+	void RefreshAssignedPointProgress(bool bForceRefresh = false) const;
+
+	/** Relative point locations to move toward, in sequence from first to last */
+	UPROPERTY(EditAnywhere, BlueprintReadWrite, Category = "Point Movement Pattern")
+	TArray<FPointMovementPathPoint> PathPoints;
+
+	/** The sum total distance of this path. Read-only for reference only. */
+	// UPROPERTY(VisibleAnywhere, BlueprintReadOnly, Transient, Category = "Point Movement Pattern")
+	mutable float TotalPathDistance = 0.f;
+
+	mutable bool bHasAssignedPointProgress = false;
+};
\ No newline at end of file
diff --git a/Engine/Plugins/Experimental/Mover/Source/Mover/Public/PhysicsMover/PathedMovement/SplineMovementPathPattern.h b/Engine/Plugins/Experimental/Mover/Source/Mover/Public/PhysicsMover/PathedMovement/SplineMovementPathPattern.h
new file mode 100644
--- /dev/null
+++ b/Engine/Plugins/Experimental/Mover/Source/Mover/Public/PhysicsMover/PathedMovement/SplineMovementPathPattern.h
@@ -0,0 +1,77 @@
+// Copyright Epic Games, Inc. All Rights Reserved.
+
+#pragma once
+
+#include "CoreMinimal.h"
+#include "PathedMovementPatternBase.h"
+#include "Engine/EngineTypes.h"
+
+#include "SplineMovementPathPattern.generated.h"
+
+class USplineComponent;
+
+/** This isn't formal "spline metadata" because that requires a custom spline and component - we want to add info to vanilla splines */
+USTRUCT()
+struct FSplinePathPatternPointData
+{
+	GENERATED_BODY()
+
+	/** Key of the spline point (or somewhere between two points) where this metadata applies */
+	UPROPERTY()
+	float SplinePointKey = 0.f;
+
+	//@todo DanH: Speed cap, pausing, per-section easing, (show a true readonly total duration that accounts for them too)
+};
+
+UCLASS()
+class USplineMovementPathPattern : public UPathedMovementPatternBase
+{
+	GENERATED_BODY()
+	
+public:
+	virtual void InitializePattern() override;
+
+	//@todo DanH: Allow patterns to state a desired duration and either complain if there are multiple or go with whatever is longest
+	//virtual TOptional<float> GetDesiredPathDuration() const override;
+
+	// The spline is already drawn, no need to draw it again at lower rez
+	virtual bool DebugDrawUsingStepSamples() const override { return false; }
+	
+	bool HasValidBounds() const { return UpperBound > LowerBound; }
+	
+protected:
+	virtual FTransform CalcUnmaskedTargetRelativeTransform(float PatternProgress, const FTransform& CurTargetTransform) const override;
+
+	/** How far into the spline the movement path actually begins.  */
+	UPROPERTY(EditAnywhere, Category = SplinePathPattern, meta = (UIMin = 0.f, ClampMin = 0.f, UIMax = 1.f, ClampMax = 1.f))
+	float LowerBound = 0.f;
+
+	/** How far into the spline the movement path ends. Must be greater than the lower bound. */
+	UPROPERTY(EditAnywhere, Category = SplinePathPattern, meta = (UIMin = 0.f, ClampMin = 0.f, UIMax = 1.f, ClampMax = 1.f))
+	float UpperBound = 1.f;
+
+	//@todo DanH: Expose this as an alternative to an explicit time - when true it effectively sets the optional desired time
+	/** If true, the duration property on the spline will be used as the path duration */
+	// UPROPERTY(EditAnywhere, Category = SplinePathPattern)
+	// bool bUseSplineDuration = false;
+
+	/**
+	 * If true, the path duration will be shortened according to how much of the spline is not being followed.
+	 * If false (default), the path duration is unchanged, so the object will move slower when the usable spline range is reduced.
+	 */
+	UPROPERTY(EditAnywhere, Category = SplinePathPattern)
+	bool bOffsetsModifyDuration = false;
+
+	UPROPERTY(EditAnywhere, Category = SplinePathPattern)
+	bool bApplySplineScaling = false;
+	
+	/**
+	 * Optional property to specify the spline component that defines the path to follow. If blank, we'll use the first spline component we find in this actor.
+	 * This is only necessary to set if you have multiple spline components on the actor, or want to follow a spline on an external actor.
+	 */
+	UPROPERTY(EditAnywhere, DisplayName = "Spline Component", Category = SplinePathPattern, meta = (AllowAnyActor, UseComponentPicker, AllowedClasses = "/Script/Engine.SplineComponent"))
+	FComponentReference SplineComponentRef;
+
+	UPROPERTY(Transient)
+	TObjectPtr<const USplineComponent> SplineComp;
+};
