[Chaos Mover] Fix point pattern's debug draw always drawing a line from the actor to the world origin

diff --git a/Engine/Plugins/Experimental/Mover/Source/Mover/Private/PhysicsMover/PathedMovement/PathedPhysicsDebugDrawComponent.cpp b/Engine/Plugins/Experimental/Mover/Source/Mover/Private/PhysicsMover/PathedMovement/PathedPhysicsDebugDrawComponent.cpp
--- a/Engine/Plugins/Experimental/Mover/Source/Mover/Private/PhysicsMover/PathedMovement/PathedPhysicsDebugDrawComponent.cpp
+++ b/Engine/Plugins/Experimental/Mover/Source/Mover/Private/PhysicsMover/PathedMovement/PathedPhysicsDebugDrawComponent.cpp
@@ -130,12 +130,7 @@ void UPathedPhysicsDebugDrawComponent::CreateRenderState_Concurrent(FRegisterCom
 			int32 NumCurvesToDraw = PathedMovementMode->bDebugDrawAggregatePath ? 1 : 0;
 			if (PathedMovementMode->bAllowPatternDebugDrawing)
 			{
-				PreviousStepPatternLocations.SetNumZeroed(ValidPatterns.Num());
-				PreviousStepPatternLocations[0] = PreviousStepAggregateLocation;
-				// for (FVector& PatternLoc : PreviousStepPatternLocations)
-				// {
-				// 	PatternLoc = PreviousStepAggregateLocation;
-				// }
+				PreviousStepPatternLocations.Init(PreviousStepAggregateLocation, ValidPatterns.Num());
 				NumCurvesToDraw += NumPatternsToDebugDraw;
 			}
 			
