[Chaos Mover] Remove erroneous #undef LOCTEXT_NAMESPACE to resolve Horde issue 736319

diff --git a/Engine/Plugins/Experimental/Mover/Source/Mover/Private/PhysicsMover/PathedMovement/PathedPhysicsDebugDrawComponent.cpp b/Engine/Plugins/Experimental/Mover/Source/Mover/Private/PhysicsMover/PathedMovement/PathedPhysicsDebugDrawComponent.cpp
--- a/Engine/Plugins/Experimental/Mover/Source/Mover/Private/PhysicsMover/PathedMovement/PathedPhysicsDebugDrawComponent.cpp
+++ b/Engine/Plugins/Experimental/Mover/Source/Mover/Private/PhysicsMover/PathedMovement/PathedPhysicsDebugDrawComponent.cpp
@@ -325,5 +325,3 @@ void UPathedPhysicsDebugDrawComponent::HandleObjectPropertyChanged(UObject* Obje
 	}
 }
 #endif
-
-#undef LOCTEXT_NAMESPACE
