[Mover] Pass the FindMode_Mutable<T>that doesn't search by name to the correct untemplated FindMode_Mutable() that - you guessed it - doesn't search by name. Previously all attempts to search only by type would fail by searching for an invalid name.

diff --git a/Engine/Plugins/Experimental/Mover/Source/Mover/Public/MoverComponent.h b/Engine/Plugins/Experimental/Mover/Source/Mover/Public/MoverComponent.h
--- a/Engine/Plugins/Experimental/Mover/Source/Mover/Public/MoverComponent.h
+++ b/Engine/Plugins/Experimental/Mover/Source/Mover/Public/MoverComponent.h
@@ -395,7 +395,7 @@ public:	// Queries
 
 	/** Find the first movement mode on this component with the given type, optionally of the given type exactly. Returns null if there is none of that type */
 	template<typename ModeT = UBaseMovementMode UE_REQUIRES(std::is_base_of_v<UBaseMovementMode, ModeT>)>
-	ModeT* FindMode_Mutable(bool bRequireExactClass = false) const { return Cast<ModeT>(FindMode_Mutable(ModeT::StaticClass(), NAME_None, bRequireExactClass)); }
+	ModeT* FindMode_Mutable(bool bRequireExactClass = false) const { return Cast<ModeT>(FindMode_Mutable(ModeT::StaticClass(), bRequireExactClass)); }
 	UBaseMovementMode* FindMode_Mutable(TSubclassOf<UBaseMovementMode> ModeType, bool bRequireExactClass = false) const;
 
 	/** Find the movement mode on this component the given name and type, optionally of the given type exactly. Returns null if there is no mode by that name, or if it's of a different type. */
