[Chaos Mover] Remove direct reliance on the private (but somehow still accessible?) internal Chaos::FPhysicsObject implementation API in the base liaison. Use the FPhysicsObjectInterface getters instead.

diff --git a/Engine/Plugins/Experimental/Mover/Source/Mover/Private/Backends/MoverNetworkPhysicsLiaisonBase.cpp b/Engine/Plugins/Experimental/Mover/Source/Mover/Private/Backends/MoverNetworkPhysicsLiaisonBase.cpp
--- a/Engine/Plugins/Experimental/Mover/Source/Mover/Private/Backends/MoverNetworkPhysicsLiaisonBase.cpp
+++ b/Engine/Plugins/Experimental/Mover/Source/Mover/Private/Backends/MoverNetworkPhysicsLiaisonBase.cpp
@@ -2,17 +2,18 @@
 
 #include "Backends/MoverNetworkPhysicsLiaisonBase.h"
 
-#include "Components/ShapeComponent.h"
-#include "Framework/Threading.h"
-#include "GameFramework/Pawn.h"
 #include "MovementModeStateMachine.h"
 #include "PBDRigidsSolver.h"
 #include "Chaos/PhysicsObjectInterface.h"
+#include "Chaos/PhysicsObjectInternalInterface.h"
+#include "Components/ShapeComponent.h"
+#include "Framework/Threading.h"
+#include "GameFramework/Pawn.h"
 #include "Physics/Experimental/PhysScene_Chaos.h"
 #include "PhysicsEngine/PhysicsObjectExternalInterface.h"
-#include "PhysicsMover/Modes/PhysicsDrivenFallingMode.h"
 #include "PhysicsMover/PhysicsMoverManager.h"
-#include "Runtime/Experimental/Chaos/Private/Chaos/PhysicsObjectInternal.h"
+#include "PhysicsMover/Modes/PhysicsDrivenFallingMode.h"
+#include "PhysicsProxy/SingleParticlePhysicsProxy.h"
 
 #if WITH_EDITOR
 #include "Misc/DataValidation.h"
@@ -794,7 +795,7 @@ Chaos::FPBDRigidParticle* UMoverNetworkPhysicsLiaisonComponentBase::GetControlle
 {
 	if (Chaos::FPhysicsObject* ControlledObject = GetControlledPhysicsObject())
 	{
-		return ControlledObject->GetParticle<Chaos::EThreadContext::External>()->CastToRigidParticle();
+		return FPhysicsObjectExternalInterface::LockRead(ControlledObject)->GetRigidParticle(ControlledObject);
 	}
 	return nullptr;
 }
@@ -803,7 +804,8 @@ Chaos::FPBDRigidParticleHandle* UMoverNetworkPhysicsLiaisonComponentBase::GetCon
 {
 	if (Chaos::FPhysicsObject* ControlledObject = GetControlledPhysicsObject())
 	{
-		return ControlledObject->GetParticle<Chaos::EThreadContext::Internal>()->CastToRigidParticle();
+		Chaos::FReadPhysicsObjectInterface_Internal Interface = Chaos::FPhysicsObjectInternalInterface::GetRead();
+		return Interface.GetRigidParticle(ControlledObject);
 	}
 	return nullptr;
 }
