Mover: Adding option for input production to be performed in parallel for the standalone backend

diff --git a/Engine/Plugins/Experimental/Mover/Source/Mover/Private/Backends/MoverStandaloneLiaison.cpp b/Engine/Plugins/Experimental/Mover/Source/Mover/Private/Backends/MoverStandaloneLiaison.cpp
--- a/Engine/Plugins/Experimental/Mover/Source/Mover/Private/Backends/MoverStandaloneLiaison.cpp
+++ b/Engine/Plugins/Experimental/Mover/Source/Mover/Private/Backends/MoverStandaloneLiaison.cpp
@@ -7,10 +7,22 @@
 #include "MoverComponent.h"
 
 namespace MoverStandaloneBackendCVars
-{ 
+{
+	static void OnRunProduceInputOnAnyThreadChanged(IConsoleVariable* Var)
+	{
+		// Force a refresh of all liaisons when the global cvar changes to reflect new settings
+		ForEachObjectOfClass(UMoverStandaloneLiaisonComponent::StaticClass(), [&](UObject* AsObj)
+		{
+			if (UMoverStandaloneLiaisonComponent* StandaloneLiaison = CastChecked<UMoverStandaloneLiaisonComponent>(AsObj, ECastCheckedType::NullAllowed))
+			{
+				StandaloneLiaison->SetUseAsyncProduceInput(StandaloneLiaison->GetUseAsyncProduceInput());
+			}
+		});
+	}
+	
 	static void OnRunMovementSimOnAnyThreadChanged(IConsoleVariable* Var)
 	{
-		// Force a refresh of all liaisons when the global 
+		// Force a refresh of all liaisons when the global cvar changes to reflect new settings
 		ForEachObjectOfClass(UMoverStandaloneLiaisonComponent::StaticClass(), [&](UObject* AsObj)
 		{
 			if (UMoverStandaloneLiaisonComponent* StandaloneLiaison = CastChecked<UMoverStandaloneLiaisonComponent>(AsObj, ECastCheckedType::NullAllowed))
@@ -20,6 +32,15 @@ namespace MoverStandaloneBackendCVars
 		});
 	}
 
+	// Whether to allow produce input calls on any thread
+	static int32 RunProduceInputOnAnyThread = 0;
+	FAutoConsoleVariableRef CVarMoverStandaloneProduceInputOnAnyThread(
+		TEXT("mover.standalone.RunProduceInputOnAnyThread"),
+		RunProduceInputOnAnyThread,
+		TEXT("Whether to allow produce input to run on any thread.\nIndividuals must also have their UseAsyncProduceInput option enabled.\n0: Game thread only, 1: Any thread"),
+		FConsoleVariableDelegate::CreateStatic(&OnRunProduceInputOnAnyThreadChanged),
+		ECVF_Default);
+	
 	// Whether to allow movement simulation ticks on any thread
 	static int32 RunMovementSimOnAnyThread = 0;
 	FAutoConsoleVariableRef CVarMoverStandaloneRunSimOnAnyThread(
@@ -41,6 +62,7 @@ UMoverStandaloneLiaisonComponent::UMoverStandaloneLiaisonComponent(const FObject
 	ProduceInputTickFunction.SetTickFunctionEnable(true);
 	ProduceInputTickFunction.TickGroup = TG_PrePhysics;
 	ProduceInputTickFunction.bAllowTickBatching = true;
+	ProduceInputTickFunction.bRunOnAnyThread = bUseAsyncProduceInput && MoverStandaloneBackendCVars::RunProduceInputOnAnyThread;
 
 	SimulateMovementTickFunction.bCanEverTick = true;
 	SimulateMovementTickFunction.bStartWithTickEnabled = true;
@@ -143,6 +165,17 @@ FTickFunction* UMoverStandaloneLiaisonComponent::FindTickFunction(EMoverTickPhas
 	return FoundFunction;
 }
 
+void UMoverStandaloneLiaisonComponent::SetUseAsyncProduceInput(bool bUseAsyncInputProduction)
+{
+	bUseAsyncProduceInput = bUseAsyncInputProduction;
+	ProduceInputTickFunction.bRunOnAnyThread = bUseAsyncProduceInput && MoverStandaloneBackendCVars::RunProduceInputOnAnyThread;
+}
+
+bool UMoverStandaloneLiaisonComponent::GetUseAsyncProduceInput() const
+{
+	return bUseAsyncProduceInput;
+}
+
 
 void UMoverStandaloneLiaisonComponent::SetUseAsyncMovementSimulationTick(bool bUseAsyncMovementSim)
 {
diff --git a/Engine/Plugins/Experimental/Mover/Source/Mover/Public/Backends/MoverStandaloneLiaison.h b/Engine/Plugins/Experimental/Mover/Source/Mover/Public/Backends/MoverStandaloneLiaison.h
--- a/Engine/Plugins/Experimental/Mover/Source/Mover/Public/Backends/MoverStandaloneLiaison.h
+++ b/Engine/Plugins/Experimental/Mover/Source/Mover/Public/Backends/MoverStandaloneLiaison.h
@@ -114,6 +114,14 @@ public:
 
 	FTickFunction* FindTickFunction(EMoverTickPhase MoverTickPhase);
 
+	
+	/** Sets whether this instance's produce input can run on worker threads or not. See @bUseAsyncProduceInput */
+	UFUNCTION(BlueprintCallable, Category = "Async Movement")
+	void SetUseAsyncProduceInput(bool bUseAsyncInputProduction);
+
+	UFUNCTION(BlueprintCallable, Category = "Async Movement")
+	bool GetUseAsyncProduceInput() const;
+	
 	/** Sets whether this instance's movement simulation tick can run on worker threads or not. See @bUseAsyncMovementSimulationTick */
 	UFUNCTION(BlueprintCallable, Category = "Async Movement")
 	void SetUseAsyncMovementSimulationTick(bool bUseAsyncMovementSim);
@@ -131,6 +139,13 @@ protected:
 	virtual void OnControllerChanged(APawn* Pawn, AController* OldController, AController* NewController);
 
 protected:
+	/**
+	 * Sets whether produce input can run on worker threads or not, also gated by global option.
+	 * Changes at runtime will take affect next frame. Has no effect on simulation ticking or applying results.
+	 */
+	UPROPERTY(EditDefaultsOnly, Category = "Async Movement")
+	bool bUseAsyncProduceInput = false;
+	
 	/**
 	 * Sets whether the movement simulation tick can run on worker threads or not, also gated by global option.
 	 * Changes at runtime will take affect next frame. Has no effect on input production or applying results.
