Mover: Adding display name to instant effects so they are easier to find/make/break in BP

diff --git a/Engine/Plugins/Experimental/Mover/Source/Mover/Public/DefaultMovementSet/InstantMovementEffects/BasicInstantMovementEffects.h b/Engine/Plugins/Experimental/Mover/Source/Mover/Public/DefaultMovementSet/InstantMovementEffects/BasicInstantMovementEffects.h
--- a/Engine/Plugins/Experimental/Mover/Source/Mover/Public/DefaultMovementSet/InstantMovementEffects/BasicInstantMovementEffects.h
+++ b/Engine/Plugins/Experimental/Mover/Source/Mover/Public/DefaultMovementSet/InstantMovementEffects/BasicInstantMovementEffects.h
@@ -9,7 +9,7 @@
 /** Teleport: instantly moves an actor to a new location and rotation
  * For async-compatible teleportation, use AsyncTeleportEffect instead.
  */
-USTRUCT(BlueprintType)
+USTRUCT(BlueprintType, Category = "Mover|Instant Movement Effect", DisplayName = "Teleport Instant Movement Effect")
 struct MOVER_API FTeleportEffect : public FInstantMovementEffect
 {
 	GENERATED_BODY()
@@ -43,7 +43,7 @@ struct MOVER_API FTeleportEffect : public FInstantMovementEffect
 };
 
 /** Async Teleport: instantly moves an actor to a new location and rotation (compatible with async movement simulation) */
-USTRUCT(BlueprintType)
+USTRUCT(BlueprintType, Category = "Mover|Instant Movement Effect", DisplayName = "Async Teleport Instant Movement Effect")
 struct MOVER_API FAsyncTeleportEffect : public FTeleportEffect
 {
 	GENERATED_BODY()
@@ -60,7 +60,7 @@ struct MOVER_API FAsyncTeleportEffect : public FTeleportEffect
 /** Jump Impulse: introduces an instantaneous upwards change in velocity. This overrides the existing 'up' component of the actor's current velocity
   * Note: this only applies the impulse for one tick!
   */
-USTRUCT(BlueprintType)
+USTRUCT(BlueprintType, Category = "Mover|Instant Movement Effect", DisplayName = "Jump Impulse Instant Movement Effect")
 struct MOVER_API FJumpImpulseEffect : public FInstantMovementEffect
 {
 	GENERATED_BODY()
@@ -89,7 +89,7 @@ struct MOVER_API FJumpImpulseEffect : public FInstantMovementEffect
 /** Apply Velocity: provides an impulse velocity to the actor after (optionally) forcing them into a particular movement mode
   * Note: this only applies the impulse for one tick!
   */
-USTRUCT(BlueprintType)
+USTRUCT(BlueprintType, Category = "Mover|Instant Movement Effect", DisplayName = "Apply Velocity Instant Movement Effect")
 struct MOVER_API FApplyVelocityEffect : public FInstantMovementEffect
 {
 	GENERATED_BODY()
