Mover: Fixing performance concern where nav walking modes perform duplicate GetNavData queries

diff --git a/Engine/Plugins/Experimental/Mover/Source/Mover/Private/DefaultMovementSet/Modes/AsyncNavWalkingMode.cpp b/Engine/Plugins/Experimental/Mover/Source/Mover/Private/DefaultMovementSet/Modes/AsyncNavWalkingMode.cpp
--- a/Engine/Plugins/Experimental/Mover/Source/Mover/Private/DefaultMovementSet/Modes/AsyncNavWalkingMode.cpp
+++ b/Engine/Plugins/Experimental/Mover/Source/Mover/Private/DefaultMovementSet/Modes/AsyncNavWalkingMode.cpp
@@ -175,6 +175,7 @@ void UAsyncNavWalkingMode::OnSimulationTick(const FSimulationTickParams& Params,
 	float SimpleHalfHeight = 0;
 	NavMoverComponent->GetSimpleCollisionCylinder(OUT SimpleRadius, OUT SimpleHalfHeight);
 
+	const INavigationDataInterface* NavData = nullptr;
 	bool bSameNavLocation = false;
 	if (CachedNavLocation.NodeRef != INVALID_NAVNODEREF)
 	{
@@ -196,7 +197,9 @@ void UAsyncNavWalkingMode::OnSimulationTick(const FSimulationTickParams& Params,
 
 		if (bDeltaMoveNearlyZero && bSameNavLocation)
 		{
-			if (const INavigationDataInterface* NavData = GetNavData())
+			NavData = GetNavData();
+			
+			if (NavData)
 			{
 				if (!NavData->IsNodeRefValid(CachedNavLocation.NodeRef))
 				{
@@ -222,8 +225,13 @@ void UAsyncNavWalkingMode::OnSimulationTick(const FSimulationTickParams& Params,
 			UMovementUtils::SetGravityVerticalComponent(TargetFeetLocation, CachedNavLocation.Location.Dot(UpDirection), UpDirection);
 		}
 
+		if (!NavData)
+		{
+			NavData = GetNavData();
+		}
+		
 		// Find the point on the NavMesh
-		const bool bHasNavigationData = FindNavFloor(TargetFeetLocation, DestNavLocation);
+		const bool bHasNavigationData = FindNavFloor(TargetFeetLocation, DestNavLocation, NavData);
 		if (!bHasNavigationData)
 		{
 			// Can't find nav mesh at this location, so fall back to normal walking mode 
@@ -272,9 +280,8 @@ void UAsyncNavWalkingMode::OnSimulationTick(const FSimulationTickParams& Params,
 	CaptureFinalState(LocationInProgress, RotationInProgress.Rotator(), /*bDidAttemptMovement=*/true, CurrentFloor, MoveRecord, OutputSyncState);
 }
 
-bool UAsyncNavWalkingMode::FindNavFloor(const FVector& TestLocation, FNavLocation& NavFloorLocation) const
+bool UAsyncNavWalkingMode::FindNavFloor(const FVector& TestLocation, FNavLocation& OutNavFloorLocation, const INavigationDataInterface* NavData) const
 {
-	const INavigationDataInterface* NavData = GetNavData();
 	if (NavData == nullptr || !NavMoverComponent.IsValid())
 	{
 		return false;
@@ -284,7 +291,7 @@ bool UAsyncNavWalkingMode::FindNavFloor(const FVector& TestLocation, FNavLocatio
 	const float SearchRadius = AgentProps.AgentRadius * 2.0f;
 	const float SearchHeight = AgentProps.AgentHeight * AgentProps.NavWalkingSearchHeightScale;
 
-	return NavData->ProjectPoint(TestLocation, NavFloorLocation, FVector(SearchRadius, SearchRadius, SearchHeight));
+	return NavData->ProjectPoint(TestLocation, OutNavFloorLocation, FVector(SearchRadius, SearchRadius, SearchHeight));
 }
 
 UObject* UAsyncNavWalkingMode::GetTurnGenerator()
diff --git a/Engine/Plugins/Experimental/Mover/Source/Mover/Private/DefaultMovementSet/Modes/NavWalkingMode.cpp b/Engine/Plugins/Experimental/Mover/Source/Mover/Private/DefaultMovementSet/Modes/NavWalkingMode.cpp
--- a/Engine/Plugins/Experimental/Mover/Source/Mover/Private/DefaultMovementSet/Modes/NavWalkingMode.cpp
+++ b/Engine/Plugins/Experimental/Mover/Source/Mover/Private/DefaultMovementSet/Modes/NavWalkingMode.cpp
@@ -170,6 +170,7 @@ void UNavWalkingMode::OnSimulationTick(const FSimulationTickParams& Params, FMov
 	float SimpleHalfHeight = 0;
 	NavMoverComponent->GetSimpleCollisionCylinder(SimpleRadius, SimpleHalfHeight);
 	
+	const INavigationDataInterface* NavData = nullptr;
 	bool bSameNavLocation = false;
 	if (CachedNavLocation.NodeRef != INVALID_NAVNODEREF)
 	{
@@ -191,7 +192,9 @@ void UNavWalkingMode::OnSimulationTick(const FSimulationTickParams& Params, FMov
 
 		if (bDeltaMoveNearlyZero && bSameNavLocation)
 		{
-			if (const INavigationDataInterface* NavData = GetNavData())
+			NavData = GetNavData();
+			
+			if (NavData)
 			{
 				if (!NavData->IsNodeRefValid(CachedNavLocation.NodeRef))
 				{
@@ -217,8 +220,13 @@ void UNavWalkingMode::OnSimulationTick(const FSimulationTickParams& Params, FMov
 			UMovementUtils::SetGravityVerticalComponent(AdjustedDest, CachedNavLocation.Location.Dot(UpDirection), UpDirection);
 		}
 
+		if (!NavData)
+		{
+			NavData = GetNavData();
+		}
+		
 		// Find the point on the NavMesh
-		const bool bHasNavigationData = FindNavFloor(AdjustedDest, DestNavLocation);
+		const bool bHasNavigationData = FindNavFloor(AdjustedDest, DestNavLocation, NavData);
 		if (!bHasNavigationData)
 		{
 			// Can't find nav mesh, so fall back to normal walking mode and let it process the intended movement
@@ -263,9 +271,8 @@ void UNavWalkingMode::OnSimulationTick(const FSimulationTickParams& Params, FMov
 	CaptureFinalState(UpdatedComponent, true, CurrentFloor, MoveRecord, OutputSyncState);
 }
 
-bool UNavWalkingMode::FindNavFloor(const FVector& TestLocation, FNavLocation& NavFloorLocation) const
+bool UNavWalkingMode::FindNavFloor(const FVector& TestLocation, FNavLocation& OutNavFloorLocation, const INavigationDataInterface* NavData) const
 {
-	const INavigationDataInterface* NavData = GetNavData();
 	if (NavData == nullptr || NavMoverComponent == nullptr)
 	{
 		return false;
@@ -275,7 +282,7 @@ bool UNavWalkingMode::FindNavFloor(const FVector& TestLocation, FNavLocation& Na
 	const float SearchRadius = AgentProps.AgentRadius * 2.0f;
 	const float SearchHeight = AgentProps.AgentHeight * AgentProps.NavWalkingSearchHeightScale;
 
-	return NavData->ProjectPoint(TestLocation, NavFloorLocation, FVector(SearchRadius, SearchRadius, SearchHeight));
+	return NavData->ProjectPoint(TestLocation, OutNavFloorLocation, FVector(SearchRadius, SearchRadius, SearchHeight));
 }
 
 UObject* UNavWalkingMode::GetTurnGenerator()
diff --git a/Engine/Plugins/Experimental/Mover/Source/Mover/Public/DefaultMovementSet/Modes/AsyncNavWalkingMode.h b/Engine/Plugins/Experimental/Mover/Source/Mover/Public/DefaultMovementSet/Modes/AsyncNavWalkingMode.h
--- a/Engine/Plugins/Experimental/Mover/Source/Mover/Public/DefaultMovementSet/Modes/AsyncNavWalkingMode.h
+++ b/Engine/Plugins/Experimental/Mover/Source/Mover/Public/DefaultMovementSet/Modes/AsyncNavWalkingMode.h
@@ -73,11 +73,12 @@ public:
 
 	/**
 	 * Project a location to navmesh to find adjusted height.
-	 * @param TestLocation		Location to project
-	 * @param NavFloorLocation	Location on navmesh
+	 * @param TestLocation			Location to project
+	 * @param OutNavFloorLocation	Location on navmesh
+	 * @param NavData				NavigationDataInterface to search TestLocation for
 	 * @return True if projection was performed (successfully or not)
 	 */
-	virtual bool FindNavFloor(const FVector& TestLocation, FNavLocation& NavFloorLocation) const;
+	virtual bool FindNavFloor(const FVector& TestLocation, FNavLocation& OutNavFloorLocation, const INavigationDataInterface* NavData) const;
 
 	// Returns the active turn generator. Note: you will need to cast the return value to the generator you expect to get, it can also be none
 	UFUNCTION(BlueprintPure, Category = Mover)
diff --git a/Engine/Plugins/Experimental/Mover/Source/Mover/Public/DefaultMovementSet/Modes/NavWalkingMode.h b/Engine/Plugins/Experimental/Mover/Source/Mover/Public/DefaultMovementSet/Modes/NavWalkingMode.h
--- a/Engine/Plugins/Experimental/Mover/Source/Mover/Public/DefaultMovementSet/Modes/NavWalkingMode.h
+++ b/Engine/Plugins/Experimental/Mover/Source/Mover/Public/DefaultMovementSet/Modes/NavWalkingMode.h
@@ -77,11 +77,12 @@ public:
 	
 	/**
 	 * Project a location to navmesh to find adjusted height.
-	 * @param TestLocation		Location to project
-	 * @param NavFloorLocation	Location on navmesh
+	 * @param TestLocation			Location to project
+	 * @param OutNavFloorLocation	Location on navmesh
+	 * @param NavData				NavigationDataInterface to search TestLocation for
 	 * @return True if projection was performed (successfully or not)
 	 */
-	virtual bool FindNavFloor(const FVector& TestLocation, FNavLocation& NavFloorLocation) const;
+	virtual bool FindNavFloor(const FVector& TestLocation, FNavLocation& OutNavFloorLocation, const INavigationDataInterface* NavData) const;
 
 	// Returns the active turn generator. Note: you will need to cast the return value to the generator you expect to get, it can also be none
 	UFUNCTION(BlueprintPure, Category=Mover)
