Mover: support for async teleport / testing for encroachment at proposed locations

diff --git a/Engine/Plugins/Experimental/Mover/Source/Mover/Private/DefaultMovementSet/InstantMovementEffects/BasicInstantMovementEffects.cpp b/Engine/Plugins/Experimental/Mover/Source/Mover/Private/DefaultMovementSet/InstantMovementEffects/BasicInstantMovementEffects.cpp
--- a/Engine/Plugins/Experimental/Mover/Source/Mover/Private/DefaultMovementSet/InstantMovementEffects/BasicInstantMovementEffects.cpp
+++ b/Engine/Plugins/Experimental/Mover/Source/Mover/Private/DefaultMovementSet/InstantMovementEffects/BasicInstantMovementEffects.cpp
@@ -9,12 +9,22 @@
 #include "DefaultMovementSet/Settings/CommonLegacyMovementSettings.h"
 #include "MoveLibrary/MovementUtils.h"
 #include "MoveLibrary/MoverBlackboard.h"
-
+#include "DrawDebugHelpers.h"
 
 // -------------------------------------------------------------------
 // FTeleportEffect
 // -------------------------------------------------------------------
 
+static int32 ShowTeleportDiffs = 0;
+static float ShowTeleportDiffsLifetimeSecs = 3.0f;
+FAutoConsoleVariableRef CVarShowTeleportDiffs(
+	TEXT("mover.debug.ShowTeleportDiffs"),
+	ShowTeleportDiffs,
+	TEXT("Whether to draw teleportation differences (red is initially blocked, green is corrected).\n")
+	TEXT("0: Disable, 1: Enable"),
+	ECVF_Cheat);
+
+
 FTeleportEffect::FTeleportEffect()
 	: TargetLocation(FVector::ZeroVector)
 	, bUseActorRotation(true)
@@ -28,6 +38,18 @@ bool FTeleportEffect::ApplyMovementEffect(FApplyMovementEffectParams& ApplyEffec
 	
 	if (ApplyEffectParams.UpdatedComponent->GetOwner()->TeleportTo(TargetLocation, Rotation))
 	{
+		if (ShowTeleportDiffs)
+		{
+			const AActor* OwnerActor = ApplyEffectParams.UpdatedComponent->GetOwner();
+
+			if (!(ApplyEffectParams.UpdatedComponent->GetComponentLocation() - TargetLocation).IsNearlyZero())	// if it was adjusted, show the original error
+			{
+				DrawDebugCapsule(OwnerActor->GetWorld(), TargetLocation, OwnerActor->GetSimpleCollisionHalfHeight(), OwnerActor->GetSimpleCollisionRadius(), FQuat::Identity, FColor::Red, false, ShowTeleportDiffsLifetimeSecs);
+			}
+
+			DrawDebugCapsule(OwnerActor->GetWorld(), ApplyEffectParams.UpdatedComponent->GetComponentLocation(), OwnerActor->GetSimpleCollisionHalfHeight(), OwnerActor->GetSimpleCollisionRadius(), FQuat::Identity, FColor(100, 100, 255), false, ShowTeleportDiffsLifetimeSecs);
+		}
+
 		FMoverDefaultSyncState& OutputSyncState = OutputState.SyncStateCollection.FindOrAddMutableDataByType<FMoverDefaultSyncState>();
 		
 		if (const FMoverDefaultSyncState* StartingSyncState = ApplyEffectParams.StartState->SyncState.SyncStateCollection.FindDataByType<FMoverDefaultSyncState>())
@@ -48,6 +70,12 @@ bool FTeleportEffect::ApplyMovementEffect(FApplyMovementEffectParams& ApplyEffec
 		}
 	}
 
+	if (ShowTeleportDiffs)
+	{
+		const AActor* OwnerActor = ApplyEffectParams.UpdatedComponent->GetOwner();
+		DrawDebugCapsule(OwnerActor->GetWorld(), TargetLocation, OwnerActor->GetSimpleCollisionHalfHeight(), OwnerActor->GetSimpleCollisionRadius(), FQuat::Identity, FColor::Red, false, ShowTeleportDiffsLifetimeSecs);
+	}
+
 	return false;
 }
 
@@ -95,25 +123,45 @@ bool FAsyncTeleportEffect::ApplyMovementEffect(FApplyMovementEffectParams& Apply
 	FVector TeleportLocation = TargetLocation;
 	const FRotator TeleportRotation = bUseActorRotation ? ApplyEffectParams.UpdatedComponent->GetComponentRotation() : TargetRotation;
 
-	UMovementUtils::FindTeleportSpot(ApplyEffectParams.MoverComp, OUT TeleportLocation, TeleportRotation);
+	if (UMovementUtils::FindTeleportSpot(ApplyEffectParams.MoverComp, OUT TeleportLocation, TeleportRotation))
+	{
+		if (ShowTeleportDiffs)
+		{
+			const AActor* OwnerActor = ApplyEffectParams.UpdatedComponent->GetOwner();
 
-	FMoverDefaultSyncState& OutputSyncState = OutputState.SyncStateCollection.FindOrAddMutableDataByType<FMoverDefaultSyncState>();
+			if (!(TeleportLocation - TargetLocation).IsNearlyZero())	// if it was adjusted, show the original error
+			{
+				DrawDebugCapsule(OwnerActor->GetWorld(), TargetLocation, OwnerActor->GetSimpleCollisionHalfHeight(), OwnerActor->GetSimpleCollisionRadius(), FQuat::Identity, FColor::Red, false, ShowTeleportDiffsLifetimeSecs);
+			}
 
-	if (const FMoverDefaultSyncState* StartingSyncState = ApplyEffectParams.StartState->SyncState.SyncStateCollection.FindDataByType<FMoverDefaultSyncState>())
-	{
-		OutputSyncState.SetTransforms_WorldSpace(TeleportLocation,
-			TeleportRotation,
-			OutputSyncState.GetVelocity_WorldSpace(),
-			nullptr); // no movement base
+			DrawDebugCapsule(OwnerActor->GetWorld(), TeleportLocation, OwnerActor->GetSimpleCollisionHalfHeight(), OwnerActor->GetSimpleCollisionRadius(), FQuat::Identity, FColor(100, 100, 255), false, ShowTeleportDiffsLifetimeSecs);
+		}
 
-		// TODO: instead of invalidating it, consider checking for a floor. Possibly a dynamic base?
-		if (UMoverBlackboard* SimBlackboard = ApplyEffectParams.MoverComp->GetSimBlackboard_Mutable())
+		FMoverDefaultSyncState& OutputSyncState = OutputState.SyncStateCollection.FindOrAddMutableDataByType<FMoverDefaultSyncState>();
+
+		if (const FMoverDefaultSyncState* StartingSyncState = ApplyEffectParams.StartState->SyncState.SyncStateCollection.FindDataByType<FMoverDefaultSyncState>())
 		{
-			SimBlackboard->Invalidate(CommonBlackboard::LastFloorResult);
-			SimBlackboard->Invalidate(CommonBlackboard::LastFoundDynamicMovementBase);
+			OutputSyncState.SetTransforms_WorldSpace(TeleportLocation,
+				TeleportRotation,
+				OutputSyncState.GetVelocity_WorldSpace(),
+				nullptr); // no movement base
+
+			// TODO: instead of invalidating it, consider checking for a floor. Possibly a dynamic base?
+			if (UMoverBlackboard* SimBlackboard = ApplyEffectParams.MoverComp->GetSimBlackboard_Mutable())
+			{
+				SimBlackboard->Invalidate(CommonBlackboard::LastFloorResult);
+				SimBlackboard->Invalidate(CommonBlackboard::LastFoundDynamicMovementBase);
+			}
+
+			return true;
 		}
+	}
 
-		return true;
+
+	if (ShowTeleportDiffs)
+	{
+		const AActor* OwnerActor = ApplyEffectParams.UpdatedComponent->GetOwner();
+		DrawDebugCapsule(OwnerActor->GetWorld(), TargetLocation, OwnerActor->GetSimpleCollisionHalfHeight(), OwnerActor->GetSimpleCollisionRadius(), FQuat::Identity, FColor::Red, false, ShowTeleportDiffsLifetimeSecs);
 	}
 
 	return false;
diff --git a/Engine/Plugins/Experimental/Mover/Source/Mover/Private/MoveLibrary/MovementUtils.cpp b/Engine/Plugins/Experimental/Mover/Source/Mover/Private/MoveLibrary/MovementUtils.cpp
--- a/Engine/Plugins/Experimental/Mover/Source/Mover/Private/MoveLibrary/MovementUtils.cpp
+++ b/Engine/Plugins/Experimental/Mover/Source/Mover/Private/MoveLibrary/MovementUtils.cpp
@@ -575,29 +575,32 @@ bool UMovementUtils::FindTeleportSpot(const UMoverComponent* MoverComp, FVector&
 	{
 		return true;
 	}
-	FVector Adjust(0.f);
+	FVector ProposedTeleportAdjustment(0.f);
 
 	const FVector OriginalTestLocation = TestLocation;
 
-	// check if fits at desired location
-	if (!TestEncroachment(MoverComp, TestLocation, TestRotation, &Adjust))
+	// check if it fits at desired location
+	if (!TestEncroachmentAndAdjust(MoverComp, TestLocation, TestRotation, OUT ProposedTeleportAdjustment))
 	{
-		return true;
+		return true;	// it fits, so we're done
 	}
 
-	if (Adjust.IsNearlyZero())
+	if (ProposedTeleportAdjustment.IsNearlyZero())
 	{
+		// Doesn't fit and we didn't find an adjustment
 		// Reset in case Adjust is not actually zero
 		TestLocation = OriginalTestLocation;
 		return false;
 	}
 
+	// Feel around for a non-encroaching location
+	// 
 	// first do only Z
 	const FVector::FReal ZeroThreshold = UE_KINDA_SMALL_NUMBER;
-	const bool bZeroZ = FMath::IsNearlyZero(Adjust.Z, ZeroThreshold);
+	const bool bZeroZ = FMath::IsNearlyZero(ProposedTeleportAdjustment.Z, ZeroThreshold);
 	if (!bZeroZ)
 	{
-		TestLocation.Z += Adjust.Z;
+		TestLocation.Z += ProposedTeleportAdjustment.Z;
 		if (!TestEncroachment(MoverComp, TestLocation, TestRotation))
 		{
 			return true;
@@ -607,12 +610,12 @@ bool UMovementUtils::FindTeleportSpot(const UMoverComponent* MoverComp, FVector&
 	}
 
 	// now try just XY
-	const bool bZeroX = FMath::IsNearlyZero(Adjust.X, ZeroThreshold);
-	const bool bZeroY = FMath::IsNearlyZero(Adjust.Y, ZeroThreshold);
+	const bool bZeroX = FMath::IsNearlyZero(ProposedTeleportAdjustment.X, ZeroThreshold);
+	const bool bZeroY = FMath::IsNearlyZero(ProposedTeleportAdjustment.Y, ZeroThreshold);
 	if (!bZeroX || !bZeroY)
 	{
-		const float X = bZeroX ? 0.f : Adjust.X;
-		const float Y = bZeroY ? 0.f : Adjust.Y;
+		const float X = bZeroX ? 0.f : ProposedTeleportAdjustment.X;
+		const float Y = bZeroY ? 0.f : ProposedTeleportAdjustment.Y;
 		FVector Adjustments[8];
 		Adjustments[0] = FVector(X, Y, 0);
 
@@ -662,7 +665,7 @@ bool UMovementUtils::FindTeleportSpot(const UMoverComponent* MoverComp, FVector&
 			for (int i = 0; i < Iterations; ++i)
 			{
 				TestLocation = OriginalTestLocation + Adjustments[i];
-				TestLocation.Z += Adjust.Z;
+				TestLocation.Z += ProposedTeleportAdjustment.Z;
 				if (!TestEncroachment(MoverComp, TestLocation, TestRotation))
 				{
 					return true;
@@ -676,7 +679,25 @@ bool UMovementUtils::FindTeleportSpot(const UMoverComponent* MoverComp, FVector&
 	return false;
 }
 
-bool UMovementUtils::TestEncroachment(const UMoverComponent* MoverComp, FVector TestLocation, FRotator TestRotation, OUT FVector* OutProposedAdjustment)
+bool UMovementUtils::TestEncroachment(const UMoverComponent* MoverComp, FVector TestLocation, FRotator TestRotation)
+{
+	const AActor* OwnerActor = MoverComp->GetOwner();
+	const UWorld* OwnerWorld = OwnerActor->GetWorld();
+	const UPrimitiveComponent* MovingPrimitiveRoot = Cast<UPrimitiveComponent>(MoverComp->GetUpdatedComponent());
+	if (MovingPrimitiveRoot && MovingPrimitiveRoot->IsQueryCollisionEnabled())
+	{ 
+		FTransform const TestRootToWorld = FTransform(TestRotation, TestLocation);
+
+		TArray<AActor*> ChildActors;
+		OwnerActor->GetAllChildActors(ChildActors);
+
+		return TestEncroachment_Internal(OwnerWorld, OwnerActor, MovingPrimitiveRoot, TestRootToWorld, ChildActors);
+	}
+
+	return false;
+}
+
+bool UMovementUtils::TestEncroachmentAndAdjust(const UMoverComponent* MoverComp, FVector TestLocation, FRotator TestRotation, OUT FVector& OutProposedAdjustment)
 {
 	if (MoverComp == nullptr)
 	{
@@ -691,7 +712,6 @@ bool UMovementUtils::TestEncroachment(const UMoverComponent* MoverComp, FVector
 
 	bool bFoundEncroacher = false;
 
-	FVector TotalAdjustment(0.f);
 	FTransform const TestRootToWorld = FTransform(TestRotation, TestLocation);
 	FTransform const WorldToOldRoot = RootComponent->GetComponentToWorld().Inverse();
 
@@ -711,36 +731,14 @@ bool UMovementUtils::TestEncroachment(const UMoverComponent* MoverComp, FVector
 			TArray<AActor*> ChildActors;
 			MoverComp->GetOwner()->GetAllChildActors(ChildActors);
 
-			FVector FoundAdjustment;
-
-			if (TestEncroachmentWithAdjustment_Internal(MoverComp->GetOwner()->GetWorld(), MoverComp->GetOwner(), MovingPrimitiveRoot, CompToNewWorld, FoundAdjustment, ChildActors))
-			{
-				if (OutProposedAdjustment == nullptr)
-				{
-					// don't need an adjustment and we know we are overlapping, so we can be done
-					return true;
-				}
-				else
-				{
-					TotalAdjustment = FoundAdjustment;
-				}
-
-				bFoundEncroacher = true;
-			}
+			bFoundEncroacher = TestEncroachmentWithAdjustment_Internal(MoverComp->GetOwner()->GetWorld(), MoverComp->GetOwner(), MovingPrimitiveRoot, CompToNewWorld, ChildActors, OutProposedAdjustment);
 		}
 	}
 
-	// copy over total adjustment
-	if (OutProposedAdjustment)
-	{
-		*OutProposedAdjustment = TotalAdjustment;
-	}
-
 	return bFoundEncroacher;
 }
 
 
-
 bool UMovementUtils::TryMoveUpdatedComponent_Internal(const FMovingComponentSet& MovingComps, const FVector& Delta, const FQuat& NewRotation, bool bSweep, EMoveComponentFlags MoveComponentFlags, FHitResult* OutHit, ETeleportType Teleport)
 {
 	if (MovingComps.UpdatedComponent.IsValid())
@@ -763,8 +761,55 @@ bool UMovementUtils::TryMoveUpdatedComponent_Internal(const FMovingComponentSet&
 // How much to adjust out collision shape box during encroachment testing, for a little leeway. This is applied to each axis, in cm.
 static const float ENCROACH_SHRINK_EPSILON = 0.15f;
 
-bool UMovementUtils::TestEncroachmentWithAdjustment_Internal(const UWorld* World, const AActor* TestActor, const UPrimitiveComponent* PrimComp, const FTransform& TestWorldTransform, FVector& OutProposedAdjustment, const TArray<AActor*>& IgnoreActors)
+bool UMovementUtils::TestEncroachment_Internal(const UWorld* World, const AActor* TestActor, const UPrimitiveComponent* PrimComp, const FTransform& TestWorldTransform, const TArray<AActor*>& IgnoreActors)
+{
+	const float Epsilon = ENCROACH_SHRINK_EPSILON;
+
+	if (World && PrimComp)
+	{
+		bool bFoundBlockingHit = false;
+
+		ECollisionChannel const BlockingChannel = PrimComp->GetCollisionObjectType();
+		FCollisionShape const CollisionShape = PrimComp->GetCollisionShape(-Epsilon);
+
+		if (CollisionShape.IsBox() && (Cast<UBoxComponent>(PrimComp) == nullptr))
+		{
+			// we have a bounding box not for a box component, which means this was the fallback aabb
+			// since we don't need the penetration info, go ahead and test the component itself for overlaps, which is more accurate
+			if (PrimComp->IsRegistered())
+			{
+				// must be registered
+				TArray<FOverlapResult> Overlaps;
+				FComponentQueryParams Params(SCENE_QUERY_STAT(ComponentEncroachesBlockingGeometry_NoAdjustment), TestActor);
+				FCollisionResponseParams ResponseParams;
+				PrimComp->InitSweepCollisionParams(Params, ResponseParams);
+				Params.AddIgnoredActors(IgnoreActors);
+				return World->ComponentOverlapMultiByChannel(Overlaps, PrimComp, TestWorldTransform.GetLocation(), TestWorldTransform.GetRotation(), BlockingChannel, Params);
+			}
+			else
+			{
+				UE_LOG(LogMover, Log, TEXT("Components must be registered in order to be used in a ComponentOverlapMulti call. PriComp: %s TestActor: %s"), *PrimComp->GetName(), *TestActor->GetName());
+				return false;
+			}
+		}
+		else
+		{
+			FCollisionQueryParams Params(SCENE_QUERY_STAT(ComponentEncroachesBlockingGeometry_NoAdjustment), false, TestActor);
+			FCollisionResponseParams ResponseParams;
+			PrimComp->InitSweepCollisionParams(Params, ResponseParams);
+			Params.AddIgnoredActors(IgnoreActors);
+			return World->OverlapBlockingTestByChannel(TestWorldTransform.GetLocation(), TestWorldTransform.GetRotation(), BlockingChannel, CollisionShape, Params, ResponseParams);
+		}
+	}
+
+	return false;
+
+}
+
+bool UMovementUtils::TestEncroachmentWithAdjustment_Internal(const UWorld* World, const AActor* TestActor, const UPrimitiveComponent* PrimComp, const FTransform& TestWorldTransform, const TArray<AActor*>& IgnoreActors, OUT FVector& OutProposedAdjustment)
 {
+	// this function is based on UWorld's ComponentEncroachesBlockingGeometry
+
 	// init our output
 	OutProposedAdjustment = FVector::ZeroVector;
 
@@ -787,7 +832,7 @@ bool UMovementUtils::TestEncroachmentWithAdjustment_Internal(const UWorld* World
 			if (PrimComp->IsRegistered())
 			{
 				// must be registered
-				FComponentQueryParams Params(SCENE_QUERY_STAT(ComponentEncroachesBlockingGeometry_WithAdjustment), TestActor);
+				FComponentQueryParams Params(SCENE_QUERY_STAT(TestEncroachmentWithAdjustment), TestActor);
 				FCollisionResponseParams ResponseParams;
 				PrimComp->InitSweepCollisionParams(Params, ResponseParams);
 				Params.AddIgnoredActors(IgnoreActors);
@@ -796,13 +841,13 @@ bool UMovementUtils::TestEncroachmentWithAdjustment_Internal(const UWorld* World
 			}
 			else
 			{
-				UE_LOG(LogPhysics, Log, TEXT("Components must be registered in order to be used in a ComponentOverlapMulti call. PriComp: %s TestActor: %s"), *PrimComp->GetName(), *TestActor->GetName());
+				UE_LOG(LogMover, Log, TEXT("Components must be registered in order to be used in a ComponentOverlapMulti call. PriComp: %s TestActor: %s"), *PrimComp->GetName(), *TestActor->GetName());
 			}
 		}
 		else
 		{
 			// overlap our shape
-			FCollisionQueryParams Params(SCENE_QUERY_STAT(ComponentEncroachesBlockingGeometry_WithAdjustment), false, TestActor);
+			FCollisionQueryParams Params(SCENE_QUERY_STAT(TestEncroachmentWithAdjustment), false, TestActor);
 			FCollisionResponseParams ResponseParams;
 			PrimComp->InitSweepCollisionParams(Params, ResponseParams);
 			Params.AddIgnoredActors(IgnoreActors);
@@ -831,7 +876,7 @@ bool UMovementUtils::TestEncroachmentWithAdjustment_Internal(const UWorld* World
 					}
 					else
 					{
-						UE_LOG(LogPhysics, Log, TEXT("OverlapTest says we are overlapping, yet MTD says we're not. Something is wrong"));
+						UE_LOG(LogMover, Log, TEXT("OverlapTest says we are overlapping, yet MTD says we're not. Something is wrong"));
 						// It's not safe to use a partial result, that could push us out to an invalid location (like the other side of a wall).
 						OutProposedAdjustment = FVector::ZeroVector;
 						return true;
@@ -851,7 +896,7 @@ bool UMovementUtils::TestEncroachmentWithAdjustment_Internal(const UWorld* World
 						else
 						{
 							// Ignore this overlap.
-							UE_LOG(LogPhysics, Log, TEXT("OverlapTest says we are overlapping, yet MTD says we're not (with smaller shape). Ignoring this overlap."));
+							UE_LOG(LogMover, Log, TEXT("OverlapTest says we are overlapping, yet MTD says we're not (with smaller shape). Ignoring this overlap."));
 							NumBlockingHits--;
 							continue;
 						}
diff --git a/Engine/Plugins/Experimental/Mover/Source/Mover/Public/MoveLibrary/MovementUtils.h b/Engine/Plugins/Experimental/Mover/Source/Mover/Public/MoveLibrary/MovementUtils.h
--- a/Engine/Plugins/Experimental/Mover/Source/Mover/Public/MoveLibrary/MovementUtils.h
+++ b/Engine/Plugins/Experimental/Mover/Source/Mover/Public/MoveLibrary/MovementUtils.h
@@ -225,8 +225,17 @@ public:
 	UFUNCTION(BlueprintCallable, Category = Mover)
 	static bool FindTeleportSpot(const UMoverComponent* MoverComp, FVector& TestLocation, FRotator TestRotation);
 
-	/** @Return true if MoverComp's Actor would encroach at TestLocation on something that blocks it.  Sets OutProposedAdjustment that might result in an unblocked TestLocation. */
-	static bool TestEncroachment(const UMoverComponent* MoverComp, FVector TestLocation, FRotator TestRotation, OUT FVector* OutProposedAdjustment=nullptr);
+	/** Returns whether MoverComp's Actor would encroach at TestLocation on something that blocks it. */
+	UFUNCTION(BlueprintCallable, Category = Mover)
+	static bool TestEncroachment(const UMoverComponent* MoverComp, FVector TestLocation, FRotator TestRotation);
+
+	/** 
+	 * Returns whether MoverComp's Actor would encroach at TestLocation on something that blocks it.
+	 * If blocked, we'll attempt to find an adjustment and set OutProposedAdjustment accordingly. 
+	 * OutProposedAdjustment will be zero'd if there's no blockage or an adjustment could not be found.
+	 */
+	UFUNCTION(BlueprintCallable, Category = Mover)
+	static bool TestEncroachmentAndAdjust(const UMoverComponent* MoverComp, FVector TestLocation, FRotator TestRotation, FVector& OutProposedAdjustment);
 
 
 	// Internal functions - not meant to be called outside of this library
@@ -236,8 +245,12 @@ public:
 	 */
 	static bool TryMoveUpdatedComponent_Internal(const FMovingComponentSet& MovingComps, const FVector& Delta, const FQuat& NewRotation, bool bSweep, EMoveComponentFlags MoveComponentFlags, FHitResult* OutHit, ETeleportType Teleport);
 
-	// Internal function for testing whether a Mover actor would encroach at a test location, and computes a proposed adjustment where they won't encrouch (if found).
-	static bool TestEncroachmentWithAdjustment_Internal(const UWorld* World, const AActor* TestActor, const UPrimitiveComponent* PrimComp, const FTransform& TestWorldTransform, OUT FVector& OutProposedAdjustment, const TArray<AActor*>& IgnoreActors);
+	// Internal function for testing whether a Mover actor would encroach at a test location
+	static bool TestEncroachment_Internal(const UWorld* World, const AActor* TestActor, const UPrimitiveComponent* PrimComp, const FTransform& TestWorldTransform, const TArray<AActor*>& IgnoreActors);
+
+
+	// Internal function for testing whether a Mover actor would encroach at a test location, and computes a proposed adjustment where they won't encroach (if found).
+	static bool TestEncroachmentWithAdjustment_Internal(const UWorld* World, const AActor* TestActor, const UPrimitiveComponent* PrimComp, const FTransform& TestWorldTransform, const TArray<AActor*>& IgnoreActors, OUT FVector& OutProposedAdjustment);
 
 };
 
