[PhysicsMover] Remove ensure to check that async phsyics is enabled

diff --git a/Engine/Plugins/Experimental/Mover/Source/Mover/Private/Backends/MoverNetworkPhysicsLiaisonBase.cpp b/Engine/Plugins/Experimental/Mover/Source/Mover/Private/Backends/MoverNetworkPhysicsLiaisonBase.cpp
--- a/Engine/Plugins/Experimental/Mover/Source/Mover/Private/Backends/MoverNetworkPhysicsLiaisonBase.cpp
+++ b/Engine/Plugins/Experimental/Mover/Source/Mover/Private/Backends/MoverNetworkPhysicsLiaisonBase.cpp
@@ -308,7 +308,7 @@ void UMoverNetworkPhysicsLiaisonComponentBase::InitializeComponent()
 		MoverComp.ModeFSM->SetModeImmediately(MoverComp.StartingMovementMode);
 	}
 
-	if (Chaos::FPhysicsSolverBase::IsNetworkPhysicsPredictionEnabled() && ensureMsgf(bUsingAsyncPhysics, TEXT("Physics mover requires async physics to work with networked physics.")))
+	if (Chaos::FPhysicsSolverBase::IsNetworkPhysicsPredictionEnabled() && bUsingAsyncPhysics)
 	{
 		NetworkPhysicsComponent = NewObject<UNetworkPhysicsComponent>(GetOwner(), TEXT("PhysMover_NetworkPhysicsComponent"));
 
