Mover: Vertical speed is now properly clamped in positive and negative directions when going over terminal vertical speed.

diff --git a/Engine/Plugins/Experimental/Mover/Source/Mover/Private/DefaultMovementSet/Modes/AsyncFallingMode.cpp b/Engine/Plugins/Experimental/Mover/Source/Mover/Private/DefaultMovementSet/Modes/AsyncFallingMode.cpp
--- a/Engine/Plugins/Experimental/Mover/Source/Mover/Private/DefaultMovementSet/Modes/AsyncFallingMode.cpp
+++ b/Engine/Plugins/Experimental/Mover/Source/Mover/Private/DefaultMovementSet/Modes/AsyncFallingMode.cpp
@@ -127,7 +127,8 @@ void UAsyncFallingMode::OnGenerateMove(const FMoverTickStartData& StartState, co
 	{
 		if (bShouldClampTerminalVerticalSpeed)
 		{
-			UMovementUtils::SetGravityVerticalComponent(OutProposedMove.LinearVelocity, -TerminalVerticalSpeed, UpDirection);
+			const float ClampedVerticalSpeed = FMath::Sign(VelocityWithGravity.Dot(UpDirection)) * TerminalVerticalSpeed;
+			UMovementUtils::SetGravityVerticalComponent(OutProposedMove.LinearVelocity, ClampedVerticalSpeed, UpDirection);
 		}
 		else
 		{
diff --git a/Engine/Plugins/Experimental/Mover/Source/Mover/Private/DefaultMovementSet/Modes/FallingMode.cpp b/Engine/Plugins/Experimental/Mover/Source/Mover/Private/DefaultMovementSet/Modes/FallingMode.cpp
--- a/Engine/Plugins/Experimental/Mover/Source/Mover/Private/DefaultMovementSet/Modes/FallingMode.cpp
+++ b/Engine/Plugins/Experimental/Mover/Source/Mover/Private/DefaultMovementSet/Modes/FallingMode.cpp
@@ -119,7 +119,8 @@ void UFallingMode::OnGenerateMove(const FMoverTickStartData& StartState, const F
 	{
 		if (bShouldClampTerminalVerticalSpeed)
 		{
-			UMovementUtils::SetGravityVerticalComponent(OutProposedMove.LinearVelocity, -TerminalVerticalSpeed, UpDirection);
+			const float ClampedVerticalSpeed = FMath::Sign(VelocityWithGravity.Dot(UpDirection)) * TerminalVerticalSpeed;
+			UMovementUtils::SetGravityVerticalComponent(OutProposedMove.LinearVelocity, ClampedVerticalSpeed, UpDirection);
 		}
 		else
 		{
