Mover: Adding more nullptr checks to the NavMoverComponent to avoid crashes with using NavMoverComponent before it's initialized

diff --git a/Engine/Plugins/Experimental/Mover/Source/Mover/Private/DefaultMovementSet/NavMoverComponent.cpp b/Engine/Plugins/Experimental/Mover/Source/Mover/Private/DefaultMovementSet/NavMoverComponent.cpp
--- a/Engine/Plugins/Experimental/Mover/Source/Mover/Private/DefaultMovementSet/NavMoverComponent.cpp
+++ b/Engine/Plugins/Experimental/Mover/Source/Mover/Private/DefaultMovementSet/NavMoverComponent.cpp
@@ -26,7 +26,7 @@ void UNavMoverComponent::InitializeComponent()
 		MoverComponent = MovementCompOwner->FindComponentByClass<UMoverComponent>();
 	}
 	
-	if (!MoverComponent)
+	if (!MoverComponent.IsValid())
 	{
 		UE_LOG(LogMover, Warning, TEXT("NavMoverComponent on %s could not find a valid MoverComponent and will not function properly."), *GetNameSafe(GetOwner()));
 	}
@@ -36,7 +36,7 @@ void UNavMoverComponent::BeginPlay()
 {
 	Super::BeginPlay();
 
-	if (MoverComponent && MoverComponent->GetUpdatedComponent())
+	if (MoverComponent.IsValid() && MoverComponent->GetUpdatedComponent())
 	{
 		UpdateNavAgent(*MoverComponent->GetUpdatedComponent());
 	}
@@ -67,14 +67,27 @@ bool UNavMoverComponent::HasRequestedNavMovement() const
 
 float UNavMoverComponent::GetMaxSpeedForNavMovement() const
 {
-	const UCommonLegacyMovementSettings* MovementSettings = MoverComponent->FindSharedSettings_Mutable<UCommonLegacyMovementSettings>();
-	return MovementSettings ? MovementSettings->MaxSpeed : 0;
+	float MaxSpeed = 0.0f;
+	
+	if (MoverComponent.IsValid())
+	{
+		if (const UCommonLegacyMovementSettings* MovementSettings = MoverComponent.Get()->FindSharedSettings_Mutable<UCommonLegacyMovementSettings>())
+		{
+			MaxSpeed = MovementSettings->MaxSpeed;
+		}
+	}
+
+	return MaxSpeed;
 }
 
 void UNavMoverComponent::StopMovementImmediately()
 {
-	TSharedPtr<FApplyVelocityEffect> VelocityEffect = MakeShared<FApplyVelocityEffect>();
-	MoverComponent->QueueInstantMovementEffect(VelocityEffect);
+	if (MoverComponent.IsValid())
+	{
+		TSharedPtr<FApplyVelocityEffect> VelocityEffect = MakeShared<FApplyVelocityEffect>();
+		MoverComponent->QueueInstantMovementEffect(VelocityEffect);
+	}
+	
 	CachedNavMoveInputIntent = FVector::ZeroVector;
 	CachedNavMoveInputVelocity = FVector::ZeroVector;
 }
@@ -118,15 +131,18 @@ FBasedPosition UNavMoverComponent::GetFeetLocationBased() const
 {
 	FBasedPosition BasedPosition(NULL, GetFeetLocation());
 	
-	FRelativeBaseInfo MovementBaseInfo;
-	if (const UMoverBlackboard* Blackboard = MoverComponent->GetSimBlackboard())
+	if (MoverComponent.IsValid())
 	{
-		if (Blackboard->TryGet(CommonBlackboard::LastFoundDynamicMovementBase, MovementBaseInfo)) 
+		if (const UMoverBlackboard* Blackboard = MoverComponent->GetSimBlackboard())
 		{
-			BasedPosition.Base = MovementBaseInfo.MovementBase->GetOwner();
-			BasedPosition.Position = MovementBaseInfo.Location;
-			BasedPosition.CachedBaseLocation = MovementBaseInfo.ContactLocalPosition;
-			BasedPosition.CachedBaseRotation = MovementBaseInfo.Rotation.Rotator();
+			FRelativeBaseInfo MovementBaseInfo;
+			if (Blackboard->TryGet(CommonBlackboard::LastFoundDynamicMovementBase, MovementBaseInfo)) 
+			{
+				BasedPosition.Base = MovementBaseInfo.MovementBase->GetOwner();
+				BasedPosition.Position = MovementBaseInfo.Location;
+				BasedPosition.CachedBaseLocation = MovementBaseInfo.ContactLocalPosition;
+				BasedPosition.CachedBaseRotation = MovementBaseInfo.Rotation.Rotator();
+			}
 		}
 	}
 
@@ -246,27 +262,27 @@ bool UNavMoverComponent::CanStartPathFollowing() const
 
 bool UNavMoverComponent::IsCrouching() const
 {
-	return MoverComponent->HasGameplayTag(Mover_IsCrouching, true);
+	return MoverComponent.IsValid() ? MoverComponent->HasGameplayTag(Mover_IsCrouching, true) : false;
 }
 
 bool UNavMoverComponent::IsFalling() const
 {
-	return MoverComponent->HasGameplayTag(Mover_IsFalling, true);
+	return MoverComponent.IsValid() ? MoverComponent->HasGameplayTag(Mover_IsFalling, true) : false;
 }
 
 bool UNavMoverComponent::IsMovingOnGround() const
 {
-	return MoverComponent->HasGameplayTag(Mover_IsOnGround, true);
+	return MoverComponent.IsValid() ? MoverComponent->HasGameplayTag(Mover_IsOnGround, true) : false;
 }
 
 bool UNavMoverComponent::IsSwimming() const
 {
-	return MoverComponent->HasGameplayTag(Mover_IsSwimming, true);
+	return MoverComponent.IsValid() ? MoverComponent->HasGameplayTag(Mover_IsSwimming, true) : false;
 }
 
 bool UNavMoverComponent::IsFlying() const
 {
-	return MoverComponent->HasGameplayTag(Mover_IsFlying, true);
+	return MoverComponent.IsValid() ? MoverComponent->HasGameplayTag(Mover_IsFlying, true) : false;
 }
 
 void UNavMoverComponent::GetSimpleCollisionCylinder(float& CollisionRadius, float& CollisionHalfHeight) const
@@ -286,5 +302,5 @@ FVector UNavMoverComponent::GetForwardVector() const
 
 FVector UNavMoverComponent::GetVelocityForNavMovement() const
 {
-	return MoverComponent->GetVelocity();
+	return MoverComponent.IsValid() ? MoverComponent->GetVelocity() : FVector::ZeroVector;
 }
diff --git a/Engine/Plugins/Experimental/Mover/Source/Mover/Public/DefaultMovementSet/NavMoverComponent.h b/Engine/Plugins/Experimental/Mover/Source/Mover/Public/DefaultMovementSet/NavMoverComponent.h
--- a/Engine/Plugins/Experimental/Mover/Source/Mover/Public/DefaultMovementSet/NavMoverComponent.h
+++ b/Engine/Plugins/Experimental/Mover/Source/Mover/Public/DefaultMovementSet/NavMoverComponent.h
@@ -54,7 +54,7 @@ private:
 
 	/** Associated Movement component that will actually move the actor */ 
 	UPROPERTY()
-	TObjectPtr<UMoverComponent> MoverComponent;
+	TWeakObjectPtr<UMoverComponent> MoverComponent;
 
 	UPROPERTY(Transient)
 	bool bShouldClearNavMoveInputs = false;
