[Mover] Updating all of the mutually exclusive BP-or-native function calls in UBaseMovementMode and UBaseMovementModeTransition to be BlueprintNativeEvents instead
- This allows the BP implementation the option to include the parent implementation
- Existing usage of the removed BlueprintImplementableEvents will be fixed up and stomp the parent impl (retaining their previous behavior). Native overrides have been deprecated but still function.

Also:
- Implementing GetWorld on Mode/ModeTransition so they can implicitly be WorldContextObjects
- Adding an unchecked templated getter for GetMoverComponent on both
- Changing the checked getter to return a reference to reinforce that the result will not be null. This is a breaking, but very easily fixed-up change for any existing calls

diff --git a/Engine/Plugins/Experimental/Mover/Config/BaseMover.ini b/Engine/Plugins/Experimental/Mover/Config/BaseMover.ini
--- a/Engine/Plugins/Experimental/Mover/Config/BaseMover.ini
+++ b/Engine/Plugins/Experimental/Mover/Config/BaseMover.ini
@@ -10,3 +10,11 @@
 
 
 
++FunctionRedirects=(OldName="/Script/Mover.BaseMovementModeTransition.K2_OnTrigger",NewName="Trigger")
++FunctionRedirects=(OldName="/Script/Mover.BaseMovementModeTransition.K2_OnEvaluate",NewName="Evaluate")
++FunctionRedirects=(OldName="/Script/Mover.BaseMovementModeTransition.GetMoverComponent",NewName="K2_GetMoverComponent")
++FunctionRedirects=(OldName="/Script/Mover.BaseMovementMode.K2_OnGenerateMove",NewName="GenerateMove")
++FunctionRedirects=(OldName="/Script/Mover.BaseMovementMode.K2_OnSimulationTick",NewName="SimulationTick")
++FunctionRedirects=(OldName="/Script/Mover.BaseMovementMode.K2_OnActivate",NewName="Activate")
++FunctionRedirects=(OldName="/Script/Mover.BaseMovementMode.K2_OnDeactivate",NewName="Deactivate")
++FunctionRedirects=(OldName="/Script/Mover.BaseMovementMode.GetMoverComponent",NewName="K2_GetMoverComponent")
\ No newline at end of file
diff --git a/Engine/Plugins/Experimental/Mover/Source/Mover/Private/Backends/MoverPathedPhysicsLiaison.cpp b/Engine/Plugins/Experimental/Mover/Source/Mover/Private/Backends/MoverPathedPhysicsLiaison.cpp
--- a/Engine/Plugins/Experimental/Mover/Source/Mover/Private/Backends/MoverPathedPhysicsLiaison.cpp
+++ b/Engine/Plugins/Experimental/Mover/Source/Mover/Private/Backends/MoverPathedPhysicsLiaison.cpp
@@ -97,6 +97,15 @@ void UMoverPathedPhysicsLiaisonComponent::ConsumeOutput_External(const FPhysicsM
 	}
 }
 
+void UMoverPathedPhysicsLiaisonComponent::PostPhysicsUpdate_External()
+{
+	Super::PostPhysicsUpdate_External();
+
+	//@todo DanH: Does this help with jitter on non-PBCM CMCs?
+	const UMoverComponent& MoverComp = GetMoverComponent();
+	MoverComp.GetUpdatedComponent()->ComponentVelocity = MoverComp.GetVelocity();
+}
+
 void UMoverPathedPhysicsLiaisonComponent::SetPathOrigin(const FTransform& NewPathOrigin)
 {
 	if (!GetPathOrigin().Equals(NewPathOrigin) && GetOwnerRole() == ROLE_Authority)
@@ -315,16 +324,16 @@ void UMoverPathedPhysicsLiaisonComponent::PerformPreSimulate_Internal(const FPhy
 		Chaos::FPBDRigidsEvolution& Evolution = *GetPhysicsSolver()->GetEvolution();
 		
 		const Chaos::FRigidTransform3 TargetWorldTransform = Chaos::FRigidTransform3::MultiplyNoScale(TargetRelativeTransform, OutputMoveState.MutableProps.PathOrigin);
-		if (OutputMoveState.MutableProps.bIsJointEnabled)
-		{
-			Chaos::FKinematicGeometryParticleHandle* EndpointParticleHandle = GetJointConstraint()->GetKinematicEndPoint()->GetHandle_LowLevel()->CastToKinematicParticle();
-			Evolution.SetParticleKinematicTarget(EndpointParticleHandle, Chaos::FKinematicTarget::MakePositionTarget(TargetWorldTransform));
-		}
-		else
+		Chaos::FKinematicGeometryParticleHandle* EndpointParticleHandle = GetJointConstraint()->GetKinematicEndPoint()->GetHandle_LowLevel()->CastToKinematicParticle();
+		Evolution.SetParticleKinematicTarget(EndpointParticleHandle, Chaos::FKinematicTarget::MakePositionTarget(TargetWorldTransform));
+		if (!OutputMoveState.MutableProps.bIsJointEnabled)
 		{
 			// When we don't want to use the joint, just move the pathed component kinematically
 			Evolution.SetParticleKinematicTarget(GetControlledParticle_Internal(), Chaos::FKinematicTarget::MakePositionTarget(TargetWorldTransform));
 		}
+		// else
+		// {
+		// }
 	}
 }
 
@@ -368,6 +377,7 @@ void UMoverPathedPhysicsLiaisonComponent::CreateTargetJoint()
 			{
 				FChaosEngineInterface::ReleaseConstraint(Handle);
 			}
+		
 		}
 	}
 }
diff --git a/Engine/Plugins/Experimental/Mover/Source/Mover/Private/DefaultMovementSet/Modes/AsyncFallingMode.cpp b/Engine/Plugins/Experimental/Mover/Source/Mover/Private/DefaultMovementSet/Modes/AsyncFallingMode.cpp
--- a/Engine/Plugins/Experimental/Mover/Source/Mover/Private/DefaultMovementSet/Modes/AsyncFallingMode.cpp
+++ b/Engine/Plugins/Experimental/Mover/Source/Mover/Private/DefaultMovementSet/Modes/AsyncFallingMode.cpp
@@ -33,7 +33,7 @@ UAsyncFallingMode::UAsyncFallingMode(const FObjectInitializer& ObjectInitializer
 }
 
 
-void UAsyncFallingMode::OnGenerateMove(const FMoverTickStartData& StartState, const FMoverTimeStep& TimeStep, FProposedMove& OutProposedMove) const
+void UAsyncFallingMode::GenerateMove_Implementation(const FMoverTickStartData& StartState, const FMoverTimeStep& TimeStep, FProposedMove& OutProposedMove) const
 {
 	const UMoverComponent* MoverComp = GetMoverComponent();
 	const FCharacterDefaultInputs* CharacterInputs = StartState.InputCmd.InputCollection.FindDataByType<FCharacterDefaultInputs>();
@@ -146,7 +146,7 @@ void UAsyncFallingMode::OnGenerateMove(const FMoverTickStartData& StartState, co
 	}
 }
 
-void UAsyncFallingMode::OnSimulationTick(const FSimulationTickParams& Params, FMoverTickEndData& OutputState)
+void UAsyncFallingMode::SimulationTick_Implementation(const FSimulationTickParams& Params, FMoverTickEndData& OutputState)
 {
 	if (!CommonLegacySettings.IsValid())
 	{
diff --git a/Engine/Plugins/Experimental/Mover/Source/Mover/Private/DefaultMovementSet/Modes/AsyncFlyingMode.cpp b/Engine/Plugins/Experimental/Mover/Source/Mover/Private/DefaultMovementSet/Modes/AsyncFlyingMode.cpp
--- a/Engine/Plugins/Experimental/Mover/Source/Mover/Private/DefaultMovementSet/Modes/AsyncFlyingMode.cpp
+++ b/Engine/Plugins/Experimental/Mover/Source/Mover/Private/DefaultMovementSet/Modes/AsyncFlyingMode.cpp
@@ -18,7 +18,7 @@ UAsyncFlyingMode::UAsyncFlyingMode(const FObjectInitializer& ObjectInitializer)
 	GameplayTags.AddTag(Mover_IsFlying);
 }
 
-void UAsyncFlyingMode::OnGenerateMove(const FMoverTickStartData& StartState, const FMoverTimeStep& TimeStep, FProposedMove& OutProposedMove) const
+void UAsyncFlyingMode::GenerateMove_Implementation(const FMoverTickStartData& StartState, const FMoverTimeStep& TimeStep, FProposedMove& OutProposedMove) const
 {
 	const UMoverComponent* MoverComp = GetMoverComponent();
 	const FCharacterDefaultInputs* CharacterInputs = StartState.InputCmd.InputCollection.FindDataByType<FCharacterDefaultInputs>();
@@ -68,7 +68,7 @@ void UAsyncFlyingMode::OnGenerateMove(const FMoverTickStartData& StartState, con
 	OutProposedMove = UAirMovementUtils::ComputeControlledFreeMove(Params);
 }
 
-void UAsyncFlyingMode::OnSimulationTick(const FSimulationTickParams& Params, FMoverTickEndData& OutputState)
+void UAsyncFlyingMode::SimulationTick_Implementation(const FSimulationTickParams& Params, FMoverTickEndData& OutputState)
 {
 	const UMoverComponent* MoverComp = GetMoverComponent();
 	const FMoverTickStartData& StartState = Params.StartState;
diff --git a/Engine/Plugins/Experimental/Mover/Source/Mover/Private/DefaultMovementSet/Modes/AsyncNavWalkingMode.cpp b/Engine/Plugins/Experimental/Mover/Source/Mover/Private/DefaultMovementSet/Modes/AsyncNavWalkingMode.cpp
--- a/Engine/Plugins/Experimental/Mover/Source/Mover/Private/DefaultMovementSet/Modes/AsyncNavWalkingMode.cpp
+++ b/Engine/Plugins/Experimental/Mover/Source/Mover/Private/DefaultMovementSet/Modes/AsyncNavWalkingMode.cpp
@@ -32,7 +32,7 @@ UAsyncNavWalkingMode::UAsyncNavWalkingMode()
 	GameplayTags.AddTag(Mover_IsNavWalking);
 }
 
-void UAsyncNavWalkingMode::OnGenerateMove(const FMoverTickStartData& StartState, const FMoverTimeStep& TimeStep, FProposedMove& OutProposedMove) const
+void UAsyncNavWalkingMode::GenerateMove_Implementation(const FMoverTickStartData& StartState, const FMoverTimeStep& TimeStep, FProposedMove& OutProposedMove) const
 {
 	if (!CommonLegacySettings.IsValid())
 	{
@@ -119,7 +119,7 @@ void UAsyncNavWalkingMode::OnGenerateMove(const FMoverTickStartData& StartState,
 	}
 }
 
-void UAsyncNavWalkingMode::OnSimulationTick(const FSimulationTickParams& Params, FMoverTickEndData& OutputState)
+void UAsyncNavWalkingMode::SimulationTick_Implementation(const FSimulationTickParams& Params, FMoverTickEndData& OutputState)
 {
 	const UMoverComponent* MoverComp = Cast<UMoverComponent>(GetMoverComponent());
 	if (!MoverComp || !NavMoverComponent.IsValid() || !CommonLegacySettings.IsValid())
@@ -342,7 +342,7 @@ void UAsyncNavWalkingMode::SetCollisionForNavWalking(bool bEnable)
 	}
 }
 
-void UAsyncNavWalkingMode::OnActivate()
+void UAsyncNavWalkingMode::Activate_Implementation()
 {
 	if (UMoverComponent* MoverComp = GetMoverComponent())
 	{
@@ -358,7 +358,7 @@ void UAsyncNavWalkingMode::OnActivate()
 	NavMeshProjectionTimer = (NavMeshProjectionInterval > 0.f) ? FMath::FRandRange(-NavMeshProjectionInterval, 0.f) : 0.f;
 }
 
-void UAsyncNavWalkingMode::OnDeactivate()
+void UAsyncNavWalkingMode::Deactivate_Implementation()
 {
 	if (UMoverComponent* MoverComp = GetMoverComponent())
 	{
diff --git a/Engine/Plugins/Experimental/Mover/Source/Mover/Private/DefaultMovementSet/Modes/AsyncWalkingMode.cpp b/Engine/Plugins/Experimental/Mover/Source/Mover/Private/DefaultMovementSet/Modes/AsyncWalkingMode.cpp
--- a/Engine/Plugins/Experimental/Mover/Source/Mover/Private/DefaultMovementSet/Modes/AsyncWalkingMode.cpp
+++ b/Engine/Plugins/Experimental/Mover/Source/Mover/Private/DefaultMovementSet/Modes/AsyncWalkingMode.cpp
@@ -26,7 +26,7 @@ UAsyncWalkingMode::UAsyncWalkingMode(const FObjectInitializer& ObjectInitializer
 	GameplayTags.AddTag(Mover_IsOnGround);
 }
 
-void UAsyncWalkingMode::OnGenerateMove(const FMoverTickStartData& StartState, const FMoverTimeStep& TimeStep, FProposedMove& OutProposedMove) const
+void UAsyncWalkingMode::GenerateMove_Implementation(const FMoverTickStartData& StartState, const FMoverTimeStep& TimeStep, FProposedMove& OutProposedMove) const
 {
 	const UMoverComponent* MoverComp = GetMoverComponent();
 	const FCharacterDefaultInputs* CharacterInputs = StartState.InputCmd.InputCollection.FindDataByType<FCharacterDefaultInputs>();
@@ -115,7 +115,7 @@ void UAsyncWalkingMode::OnGenerateMove(const FMoverTickStartData& StartState, co
 	}
 }
 
-void UAsyncWalkingMode::OnSimulationTick(const FSimulationTickParams& Params, FMoverTickEndData& OutputState)
+void UAsyncWalkingMode::SimulationTick_Implementation(const FSimulationTickParams& Params, FMoverTickEndData& OutputState)
 {
 	if (!CommonLegacySettings.IsValid())
 	{
diff --git a/Engine/Plugins/Experimental/Mover/Source/Mover/Private/DefaultMovementSet/Modes/FallingMode.cpp b/Engine/Plugins/Experimental/Mover/Source/Mover/Private/DefaultMovementSet/Modes/FallingMode.cpp
--- a/Engine/Plugins/Experimental/Mover/Source/Mover/Private/DefaultMovementSet/Modes/FallingMode.cpp
+++ b/Engine/Plugins/Experimental/Mover/Source/Mover/Private/DefaultMovementSet/Modes/FallingMode.cpp
@@ -30,7 +30,7 @@ UFallingMode::UFallingMode(const FObjectInitializer& ObjectInitializer)
 }
 
 
-void UFallingMode::OnGenerateMove(const FMoverTickStartData& StartState, const FMoverTimeStep& TimeStep, FProposedMove& OutProposedMove) const
+void UFallingMode::GenerateMove_Implementation(const FMoverTickStartData& StartState, const FMoverTimeStep& TimeStep, FProposedMove& OutProposedMove) const
 {
 	const UMoverComponent* MoverComp = GetMoverComponent();
 	const FCharacterDefaultInputs* CharacterInputs = StartState.InputCmd.InputCollection.FindDataByType<FCharacterDefaultInputs>();
@@ -138,7 +138,7 @@ void UFallingMode::OnGenerateMove(const FMoverTickStartData& StartState, const F
 	}
 }
 
-void UFallingMode::OnSimulationTick(const FSimulationTickParams& Params, FMoverTickEndData& OutputState)
+void UFallingMode::SimulationTick_Implementation(const FSimulationTickParams& Params, FMoverTickEndData& OutputState)
 {
 	UMoverComponent* MoverComponent = GetMoverComponent();
 	const FMoverTickStartData& StartState = Params.StartState;
diff --git a/Engine/Plugins/Experimental/Mover/Source/Mover/Private/DefaultMovementSet/Modes/FlyingMode.cpp b/Engine/Plugins/Experimental/Mover/Source/Mover/Private/DefaultMovementSet/Modes/FlyingMode.cpp
--- a/Engine/Plugins/Experimental/Mover/Source/Mover/Private/DefaultMovementSet/Modes/FlyingMode.cpp
+++ b/Engine/Plugins/Experimental/Mover/Source/Mover/Private/DefaultMovementSet/Modes/FlyingMode.cpp
@@ -16,7 +16,7 @@ UFlyingMode::UFlyingMode(const FObjectInitializer& ObjectInitializer)
 	GameplayTags.AddTag(Mover_IsFlying);
 }
 
-void UFlyingMode::OnGenerateMove(const FMoverTickStartData& StartState, const FMoverTimeStep& TimeStep, FProposedMove& OutProposedMove) const
+void UFlyingMode::GenerateMove_Implementation(const FMoverTickStartData& StartState, const FMoverTimeStep& TimeStep, FProposedMove& OutProposedMove) const
 {
 	const UMoverComponent* MoverComp = GetMoverComponent();
 	const FCharacterDefaultInputs* CharacterInputs = StartState.InputCmd.InputCollection.FindDataByType<FCharacterDefaultInputs>();
@@ -66,7 +66,7 @@ void UFlyingMode::OnGenerateMove(const FMoverTickStartData& StartState, const FM
 	OutProposedMove = UAirMovementUtils::ComputeControlledFreeMove(Params);
 }
 
-void UFlyingMode::OnSimulationTick(const FSimulationTickParams& Params, FMoverTickEndData& OutputState)
+void UFlyingMode::SimulationTick_Implementation(const FSimulationTickParams& Params, FMoverTickEndData& OutputState)
 {
 	const UMoverComponent* MoverComp = GetMoverComponent();
 	const FMoverTickStartData& StartState = Params.StartState;
diff --git a/Engine/Plugins/Experimental/Mover/Source/Mover/Private/DefaultMovementSet/Modes/NavWalkingMode.cpp b/Engine/Plugins/Experimental/Mover/Source/Mover/Private/DefaultMovementSet/Modes/NavWalkingMode.cpp
--- a/Engine/Plugins/Experimental/Mover/Source/Mover/Private/DefaultMovementSet/Modes/NavWalkingMode.cpp
+++ b/Engine/Plugins/Experimental/Mover/Source/Mover/Private/DefaultMovementSet/Modes/NavWalkingMode.cpp
@@ -29,7 +29,7 @@ UNavWalkingMode::UNavWalkingMode()
 	GameplayTags.AddTag(Mover_IsNavWalking);
 }
 
-void UNavWalkingMode::OnGenerateMove(const FMoverTickStartData& StartState, const FMoverTimeStep& TimeStep, FProposedMove& OutProposedMove) const
+void UNavWalkingMode::GenerateMove_Implementation(const FMoverTickStartData& StartState, const FMoverTimeStep& TimeStep, FProposedMove& OutProposedMove) const
 {
 	const UMoverComponent* MoverComp = GetMoverComponent();
 	const FCharacterDefaultInputs* CharacterInputs = StartState.InputCmd.InputCollection.FindDataByType<FCharacterDefaultInputs>();
@@ -111,7 +111,7 @@ void UNavWalkingMode::OnGenerateMove(const FMoverTickStartData& StartState, cons
 	}
 }
 
-void UNavWalkingMode::OnSimulationTick(const FSimulationTickParams& Params, FMoverTickEndData& OutputState)
+void UNavWalkingMode::SimulationTick_Implementation(const FSimulationTickParams& Params, FMoverTickEndData& OutputState)
 {
 	const UMoverComponent* MoverComp = Cast<UMoverComponent>(GetMoverComponent());
 	if (!ensureMsgf(MoverComp, TEXT("Nav Walking Mode couldn't find a valid MoverComponent!")))
@@ -345,12 +345,12 @@ void UNavWalkingMode::SetCollisionForNavWalking(bool bEnable)
 	}
 }
 
-void UNavWalkingMode::OnActivate()
+void UNavWalkingMode::Activate_Implementation()
 {
 	SetCollisionForNavWalking(true);
 }
 
-void UNavWalkingMode::OnDeactivate()
+void UNavWalkingMode::Deactivate_Implementation()
 {
 	SetCollisionForNavWalking(false);
 }
diff --git a/Engine/Plugins/Experimental/Mover/Source/Mover/Private/DefaultMovementSet/Modes/SwimmingMode.cpp b/Engine/Plugins/Experimental/Mover/Source/Mover/Private/DefaultMovementSet/Modes/SwimmingMode.cpp
--- a/Engine/Plugins/Experimental/Mover/Source/Mover/Private/DefaultMovementSet/Modes/SwimmingMode.cpp
+++ b/Engine/Plugins/Experimental/Mover/Source/Mover/Private/DefaultMovementSet/Modes/SwimmingMode.cpp
@@ -21,7 +21,7 @@ USwimmingMode::USwimmingMode(const FObjectInitializer& ObjectInitializer)
 	GameplayTags.AddTag(Mover_IsSwimming);
 }
 
-void USwimmingMode::OnGenerateMove(const FMoverTickStartData& StartState, const FMoverTimeStep& TimeStep, FProposedMove& OutProposedMove) const
+void USwimmingMode::GenerateMove_Implementation(const FMoverTickStartData& StartState, const FMoverTimeStep& TimeStep, FProposedMove& OutProposedMove) const
 {
 	const UMoverComponent* MoverComp = GetMoverComponent();
 	const FCharacterDefaultInputs* CharacterInputs = StartState.InputCmd.InputCollection.FindDataByType<FCharacterDefaultInputs>();
@@ -161,7 +161,7 @@ void USwimmingMode::OnGenerateMove(const FMoverTickStartData& StartState, const
 	}
 }
 
-void USwimmingMode::OnSimulationTick(const FSimulationTickParams& Params, FMoverTickEndData& OutputState)
+void USwimmingMode::SimulationTick_Implementation(const FSimulationTickParams& Params, FMoverTickEndData& OutputState)
 {
 	// TODO Har
 }
diff --git a/Engine/Plugins/Experimental/Mover/Source/Mover/Private/DefaultMovementSet/Modes/WalkingMode.cpp b/Engine/Plugins/Experimental/Mover/Source/Mover/Private/DefaultMovementSet/Modes/WalkingMode.cpp
--- a/Engine/Plugins/Experimental/Mover/Source/Mover/Private/DefaultMovementSet/Modes/WalkingMode.cpp
+++ b/Engine/Plugins/Experimental/Mover/Source/Mover/Private/DefaultMovementSet/Modes/WalkingMode.cpp
@@ -24,7 +24,7 @@ UWalkingMode::UWalkingMode(const FObjectInitializer& ObjectInitializer)
 	GameplayTags.AddTag(Mover_IsOnGround);
 }
 
-void UWalkingMode::OnGenerateMove(const FMoverTickStartData& StartState, const FMoverTimeStep& TimeStep, FProposedMove& OutProposedMove) const
+void UWalkingMode::GenerateMove_Implementation(const FMoverTickStartData& StartState, const FMoverTimeStep& TimeStep, FProposedMove& OutProposedMove) const
 {
 	const UMoverComponent* MoverComp = GetMoverComponent();
 	const FCharacterDefaultInputs* CharacterInputs = StartState.InputCmd.InputCollection.FindDataByType<FCharacterDefaultInputs>();
@@ -109,7 +109,7 @@ void UWalkingMode::OnGenerateMove(const FMoverTickStartData& StartState, const F
 	}
 }
 
-void UWalkingMode::OnSimulationTick(const FSimulationTickParams& Params, FMoverTickEndData& OutputState)
+void UWalkingMode::SimulationTick_Implementation(const FSimulationTickParams& Params, FMoverTickEndData& OutputState)
 {
 	UMoverComponent* MoverComp = GetMoverComponent();
 	const FMoverTickStartData& StartState = Params.StartState;
diff --git a/Engine/Plugins/Experimental/Mover/Source/Mover/Private/MovementMode.cpp b/Engine/Plugins/Experimental/Mover/Source/Mover/Private/MovementMode.cpp
--- a/Engine/Plugins/Experimental/Mover/Source/Mover/Private/MovementMode.cpp
+++ b/Engine/Plugins/Experimental/Mover/Source/Mover/Private/MovementMode.cpp
@@ -9,105 +9,61 @@
 #include "Misc/DataValidation.h"
 #endif
 
-
 #include UE_INLINE_GENERATED_CPP_BY_NAME(MovementMode)
 
 #define LOCTEXT_NAMESPACE "Mover"
 
-UBaseMovementMode::UBaseMovementMode(const FObjectInitializer& ObjectInitializer)
-	: Super(ObjectInitializer)
+UWorld* UBaseMovementMode::GetWorld() const
 {
-
+	return GetOuterUMoverComponent()->GetWorld();
 }
 
-void UBaseMovementMode::DoRegister(const FName ModeName)
-{
-	auto IsImplementedInBlueprint = [](const UFunction* Func) -> bool
-	{
-		return Func && ensure(Func->GetOuter())
-			&& Func->GetOuter()->IsA(UBlueprintGeneratedClass::StaticClass());
-	};
-
-	static FName SimTickFuncName = FName(TEXT("K2_OnSimulationTick"));
-	UFunction* SimulationTickFunction = GetClass()->FindFunctionByName(SimTickFuncName);
-	bHasBlueprintSimulationTick = IsImplementedInBlueprint(SimulationTickFunction);
-
-	static FName GenMoveFuncName = FName(TEXT("K2_OnGenerateMove"));
-	UFunction* GenMoveFunction = GetClass()->FindFunctionByName(GenMoveFuncName);
-	bHasBlueprintGenerateMove = IsImplementedInBlueprint(GenMoveFunction);
 
+void UBaseMovementMode::OnRegistered(const FName ModeName)
+{
 	for (TObjectPtr<UBaseMovementModeTransition>& Transition : Transitions)
 	{
-		Transition->DoRegister();
+		Transition->OnRegistered();
 	}
-
-	OnRegistered(ModeName);
+	
+	K2_OnRegistered(ModeName);
 }
 
-void UBaseMovementMode::DoUnregister()
+void UBaseMovementMode::OnUnregistered()
 {
 	for (TObjectPtr<UBaseMovementModeTransition>& Transition : Transitions)
 	{
-		Transition->DoUnregister();
+		Transition->OnUnregistered();
 	}
 
-	OnUnregistered();
+	K2_OnUnregistered();
 }
 
-
-void UBaseMovementMode::DoGenerateMove(const FMoverTickStartData& StartState, const FMoverTimeStep& TimeStep, FProposedMove& OutProposedMove) const
+PRAGMA_DISABLE_DEPRECATION_WARNINGS
+void UBaseMovementMode::Activate_Implementation()
 {
-	if (bHasBlueprintGenerateMove)
-	{
-		OutProposedMove = K2_OnGenerateMove(StartState, TimeStep);
-	}
-	else
-	{
-		OnGenerateMove(StartState, TimeStep, OutProposedMove);
-	}
+	OnActivate();
 }
 
-
-void UBaseMovementMode::DoSimulationTick(const FSimulationTickParams& Params, FMoverTickEndData& OutputState)
+void UBaseMovementMode::Deactivate_Implementation()
 {
-	if (bHasBlueprintSimulationTick)
-	{
-		OutputState = K2_OnSimulationTick(Params);
-	}
-	else
-	{
-		OnSimulationTick(Params, OutputState);
-	}
+	OnDeactivate();
 }
 
-void UBaseMovementMode::DoActivate()
+void UBaseMovementMode::GenerateMove_Implementation(const FMoverTickStartData& StartState, const FMoverTimeStep& TimeStep, FProposedMove& OutProposedMove) const
 {
-	if (bHasBlueprintOnActivate)
-	{
-		K2_OnActivate();
-	}
-	else
-	{
-		OnActivate();
-	}
+	OnGenerateMove(StartState, TimeStep, OutProposedMove);
 }
 
-void UBaseMovementMode::DoDeactivate()
+void UBaseMovementMode::SimulationTick_Implementation(const FSimulationTickParams& Params, FMoverTickEndData& OutputState)
 {
-	if (bHasBlueprintOnDeactivate)
-	{
-		K2_OnDeactivate();
-	}
-	else
-	{
-		OnDeactivate();
-	}
+	OnSimulationTick(Params, OutputState);
 }
+PRAGMA_ENABLE_DEPRECATION_WARNINGS
 
-
-UMoverComponent* UBaseMovementMode::GetMoverComponent() const
+UMoverComponent* UBaseMovementMode::K2_GetMoverComponent() const
 {
-	return CastChecked<UMoverComponent>(GetOuter());
+	return GetOuterUMoverComponent();
 }
 
 #if WITH_EDITOR
@@ -144,28 +100,4 @@ bool UBaseMovementMode::HasGameplayTag(FGameplayTag TagToFind, bool bExactMatch)
 	return GameplayTags.HasTag(TagToFind);
 }
 
-void UBaseMovementMode::OnRegistered(const FName ModeName)
-{
-}
-
-void UBaseMovementMode::OnUnregistered()
-{
-}
-
-void UBaseMovementMode::OnGenerateMove(const FMoverTickStartData& StartState, const FMoverTimeStep& TimeStep, FProposedMove& OutProposedMove) const
-{
-}
-
-void UBaseMovementMode::OnSimulationTick(const FSimulationTickParams& Params, FMoverTickEndData& OutputState)
-{
-}
-
-void UBaseMovementMode::OnActivate()
-{
-}
-
-void UBaseMovementMode::OnDeactivate()
-{
-}
-
 #undef LOCTEXT_NAMESPACE
diff --git a/Engine/Plugins/Experimental/Mover/Source/Mover/Private/MovementModeStateMachine.cpp b/Engine/Plugins/Experimental/Mover/Source/Mover/Private/MovementModeStateMachine.cpp
--- a/Engine/Plugins/Experimental/Mover/Source/Mover/Private/MovementModeStateMachine.cpp
+++ b/Engine/Plugins/Experimental/Mover/Source/Mover/Private/MovementModeStateMachine.cpp
@@ -20,7 +20,7 @@ UNullMovementMode::UNullMovementMode(const FObjectInitializer& ObjectInitializer
 {
 }
 
-void UNullMovementMode::OnSimulationTick(const FSimulationTickParams& Params, FMoverTickEndData& OutputState)
+void UNullMovementMode::SimulationTick_Implementation(const FSimulationTickParams& Params, FMoverTickEndData& OutputState)
 {
 }
 
@@ -45,7 +45,7 @@ void UMovementModeStateMachine::RegisterMovementMode(FName ModeName, TObjectPtr<
 		DefaultModeName = ModeName;
 	}
 
-	Mode->DoRegister(ModeName);
+	Mode->OnRegistered(ModeName);
 }
 
 
@@ -64,19 +64,19 @@ void UMovementModeStateMachine::UnregisterMovementMode(FName ModeName)
 
 	if (ModeToUnregister)
 	{
-		ModeToUnregister->DoUnregister();
+		ModeToUnregister->OnUnregistered();
 	}
 }
 
 void UMovementModeStateMachine::ClearAllMovementModes()
 {
+	Modes.Empty();
+
 	for (TPair<FName, TObjectPtr<UBaseMovementMode>>& Element : Modes)
 	{
-		Element.Value->DoUnregister();
+		Element.Value->OnUnregistered();
 	}
 
-	Modes.Empty();
-
 	ConstructDefaultModes();	// Note that we're resetting to our defaults so we keep the null movement mode
 }
 
@@ -91,12 +91,12 @@ void UMovementModeStateMachine::RegisterGlobalTransition(TObjectPtr<UBaseMovemen
 {
 	GlobalTransitions.Add(Transition);
 
-	Transition->DoRegister();
+	Transition->OnRegistered();
 }
 
 void UMovementModeStateMachine::UnregisterGlobalTransition(TObjectPtr<UBaseMovementModeTransition> Transition)
 {
-	Transition->DoUnregister();
+	Transition->OnUnregistered();
 
 	GlobalTransitions.Remove(Transition);
 }
@@ -105,7 +105,7 @@ void UMovementModeStateMachine::ClearAllGlobalTransitions()
 {
 	for (TObjectPtr<UBaseMovementModeTransition>& Transition : GlobalTransitions)
 	{
-		Transition->DoUnregister();
+		Transition->OnUnregistered();
 	}
 
 	GlobalTransitions.Empty();
@@ -265,7 +265,7 @@ void UMovementModeStateMachine::OnSimulationTick(USceneComponent* UpdatedCompone
 		{
 			UBaseMovementMode* CurrentMode = Modes[CurrentModeName];
 			FProposedMove CombinedMove;
-			CurrentMode->DoGenerateMove(SubstepStartData, SubTimeStep, OUT CombinedMove);
+			CurrentMode->GenerateMove(SubstepStartData, SubTimeStep, OUT CombinedMove);
 
 			if (bHasLayeredMoveContributions)
 			{
@@ -301,7 +301,7 @@ void UMovementModeStateMachine::OnSimulationTick(USceneComponent* UpdatedCompone
 				{
 					if (IsValid(Transition))
 					{
-						EvalResult = Transition->DoEvaluate(SimTickParams);
+						EvalResult = Transition->Evaluate(SimTickParams);
 
 						if (!EvalResult.NextMode.IsNone())
 						{
@@ -320,7 +320,7 @@ void UMovementModeStateMachine::OnSimulationTick(USceneComponent* UpdatedCompone
 					{
 						if (IsValid(Transition))
 						{
-							EvalResult = Transition->DoEvaluate(SimTickParams);
+							EvalResult = Transition->Evaluate(SimTickParams);
 
 							if (!EvalResult.NextMode.IsNone())
 							{
@@ -338,11 +338,11 @@ void UMovementModeStateMachine::OnSimulationTick(USceneComponent* UpdatedCompone
 				{
 					OutputState.MovementEndState.NextModeName = EvalResult.NextMode;
 					OutputState.MovementEndState.RemainingMs = SimTickParams.TimeStep.StepMs; 	// Pass all remaining time to next mode
-					TransitionToTrigger->DoTrigger(SimTickParams);
+					TransitionToTrigger->Trigger(SimTickParams);
 				}
 				else
 				{
-					CurrentMode->DoSimulationTick(SimTickParams, OutputState);
+					CurrentMode->SimulationTick(SimTickParams, OutputState);
 				}
 
 				OutputState.MovementEndState.RemainingMs = FMath::Clamp(OutputState.MovementEndState.RemainingMs, 0.0f, SubTimeStep.StepMs);
@@ -510,10 +510,10 @@ void UMovementModeStateMachine::AdvanceToNextMode()
 
 		if (PreviousModeName != NAME_None && Modes.Contains(PreviousModeName))
 		{
-			Modes[PreviousModeName]->DoDeactivate();
+			Modes[PreviousModeName]->Deactivate();
 		}
 
-		Modes[CurrentModeName]->DoActivate();
+		Modes[CurrentModeName]->Activate();
 
 		// signal movement mode change event
 		const UMoverComponent* MoverComp = CastChecked<UMoverComponent>(GetOuter());
@@ -658,7 +658,7 @@ UImmediateMovementModeTransition::UImmediateMovementModeTransition(const FObject
 	Clear();
 }
 
-FTransitionEvalResult UImmediateMovementModeTransition::OnEvaluate(const FSimulationTickParams& Params) const
+FTransitionEvalResult UImmediateMovementModeTransition::Evaluate_Implementation(const FSimulationTickParams& Params) const
 {
 	if (NextMode != NAME_None)
 	{
@@ -675,7 +675,7 @@ FTransitionEvalResult UImmediateMovementModeTransition::OnEvaluate(const FSimula
 	return FTransitionEvalResult::NoTransition;
 }
 
-void UImmediateMovementModeTransition::OnTrigger(const FSimulationTickParams& Params)
+void UImmediateMovementModeTransition::Trigger_Implementation(const FSimulationTickParams& Params)
 {
 	Clear();
 }
diff --git a/Engine/Plugins/Experimental/Mover/Source/Mover/Private/MovementModeTransition.cpp b/Engine/Plugins/Experimental/Mover/Source/Mover/Private/MovementModeTransition.cpp
--- a/Engine/Plugins/Experimental/Mover/Source/Mover/Private/MovementModeTransition.cpp
+++ b/Engine/Plugins/Experimental/Mover/Source/Mover/Private/MovementModeTransition.cpp
@@ -11,93 +11,48 @@
 
 const FTransitionEvalResult FTransitionEvalResult::NoTransition = FTransitionEvalResult();
 
-UBaseMovementModeTransition::UBaseMovementModeTransition(const FObjectInitializer& ObjectInitializer)
-	: Super(ObjectInitializer)
-	, bAllowModeReentry(false)
+UWorld* UBaseMovementModeTransition::GetWorld() const
 {
-}
-
-void UBaseMovementModeTransition::DoRegister()
-{
-	auto IsImplementedInBlueprint = [](const UFunction* Func) -> bool
+	if (UMoverComponent* MoverComponent = GetMoverComponent())
 	{
-		return Func && ensure(Func->GetOuter())
-			&& Func->GetOuter()->IsA(UBlueprintGeneratedClass::StaticClass());
-	};
-
-	static FName EvaluateFuncName = FName(TEXT("K2_OnEvaluate"));
-	UFunction* EvaluateFunction = GetClass()->FindFunctionByName(EvaluateFuncName);
-	bHasBlueprintEvaluate = IsImplementedInBlueprint(EvaluateFunction);
-
-	static FName TriggerFuncName = FName(TEXT("K2_OnTrigger"));
-	UFunction* TriggerFunction = GetClass()->FindFunctionByName(TriggerFuncName);
-	bHasBlueprintTrigger = IsImplementedInBlueprint(TriggerFunction);
-
-	OnRegistered();
-}
-
-void UBaseMovementModeTransition::DoUnregister()
-{
-	OnUnregistered();
+		return MoverComponent->GetWorld();
+	}
+	return nullptr;
 }
 
 void UBaseMovementModeTransition::OnRegistered()
 {
+	K2_OnRegistered();
 }
 
 void UBaseMovementModeTransition::OnUnregistered()
 {
+	K2_OnUnregistered();
 }
 
-UMoverComponent* UBaseMovementModeTransition::GetMoverComponent() const
+UMoverComponent* UBaseMovementModeTransition::K2_GetMoverComponent() const
 {
-	// Transitions can belong to either the mover component or to
-	// a particular movement mode
-	if (UBaseMovementMode* MovementMode = Cast<UBaseMovementMode>(GetOuter()))
-	{
-		return MovementMode->GetMoverComponent();
-	}
-	else
-	{
-		return Cast<UMoverComponent>(GetOuter());
-	}
+	// Transitions can belong to either a mode or the component itself - either way they're always ultimately outer'd to a mover comp
+	return GetTypedOuter<UMoverComponent>();
 }
 
-FTransitionEvalResult UBaseMovementModeTransition::DoEvaluate(const FSimulationTickParams& Params) const
+PRAGMA_DISABLE_DEPRECATION_WARNINGS
+FTransitionEvalResult UBaseMovementModeTransition::Evaluate_Implementation(const FSimulationTickParams& Params) const
 {
-	if (bHasBlueprintEvaluate)
-	{
-		return K2_OnEvaluate(Params);
-	}
-	else
-	{
-		return OnEvaluate(Params);
-	}
+	return OnEvaluate(Params);
 }
 
-FTransitionEvalResult UBaseMovementModeTransition::OnEvaluate(const FSimulationTickParams& Params) const
+void UBaseMovementModeTransition::Trigger_Implementation(const FSimulationTickParams& Params)
 {
-	return FTransitionEvalResult::NoTransition;
+	OnTrigger(Params);
 }
+PRAGMA_ENABLE_DEPRECATION_WARNINGS
 
-
-void UBaseMovementModeTransition::DoTrigger(const FSimulationTickParams& Params)
-{
-	if (bHasBlueprintTrigger)
-	{
-		return K2_OnTrigger(Params);
-	}
-	else
-	{
-		return OnTrigger(Params);
-	}
-}
-
-void UBaseMovementModeTransition::OnTrigger(const FSimulationTickParams& Params)
+FTransitionEvalResult UBaseMovementModeTransition::OnEvaluate(const FSimulationTickParams& Params) const
 {
+	return FTransitionEvalResult::NoTransition;
 }
 
-
 #if WITH_EDITOR
 EDataValidationResult UBaseMovementModeTransition::IsDataValid(FDataValidationContext& Context) const
 {
diff --git a/Engine/Plugins/Experimental/Mover/Source/Mover/Private/MoverComponent.cpp b/Engine/Plugins/Experimental/Mover/Source/Mover/Private/MoverComponent.cpp
--- a/Engine/Plugins/Experimental/Mover/Source/Mover/Private/MoverComponent.cpp
+++ b/Engine/Plugins/Experimental/Mover/Source/Mover/Private/MoverComponent.cpp
@@ -1886,7 +1886,7 @@ TArray<FTrajectorySampleInfo> UMoverComponent::GetPredictedTrajectory(FMoverPred
 
 					// Generate next move from current step state
 					FProposedMove StepMove;
-					CurrentMovementMode->DoGenerateMove(StepState, FutureTimeStep, StepMove);
+					CurrentMovementMode->GenerateMove(StepState, FutureTimeStep, StepMove);
 
 					// Advance state based on move
 					StepSyncState->SetTransforms_WorldSpace(StepSyncState->GetLocation_WorldSpace() + (StepMove.LinearVelocity * PredictionParams.SecondsPerSample),
diff --git a/Engine/Plugins/Experimental/Mover/Source/Mover/Private/PhysicsMover/Modes/PhysicsDrivenFallingMode.cpp b/Engine/Plugins/Experimental/Mover/Source/Mover/Private/PhysicsMover/Modes/PhysicsDrivenFallingMode.cpp
--- a/Engine/Plugins/Experimental/Mover/Source/Mover/Private/PhysicsMover/Modes/PhysicsDrivenFallingMode.cpp
+++ b/Engine/Plugins/Experimental/Mover/Source/Mover/Private/PhysicsMover/Modes/PhysicsDrivenFallingMode.cpp
@@ -86,7 +86,7 @@ void UPhysicsDrivenFallingMode::OnRegistered(const FName ModeName)
 	}
 }
 
-void UPhysicsDrivenFallingMode::OnSimulationTick(const FSimulationTickParams& Params, FMoverTickEndData& OutputState)
+void UPhysicsDrivenFallingMode::SimulationTick_Implementation(const FSimulationTickParams& Params, FMoverTickEndData& OutputState)
 {
 	const FMoverTickStartData& StartState = Params.StartState;
 	USceneComponent* UpdatedComponent = Params.MovingComps.UpdatedComponent.Get();
diff --git a/Engine/Plugins/Experimental/Mover/Source/Mover/Private/PhysicsMover/Modes/PhysicsDrivenFlyingMode.cpp b/Engine/Plugins/Experimental/Mover/Source/Mover/Private/PhysicsMover/Modes/PhysicsDrivenFlyingMode.cpp
--- a/Engine/Plugins/Experimental/Mover/Source/Mover/Private/PhysicsMover/Modes/PhysicsDrivenFlyingMode.cpp
+++ b/Engine/Plugins/Experimental/Mover/Source/Mover/Private/PhysicsMover/Modes/PhysicsDrivenFlyingMode.cpp
@@ -47,7 +47,7 @@ EDataValidationResult UPhysicsDrivenFlyingMode::IsDataValid(FDataValidationConte
 }
 #endif // WITH_EDITOR
 
-void UPhysicsDrivenFlyingMode::OnSimulationTick(const FSimulationTickParams& Params, FMoverTickEndData& OutputState)
+void UPhysicsDrivenFlyingMode::SimulationTick_Implementation(const FSimulationTickParams& Params, FMoverTickEndData& OutputState)
 {
 	const UMoverComponent* MoverComp = GetMoverComponent();
 	const FMoverTickStartData& StartState = Params.StartState;
diff --git a/Engine/Plugins/Experimental/Mover/Source/Mover/Private/PhysicsMover/Modes/PhysicsDrivenSwimmingMode.cpp b/Engine/Plugins/Experimental/Mover/Source/Mover/Private/PhysicsMover/Modes/PhysicsDrivenSwimmingMode.cpp
--- a/Engine/Plugins/Experimental/Mover/Source/Mover/Private/PhysicsMover/Modes/PhysicsDrivenSwimmingMode.cpp
+++ b/Engine/Plugins/Experimental/Mover/Source/Mover/Private/PhysicsMover/Modes/PhysicsDrivenSwimmingMode.cpp
@@ -87,7 +87,7 @@ void UPhysicsDrivenSwimmingMode::OnRegistered(const FName ModeName)
 	}
 }
 
-void UPhysicsDrivenSwimmingMode::OnSimulationTick(const FSimulationTickParams& Params, FMoverTickEndData& OutputState)
+void UPhysicsDrivenSwimmingMode::SimulationTick_Implementation(const FSimulationTickParams& Params, FMoverTickEndData& OutputState)
 {
 	const UMoverComponent* MoverComp = GetMoverComponent();
 
diff --git a/Engine/Plugins/Experimental/Mover/Source/Mover/Private/PhysicsMover/Modes/PhysicsDrivenWalkingMode.cpp b/Engine/Plugins/Experimental/Mover/Source/Mover/Private/PhysicsMover/Modes/PhysicsDrivenWalkingMode.cpp
--- a/Engine/Plugins/Experimental/Mover/Source/Mover/Private/PhysicsMover/Modes/PhysicsDrivenWalkingMode.cpp
+++ b/Engine/Plugins/Experimental/Mover/Source/Mover/Private/PhysicsMover/Modes/PhysicsDrivenWalkingMode.cpp
@@ -297,7 +297,7 @@ void UPhysicsDrivenWalkingMode::OnRegistered(const FName ModeName)
 	}
 }
 
-void UPhysicsDrivenWalkingMode::OnSimulationTick(const FSimulationTickParams& Params, FMoverTickEndData& OutputState)
+void UPhysicsDrivenWalkingMode::SimulationTick_Implementation(const FSimulationTickParams& Params, FMoverTickEndData& OutputState)
 {
 	const UMoverComponent* MoverComp = GetMoverComponent();
 	const FMoverTickStartData& StartState = Params.StartState;
diff --git a/Engine/Plugins/Experimental/Mover/Source/Mover/Private/PhysicsMover/PathedMovement/PathedMovementMode.cpp b/Engine/Plugins/Experimental/Mover/Source/Mover/Private/PhysicsMover/PathedMovement/PathedMovementMode.cpp
--- a/Engine/Plugins/Experimental/Mover/Source/Mover/Private/PhysicsMover/PathedMovement/PathedMovementMode.cpp
+++ b/Engine/Plugins/Experimental/Mover/Source/Mover/Private/PhysicsMover/PathedMovement/PathedMovementMode.cpp
@@ -75,7 +75,7 @@ void UPathedPhysicsMovementMode::PostEditChangeChainProperty(FPropertyChangedCha
 }
 #endif
 
-void UPathedPhysicsMovementMode::OnSimulationTick(const FSimulationTickParams& Params, FMoverTickEndData& OutputState)
+void UPathedPhysicsMovementMode::SimulationTick_Implementation(const FSimulationTickParams& Params, FMoverTickEndData& OutputState)
 {
 	Chaos::EnsureIsInPhysicsThreadContext();
 	
diff --git a/Engine/Plugins/Experimental/Mover/Source/Mover/Private/PhysicsMover/Transitions/PhysicsJumpCheck.cpp b/Engine/Plugins/Experimental/Mover/Source/Mover/Private/PhysicsMover/Transitions/PhysicsJumpCheck.cpp
--- a/Engine/Plugins/Experimental/Mover/Source/Mover/Private/PhysicsMover/Transitions/PhysicsJumpCheck.cpp
+++ b/Engine/Plugins/Experimental/Mover/Source/Mover/Private/PhysicsMover/Transitions/PhysicsJumpCheck.cpp
@@ -40,7 +40,7 @@ UPhysicsJumpCheck::UPhysicsJumpCheck(const FObjectInitializer& ObjectInitializer
 	TransitionToMode = DefaultModeNames::Falling;
 }
 
-FTransitionEvalResult UPhysicsJumpCheck::OnEvaluate(const FSimulationTickParams& Params) const
+FTransitionEvalResult UPhysicsJumpCheck::Evaluate_Implementation(const FSimulationTickParams& Params) const
 {
 	FTransitionEvalResult EvalResult; 
 
@@ -56,7 +56,7 @@ FTransitionEvalResult UPhysicsJumpCheck::OnEvaluate(const FSimulationTickParams&
 	return EvalResult;
 }
 
-void UPhysicsJumpCheck::OnTrigger(const FSimulationTickParams& Params)
+void UPhysicsJumpCheck::Trigger_Implementation(const FSimulationTickParams& Params)
 {
 	TSharedPtr<FJumpImpulseEffect> JumpMove = MakeShared<FJumpImpulseEffect>();
 
diff --git a/Engine/Plugins/Experimental/Mover/Source/Mover/Public/Backends/MoverPathedPhysicsLiaison.h b/Engine/Plugins/Experimental/Mover/Source/Mover/Public/Backends/MoverPathedPhysicsLiaison.h
--- a/Engine/Plugins/Experimental/Mover/Source/Mover/Public/Backends/MoverPathedPhysicsLiaison.h
+++ b/Engine/Plugins/Experimental/Mover/Source/Mover/Public/Backends/MoverPathedPhysicsLiaison.h
@@ -31,6 +31,7 @@ public:
 	virtual void OnDestroyPhysicsState() override;
 	
 	virtual void ConsumeOutput_External(const FPhysicsMoverAsyncOutput& Output, const double OutputTimeInSeconds) override;
+	virtual void PostPhysicsUpdate_External() override;
 
 	const FTransform& GetPathOrigin() const { return Inputs_External.PathOrigin; }
 	void SetPathOrigin(const FTransform& NewPathOrigin);
diff --git a/Engine/Plugins/Experimental/Mover/Source/Mover/Public/DefaultMovementSet/Modes/AsyncFallingMode.h b/Engine/Plugins/Experimental/Mover/Source/Mover/Public/DefaultMovementSet/Modes/AsyncFallingMode.h
--- a/Engine/Plugins/Experimental/Mover/Source/Mover/Public/DefaultMovementSet/Modes/AsyncFallingMode.h
+++ b/Engine/Plugins/Experimental/Mover/Source/Mover/Public/DefaultMovementSet/Modes/AsyncFallingMode.h
@@ -22,11 +22,10 @@ class MOVER_API UAsyncFallingMode : public UBaseMovementMode
 public:
 	UAsyncFallingMode(const FObjectInitializer& ObjectInitializer);
 
-	UFUNCTION(BlueprintCallable, Category=Mover)
-	virtual void OnGenerateMove(const FMoverTickStartData& StartState, const FMoverTimeStep& TimeStep, FProposedMove& OutProposedMove) const override;
-
-	UFUNCTION(BlueprintCallable, Category=Mover)
-	virtual void OnSimulationTick(const FSimulationTickParams& Params, FMoverTickEndData& OutputState) override;
+	virtual void OnRegistered(const FName ModeName) override;
+    virtual void OnUnregistered() override;
+	virtual void GenerateMove_Implementation(const FMoverTickStartData& StartState, const FMoverTimeStep& TimeStep, FProposedMove& OutProposedMove) const override;
+	virtual void SimulationTick_Implementation(const FSimulationTickParams& Params, FMoverTickEndData& OutputState) override;
 
 	// Broadcast when this actor lands on a valid surface.
 	UPROPERTY(BlueprintAssignable, Category = Mover)
@@ -89,9 +88,6 @@ protected:
 	 */
 	UFUNCTION(BlueprintCallable, Category=Mover)
 	virtual void ProcessLanded(const FFloorCheckResult& FloorResult, FVector& Velocity, FRelativeBaseInfo& BaseInfo, FMoverTickEndData& TickEndData) const;
-	
-	virtual void OnRegistered(const FName ModeName) override;
-	virtual void OnUnregistered() override;
 
 	void CaptureFinalState(const FVector FinalLocation, const FRotator FinalRotation, const FFloorCheckResult& FloorResult, float DeltaSeconds, float DeltaSecondsUsed, FMoverDefaultSyncState& OutputSyncState, FMoverTickEndData& TickEndData, FMovementRecord& Record) const;
 
diff --git a/Engine/Plugins/Experimental/Mover/Source/Mover/Public/DefaultMovementSet/Modes/AsyncFlyingMode.h b/Engine/Plugins/Experimental/Mover/Source/Mover/Public/DefaultMovementSet/Modes/AsyncFlyingMode.h
--- a/Engine/Plugins/Experimental/Mover/Source/Mover/Public/DefaultMovementSet/Modes/AsyncFlyingMode.h
+++ b/Engine/Plugins/Experimental/Mover/Source/Mover/Public/DefaultMovementSet/Modes/AsyncFlyingMode.h
@@ -21,15 +21,11 @@ class MOVER_API UAsyncFlyingMode : public UBaseMovementMode
 
 
 public:
-	UFUNCTION(BlueprintCallable, Category=Mover)
-	virtual void OnGenerateMove(const FMoverTickStartData& StartState, const FMoverTimeStep& TimeStep, FProposedMove& OutProposedMove) const override;
-
-	UFUNCTION(BlueprintCallable, Category=Mover)
-	virtual void OnSimulationTick(const FSimulationTickParams& Params, FMoverTickEndData& OutputState) override;
-
-protected:
 	virtual void OnRegistered(const FName ModeName) override;
 	virtual void OnUnregistered() override;
+	virtual void GenerateMove_Implementation(const FMoverTickStartData& StartState, const FMoverTimeStep& TimeStep, FProposedMove& OutProposedMove) const override;
+	virtual void SimulationTick_Implementation(const FSimulationTickParams& Params, FMoverTickEndData& OutputState) override;
 
+protected:
 	TObjectPtr<const UCommonLegacyMovementSettings> CommonLegacySettings;
 };
diff --git a/Engine/Plugins/Experimental/Mover/Source/Mover/Public/DefaultMovementSet/Modes/AsyncNavWalkingMode.h b/Engine/Plugins/Experimental/Mover/Source/Mover/Public/DefaultMovementSet/Modes/AsyncNavWalkingMode.h
--- a/Engine/Plugins/Experimental/Mover/Source/Mover/Public/DefaultMovementSet/Modes/AsyncNavWalkingMode.h
+++ b/Engine/Plugins/Experimental/Mover/Source/Mover/Public/DefaultMovementSet/Modes/AsyncNavWalkingMode.h
@@ -27,11 +27,10 @@ class MOVER_API UAsyncNavWalkingMode : public UBaseMovementMode
 public:
 	UAsyncNavWalkingMode();
 
-	UFUNCTION(BlueprintCallable, Category = Mover)
-	virtual void OnGenerateMove(const FMoverTickStartData& StartState, const FMoverTimeStep& TimeStep, FProposedMove& OutProposedMove) const override;
-
-	UFUNCTION(BlueprintCallable, Category = Mover)
-	virtual void OnSimulationTick(const FSimulationTickParams& Params, FMoverTickEndData& OutputState) override;
+	virtual void Activate_Implementation() override;
+	virtual void Deactivate_Implementation() override;
+	virtual void SimulationTick_Implementation(const FSimulationTickParams& Params, FMoverTickEndData& OutputState) override;
+	virtual void GenerateMove_Implementation(const FMoverTickStartData& StartState, const FMoverTimeStep& TimeStep, FProposedMove& OutProposedMove) const override;
 
 	/**
 	 * Whether or not the actor should sweep for collision geometry while walking.
@@ -102,11 +101,7 @@ protected:
 
 	/** Switch collision settings for NavWalking mode (ignore world collisions) */
 	virtual void SetCollisionForNavWalking(bool bEnable);
-
-	virtual void OnActivate() override;
-
-	virtual void OnDeactivate() override;
-
+	
 	/** Get Navigation data for the actor. Returns null if there is no associated nav data. */
 	const INavigationDataInterface* GetNavData() const;
 
diff --git a/Engine/Plugins/Experimental/Mover/Source/Mover/Public/DefaultMovementSet/Modes/AsyncWalkingMode.h b/Engine/Plugins/Experimental/Mover/Source/Mover/Public/DefaultMovementSet/Modes/AsyncWalkingMode.h
--- a/Engine/Plugins/Experimental/Mover/Source/Mover/Public/DefaultMovementSet/Modes/AsyncWalkingMode.h
+++ b/Engine/Plugins/Experimental/Mover/Source/Mover/Public/DefaultMovementSet/Modes/AsyncWalkingMode.h
@@ -25,11 +25,9 @@ class MOVER_API UAsyncWalkingMode : public UBaseMovementMode
 public:
 	UAsyncWalkingMode(const FObjectInitializer& ObjectInitializer);
 	
-	UFUNCTION(BlueprintCallable, Category=Mover)
-	virtual void OnGenerateMove(const FMoverTickStartData& StartState, const FMoverTimeStep& TimeStep, FProposedMove& OutProposedMove) const override;
+	virtual void GenerateMove_Implementation(const FMoverTickStartData& StartState, const FMoverTimeStep& TimeStep, FProposedMove& OutProposedMove) const override;
 
-	UFUNCTION(BlueprintCallable, Category=Mover)
-	virtual void OnSimulationTick(const FSimulationTickParams& Params, FMoverTickEndData& OutputState) override;
+	virtual void SimulationTick_Implementation(const FSimulationTickParams& Params, FMoverTickEndData& OutputState) override;
 
 	// Returns the active turn generator. Note: you will need to cast the return value to the generator you expect to get, it can also be none
 	UFUNCTION(BlueprintPure, Category=Mover)
diff --git a/Engine/Plugins/Experimental/Mover/Source/Mover/Public/DefaultMovementSet/Modes/FallingMode.h b/Engine/Plugins/Experimental/Mover/Source/Mover/Public/DefaultMovementSet/Modes/FallingMode.h
--- a/Engine/Plugins/Experimental/Mover/Source/Mover/Public/DefaultMovementSet/Modes/FallingMode.h
+++ b/Engine/Plugins/Experimental/Mover/Source/Mover/Public/DefaultMovementSet/Modes/FallingMode.h
@@ -22,11 +22,9 @@ class MOVER_API UFallingMode : public UBaseMovementMode
 
 
 public:
-	UFUNCTION(BlueprintCallable, Category=Mover)
-	virtual void OnGenerateMove(const FMoverTickStartData& StartState, const FMoverTimeStep& TimeStep, FProposedMove& OutProposedMove) const override;
+	virtual void GenerateMove_Implementation(const FMoverTickStartData& StartState, const FMoverTimeStep& TimeStep, FProposedMove& OutProposedMove) const override;
 
-	UFUNCTION(BlueprintCallable, Category=Mover)
-	virtual void OnSimulationTick(const FSimulationTickParams& Params, FMoverTickEndData& OutputState) override;
+	virtual void SimulationTick_Implementation(const FSimulationTickParams& Params, FMoverTickEndData& OutputState) override;
 
 	// Broadcast when this actor lands on a valid surface.
 	UPROPERTY(BlueprintAssignable, Category = Mover)
diff --git a/Engine/Plugins/Experimental/Mover/Source/Mover/Public/DefaultMovementSet/Modes/FlyingMode.h b/Engine/Plugins/Experimental/Mover/Source/Mover/Public/DefaultMovementSet/Modes/FlyingMode.h
--- a/Engine/Plugins/Experimental/Mover/Source/Mover/Public/DefaultMovementSet/Modes/FlyingMode.h
+++ b/Engine/Plugins/Experimental/Mover/Source/Mover/Public/DefaultMovementSet/Modes/FlyingMode.h
@@ -20,11 +20,9 @@ class MOVER_API UFlyingMode : public UBaseMovementMode
 
 
 public:
-	UFUNCTION(BlueprintCallable, Category=Mover)
-	virtual void OnGenerateMove(const FMoverTickStartData& StartState, const FMoverTimeStep& TimeStep, FProposedMove& OutProposedMove) const override;
+	virtual void GenerateMove_Implementation(const FMoverTickStartData& StartState, const FMoverTimeStep& TimeStep, FProposedMove& OutProposedMove) const override;
 
-	UFUNCTION(BlueprintCallable, Category=Mover)
-	virtual void OnSimulationTick(const FSimulationTickParams& Params, FMoverTickEndData& OutputState) override;
+	virtual void SimulationTick_Implementation(const FSimulationTickParams& Params, FMoverTickEndData& OutputState) override;
 
 protected:
 	virtual void OnRegistered(const FName ModeName) override;
diff --git a/Engine/Plugins/Experimental/Mover/Source/Mover/Public/DefaultMovementSet/Modes/NavWalkingMode.h b/Engine/Plugins/Experimental/Mover/Source/Mover/Public/DefaultMovementSet/Modes/NavWalkingMode.h
--- a/Engine/Plugins/Experimental/Mover/Source/Mover/Public/DefaultMovementSet/Modes/NavWalkingMode.h
+++ b/Engine/Plugins/Experimental/Mover/Source/Mover/Public/DefaultMovementSet/Modes/NavWalkingMode.h
@@ -31,11 +31,10 @@ class MOVER_API UNavWalkingMode : public UBaseMovementMode
 public:
 	UNavWalkingMode();
 	
-	UFUNCTION(BlueprintCallable, Category=Mover)
-	virtual void OnGenerateMove(const FMoverTickStartData& StartState, const FMoverTimeStep& TimeStep, FProposedMove& OutProposedMove) const override;
-
-	UFUNCTION(BlueprintCallable, Category=Mover)
-	virtual void OnSimulationTick(const FSimulationTickParams& Params, FMoverTickEndData& OutputState) override;
+	virtual void Activate_Implementation() override;
+	virtual void Deactivate_Implementation() override;
+	virtual void SimulationTick_Implementation(const FSimulationTickParams& Params, FMoverTickEndData& OutputState) override;
+	virtual void GenerateMove_Implementation(const FMoverTickStartData& StartState, const FMoverTimeStep& TimeStep, FProposedMove& OutProposedMove) const override;
 
 	/**
 	 * Whether or not the actor should sweep for collision geometry while walking.
@@ -108,10 +107,6 @@ protected:
 	/** Switch collision settings for NavWalking mode (ignore world collisions) */
 	virtual void SetCollisionForNavWalking(bool bEnable);
 	
-	virtual void OnActivate() override;
-
-	virtual void OnDeactivate() override;
-	
 	/** Get Navigation data for the actor. Returns null if there is no associated nav data. */
 	const INavigationDataInterface* GetNavData() const;
 
diff --git a/Engine/Plugins/Experimental/Mover/Source/Mover/Public/DefaultMovementSet/Modes/SwimmingMode.h b/Engine/Plugins/Experimental/Mover/Source/Mover/Public/DefaultMovementSet/Modes/SwimmingMode.h
--- a/Engine/Plugins/Experimental/Mover/Source/Mover/Public/DefaultMovementSet/Modes/SwimmingMode.h
+++ b/Engine/Plugins/Experimental/Mover/Source/Mover/Public/DefaultMovementSet/Modes/SwimmingMode.h
@@ -196,11 +196,9 @@ class MOVER_API USwimmingMode : public UBaseMovementMode
 
 public:
 
-	UFUNCTION(BlueprintCallable, Category=Mover)
-	virtual void OnGenerateMove(const FMoverTickStartData& StartState, const FMoverTimeStep& TimeStep, FProposedMove& OutProposedMove) const override;
+	virtual void GenerateMove_Implementation(const FMoverTickStartData& StartState, const FMoverTimeStep& TimeStep, FProposedMove& OutProposedMove) const override;
 
-	UFUNCTION(BlueprintCallable, Category=Mover)
-	virtual void OnSimulationTick(const FSimulationTickParams& Params, FMoverTickEndData& OutputState) override;
+	virtual void SimulationTick_Implementation(const FSimulationTickParams& Params, FMoverTickEndData& OutputState) override;
 
 	UPROPERTY(EditAnywhere, BlueprintReadOnly, Category = "Water Control")
 	FSwimmingControlSettings SurfaceSwimmingWaterControlSettings;
diff --git a/Engine/Plugins/Experimental/Mover/Source/Mover/Public/DefaultMovementSet/Modes/WalkingMode.h b/Engine/Plugins/Experimental/Mover/Source/Mover/Public/DefaultMovementSet/Modes/WalkingMode.h
--- a/Engine/Plugins/Experimental/Mover/Source/Mover/Public/DefaultMovementSet/Modes/WalkingMode.h
+++ b/Engine/Plugins/Experimental/Mover/Source/Mover/Public/DefaultMovementSet/Modes/WalkingMode.h
@@ -35,11 +35,9 @@ class MOVER_API UWalkingMode : public UBaseMovementMode
 
 public:
 	
-	UFUNCTION(BlueprintCallable, Category=Mover)
-	virtual void OnGenerateMove(const FMoverTickStartData& StartState, const FMoverTimeStep& TimeStep, FProposedMove& OutProposedMove) const override;
+	virtual void GenerateMove_Implementation(const FMoverTickStartData& StartState, const FMoverTimeStep& TimeStep, FProposedMove& OutProposedMove) const override;
 
-	UFUNCTION(BlueprintCallable, Category=Mover)
-	virtual void OnSimulationTick(const FSimulationTickParams& Params, FMoverTickEndData& OutputState) override;
+	virtual void SimulationTick_Implementation(const FSimulationTickParams& Params, FMoverTickEndData& OutputState) override;
 
 	// Returns the active turn generator. Note: you will need to cast the return value to the generator you expect to get, it can also be none
 	UFUNCTION(BlueprintPure, Category=Mover)
diff --git a/Engine/Plugins/Experimental/Mover/Source/Mover/Public/MovementMode.h b/Engine/Plugins/Experimental/Mover/Source/Mover/Public/MovementMode.h
--- a/Engine/Plugins/Experimental/Mover/Source/Mover/Public/MovementMode.h
+++ b/Engine/Plugins/Experimental/Mover/Source/Mover/Public/MovementMode.h
@@ -36,28 +36,52 @@ public:
 UCLASS(Abstract, Within = MoverComponent, Blueprintable, BlueprintType, EditInlineNew, DefaultToInstanced)
 class MOVER_API UBaseMovementMode : public UObject
 {
-	GENERATED_UCLASS_BODY()
+	GENERATED_BODY()
 
 public:
+	virtual UWorld* GetWorld() const override;
+	
+	virtual void OnRegistered(const FName ModeName);
+	virtual void OnUnregistered();
+
+	UFUNCTION(BlueprintNativeEvent)
+	void Activate();
+	
+	UFUNCTION(BlueprintNativeEvent)
+	void Deactivate();
 	
-	void DoRegister(const FName ModeName);
-	void DoUnregister();
-	void DoGenerateMove(const FMoverTickStartData& StartState, const FMoverTimeStep& TimeStep, FProposedMove& OutProposedMove) const;
-	void DoSimulationTick(const FSimulationTickParams& Params, FMoverTickEndData& OutputState);
-	void DoActivate();
-	void DoDeactivate();
-
-	UFUNCTION(BlueprintCallable, Category=Mover, meta=(DisplayName="Get Mover Component"))
-	UMoverComponent* GetMoverComponent() const;
-
-	/** Templated convenience version of GetMoverComponent which checks the type is as presumed. */
-	template<class T = UMoverComponent>
-	T* GetMoverComponentChecked() const
+	UFUNCTION(BlueprintNativeEvent, meta = (DisplayName = "Generate Move"))
+	void GenerateMove(const FMoverTickStartData& StartState, const FMoverTimeStep& TimeStep, FProposedMove& OutProposedMove) const;
+
+	UFUNCTION(BlueprintNativeEvent, meta = (DisplayName = "Simulation Tick"))
+	void SimulationTick(const FSimulationTickParams& Params, FMoverTickEndData& OutputState);
+	
+	/** Gets the MoverComponent that owns this movement mode */
+	UFUNCTION(BlueprintCallable, Category=Mover, meta=(DisplayName="Get Mover Component", ScriptName = GetMoverComponent))
+	UMoverComponent* K2_GetMoverComponent() const;
+
+	template<typename MoverT = UMoverComponent UE_REQUIRES(std::is_base_of_v<MoverT, UMoverComponent>)>
+	MoverT* GetMoverComponent() const
+	{
+		return Cast<MoverT>(GetOuterUMoverComponent());
+	}
+
+	template<typename MoverT = UMoverComponent UE_REQUIRES(std::is_base_of_v<MoverT, UMoverComponent>)>
+	MoverT& GetMoverComponentChecked() const
 	{
-		static_assert(TPointerIsConvertibleFromTo<T, const UMoverComponent>::Value, "'T' template parameter to GetMoverComponentChecked must be derived from UMoverComponent");
-		return CastChecked<T>(GetOuterUMoverComponent());
+		return *CastChecked<MoverT>(GetOuterUMoverComponent());
 	}
 
+	/**
+   	 * Check Movement Mode for a gameplay tag.
+   	 *
+   	 * @param TagToFind			Tag to check on the Mover systems
+   	 * @param bExactMatch		If true, the tag has to be exactly present, if false then TagToFind will include it's parent tags while matching
+   	 * 
+   	 * @return True if the TagToFind was found
+   	 */
+	virtual bool HasGameplayTag(FGameplayTag TagToFind, bool bExactMatch) const;
+
 #if WITH_EDITOR
 	virtual EDataValidationResult IsDataValid(class FDataValidationContext& Context) const override;
 #endif // WITH_EDITOR
@@ -73,54 +97,20 @@ public:
 	/** A list of gameplay tags associated with this movement mode */
 	UPROPERTY(EditDefaultsOnly, Category = Mover)
 	FGameplayTagContainer GameplayTags;
-
-	/**
-   	 * Check Movement Mode for a gameplay tag.
-   	 *
-   	 * @param TagToFind			Tag to check on the Mover systems
-   	 * @param bExactMatch		If true, the tag has to be exactly present, if false then TagToFind will include it's parent tags while matching
-   	 * 
-   	 * @return True if the TagToFind was found
-   	 */
-	virtual bool HasGameplayTag(FGameplayTag TagToFind, bool bExactMatch) const;
 	
 protected:
-
-	virtual void OnRegistered(const FName ModeName);
-
-	UFUNCTION(BlueprintImplementableEvent, meta = (DisplayName = "OnRegistered", ScriptName = "OnRegistered"))
+	UFUNCTION(BlueprintImplementableEvent, meta = (DisplayName = "On Registered", ScriptName = "OnRegistered"))
 	void K2_OnRegistered(const FName ModeName);
 
-
-	virtual void OnUnregistered();
-
-	UFUNCTION(BlueprintImplementableEvent, meta = (DisplayName = "OnUnregistered", ScriptName = "OnUnregistered"))
+	UFUNCTION(BlueprintImplementableEvent, meta = (DisplayName = "On Unregistered", ScriptName = "OnUnregistered"))
 	void K2_OnUnregistered();
 
-
-	virtual void OnGenerateMove(const FMoverTickStartData& StartState, const FMoverTimeStep& TimeStep, FProposedMove& OutProposedMove) const;
-
-	UFUNCTION(BlueprintImplementableEvent, DisplayName = "OnGenerateMove", meta = (ScriptName = "OnGenerateMove"))
-	FProposedMove K2_OnGenerateMove(const FMoverTickStartData& StartState, const FMoverTimeStep& TimeStep) const;
-
-	virtual void OnSimulationTick(const FSimulationTickParams& Params, FMoverTickEndData& OutputState);
-	
-	UFUNCTION(BlueprintImplementableEvent, DisplayName = "OnSimulationTick", meta = (ScriptName = "OnSimulationTick"))
-	FMoverTickEndData K2_OnSimulationTick(const FSimulationTickParams& Params);
-
-	virtual void OnActivate();
-
-	UFUNCTION(BlueprintImplementableEvent, DisplayName = "OnActivate", meta = (ScriptName = "OnActivate"))
-	void K2_OnActivate();
-
-	virtual void OnDeactivate();
-
-	UFUNCTION(BlueprintImplementableEvent, DisplayName = "OnDeactivate", meta = (ScriptName = "OnDeactivate"))
-	void K2_OnDeactivate();
-
-private:
-	bool bHasBlueprintGenerateMove = false;
-	bool bHasBlueprintSimulationTick = false;
-	bool bHasBlueprintOnActivate = false;
-	bool bHasBlueprintOnDeactivate = false;
+	UE_DEPRECATED(5.6, "OnActivate() has been replaced with an Activate() BlueprintNativeEvent. Rename your override to Activate_Implementation().")
+	virtual void OnActivate() {}
+	UE_DEPRECATED(5.6, "OnDeactivate() has been replaced with a Deactivate() BlueprintNativeEvent. Rename your override to Deactivate_Implementation().")
+	virtual void OnDeactivate() {}
+	UE_DEPRECATED(5.6, "OnGenerateMove() has been replaced with a GenerateMove() BlueprintNativeEvent. Rename your override to GenerateMove_Implementation().")
+	virtual void OnGenerateMove(const FMoverTickStartData& StartState, const FMoverTimeStep& TimeStep, FProposedMove& OutProposedMove) const {}
+	UE_DEPRECATED(5.6, "OnSimulationTick() has been replaced with a SimulationTick() BlueprintNativeEvent. Rename your override to SimulationTick_Implementation().")
+	virtual void OnSimulationTick(const FSimulationTickParams& Params, FMoverTickEndData& OutputState) {}
 };
\ No newline at end of file
diff --git a/Engine/Plugins/Experimental/Mover/Source/Mover/Public/MovementModeStateMachine.h b/Engine/Plugins/Experimental/Mover/Source/Mover/Public/MovementModeStateMachine.h
--- a/Engine/Plugins/Experimental/Mover/Source/Mover/Public/MovementModeStateMachine.h
+++ b/Engine/Plugins/Experimental/Mover/Source/Mover/Public/MovementModeStateMachine.h
@@ -21,7 +21,7 @@ class UNullMovementMode : public UBaseMovementMode
 	GENERATED_UCLASS_BODY()
 
 public:
-	virtual void OnSimulationTick(const FSimulationTickParams& Params, FMoverTickEndData& OutputState) override;
+	virtual void SimulationTick_Implementation(const FSimulationTickParams& Params, FMoverTickEndData& OutputState) override;
 
 	const static FName NullModeName;
 };
@@ -121,8 +121,8 @@ private:
 	 GENERATED_UCLASS_BODY()
 
  public:
-	 virtual FTransitionEvalResult OnEvaluate(const FSimulationTickParams& Params) const override;
-	 virtual void OnTrigger(const FSimulationTickParams& Params) override;
+	 virtual FTransitionEvalResult Evaluate_Implementation(const FSimulationTickParams& Params) const override;
+	 virtual void Trigger_Implementation(const FSimulationTickParams& Params) override;
 
 
 	 bool IsSet() const { return !NextMode.IsNone(); }
diff --git a/Engine/Plugins/Experimental/Mover/Source/Mover/Public/MovementModeTransition.h b/Engine/Plugins/Experimental/Mover/Source/Mover/Public/MovementModeTransition.h
--- a/Engine/Plugins/Experimental/Mover/Source/Mover/Public/MovementModeTransition.h
+++ b/Engine/Plugins/Experimental/Mover/Source/Mover/Public/MovementModeTransition.h
@@ -32,29 +32,38 @@ struct MOVER_API FTransitionEvalResult
 /**
  * Base class for all transitions
  */
-UCLASS(Abstract, Blueprintable, BlueprintType, EditInlineNew)
+UCLASS(Abstract, Blueprintable, BlueprintType, EditInlineNew, DefaultToInstanced)
 class MOVER_API UBaseMovementModeTransition : public UObject
 {
-	GENERATED_UCLASS_BODY()
+	GENERATED_BODY()
 
 public:
-	void DoRegister();
-	void DoUnregister();
+	virtual UWorld* GetWorld() const override;
+	
+	virtual void OnRegistered();
+	virtual void OnUnregistered();
 
+	/** Gets the MoverComponent that ultimately owns this transition */
 	UFUNCTION(BlueprintCallable, Category = Mover, meta = (DisplayName = "Get Mover Component"))
-	UMoverComponent* GetMoverComponent() const;
+	UMoverComponent* K2_GetMoverComponent() const;
 
-	/** Templated convenience version of GetMoverComponent which checks the type is as presumed. */
-	template<class T>
-	T* GetMoverComponentChecked() const
+	template<typename MoverT = UMoverComponent UE_REQUIRES(std::is_base_of_v<UMoverComponent, MoverT>)>
+	MoverT* GetMoverComponent() const
 	{
-		static_assert(TPointerIsConvertibleFromTo<T, const UMoverComponent>::Value, "'T' template parameter to GetMoverComponentChecked must be derived from UMoverComponent");
-		return CastChecked<T>(GetMoverComponent());
+		return Cast<MoverT>(K2_GetMoverComponent());
+	}
+	
+	template<typename MoverT = UMoverComponent UE_REQUIRES(std::is_base_of_v<UMoverComponent, MoverT>)>
+	MoverT& GetMoverComponentChecked() const
+	{
+		return *CastChecked<MoverT>(K2_GetMoverComponent());
 	}
 
-	FTransitionEvalResult DoEvaluate(const FSimulationTickParams& Params) const;
+	UFUNCTION(BlueprintNativeEvent, Category = MovementModeTransition)
+	FTransitionEvalResult Evaluate(const FSimulationTickParams& Params) const;
 
-	void DoTrigger(const FSimulationTickParams& Params);
+	UFUNCTION(BlueprintNativeEvent, Category = MovementModeTransition)
+	void Trigger(const FSimulationTickParams& Params);
 
 #if WITH_EDITOR
 	virtual EDataValidationResult IsDataValid(class FDataValidationContext& Context) const override;
@@ -66,27 +75,19 @@ public:
 
 	/** Whether this transition should reenter a mode if it evaluates true and wants to transition into a mode the actor is already in */
 	UPROPERTY(EditDefaultsOnly, BlueprintReadWrite, Category = Mover)
-	bool bAllowModeReentry;
+	bool bAllowModeReentry = false;
 	
 protected:
-	virtual void OnRegistered();
-
-	virtual void OnUnregistered();
+	UFUNCTION(BlueprintImplementableEvent, Category = MovementModeTransition, meta = (DisplayName = "On Registered", ScriptName = "OnRegistered"))
+	void K2_OnRegistered();
+	
+	UFUNCTION(BlueprintImplementableEvent, Category = MovementModeTransition, meta = (DisplayName = "On Unregistered", ScriptName = "OnUnregistered"))
+	void K2_OnUnregistered();
 
+	UE_DEPRECATED(5.6, "OnEvaluate() has been replaced with an Evaluate() BlueprintNativeEvent. Rename your override to Evaluate_Implementation().")
 	virtual FTransitionEvalResult OnEvaluate(const FSimulationTickParams& Params) const;
-
-	UFUNCTION(BlueprintImplementableEvent, DisplayName = "OnEvaluate", meta = (ScriptName = "OnEvaluate"))
-	FTransitionEvalResult K2_OnEvaluate(const FSimulationTickParams& Params) const;
-
-
-	virtual void OnTrigger(const FSimulationTickParams& Params);
-
-	UFUNCTION(BlueprintImplementableEvent, DisplayName = "OnTrigger", meta = (ScriptName = "OnTrigger"))
-	void K2_OnTrigger(const FSimulationTickParams& Params);
-
-private:
-	bool bHasBlueprintEvaluate = false;
-	bool bHasBlueprintTrigger = false;
+	UE_DEPRECATED(5.6, "OnTrigger() has been replaced with a Trigger() BlueprintNativeEvent. Rename your override to Trigger_Implementation().")
+	virtual void OnTrigger(const FSimulationTickParams& Params) {}
 };
 
 
diff --git a/Engine/Plugins/Experimental/Mover/Source/Mover/Public/PhysicsMover/Modes/PhysicsDrivenFallingMode.h b/Engine/Plugins/Experimental/Mover/Source/Mover/Public/PhysicsMover/Modes/PhysicsDrivenFallingMode.h
--- a/Engine/Plugins/Experimental/Mover/Source/Mover/Public/PhysicsMover/Modes/PhysicsDrivenFallingMode.h
+++ b/Engine/Plugins/Experimental/Mover/Source/Mover/Public/PhysicsMover/Modes/PhysicsDrivenFallingMode.h
@@ -21,7 +21,7 @@ class MOVER_API UPhysicsDrivenFallingMode : public UFallingMode, public IPhysics
 
 public:
 	virtual void OnRegistered(const FName ModeName) override;
-	virtual void OnSimulationTick(const FSimulationTickParams& Params, FMoverTickEndData& OutputState) override;
+	virtual void SimulationTick_Implementation(const FSimulationTickParams& Params, FMoverTickEndData& OutputState) override;
 	virtual void UpdateConstraintSettings(Chaos::FCharacterGroundConstraint& Constraint) const override;
 
 #if WITH_EDITOR
diff --git a/Engine/Plugins/Experimental/Mover/Source/Mover/Public/PhysicsMover/Modes/PhysicsDrivenFlyingMode.h b/Engine/Plugins/Experimental/Mover/Source/Mover/Public/PhysicsMover/Modes/PhysicsDrivenFlyingMode.h
--- a/Engine/Plugins/Experimental/Mover/Source/Mover/Public/PhysicsMover/Modes/PhysicsDrivenFlyingMode.h
+++ b/Engine/Plugins/Experimental/Mover/Source/Mover/Public/PhysicsMover/Modes/PhysicsDrivenFlyingMode.h
@@ -21,7 +21,7 @@ class MOVER_API UPhysicsDrivenFlyingMode : public UFlyingMode, public IPhysicsCh
 
 public:
 
-	virtual void OnSimulationTick(const FSimulationTickParams& Params, FMoverTickEndData& OutputState) override;
+	virtual void SimulationTick_Implementation(const FSimulationTickParams& Params, FMoverTickEndData& OutputState) override;
 
 	virtual void UpdateConstraintSettings(Chaos::FCharacterGroundConstraint& Constraint) const override;
 
diff --git a/Engine/Plugins/Experimental/Mover/Source/Mover/Public/PhysicsMover/Modes/PhysicsDrivenSwimmingMode.h b/Engine/Plugins/Experimental/Mover/Source/Mover/Public/PhysicsMover/Modes/PhysicsDrivenSwimmingMode.h
--- a/Engine/Plugins/Experimental/Mover/Source/Mover/Public/PhysicsMover/Modes/PhysicsDrivenSwimmingMode.h
+++ b/Engine/Plugins/Experimental/Mover/Source/Mover/Public/PhysicsMover/Modes/PhysicsDrivenSwimmingMode.h
@@ -21,7 +21,7 @@ class MOVER_API UPhysicsDrivenSwimmingMode : public USwimmingMode, public IPhysi
 
 public:
 	virtual void OnRegistered(const FName ModeName) override;
-	virtual void OnSimulationTick(const FSimulationTickParams& Params, FMoverTickEndData& OutputState) override;
+	virtual void SimulationTick_Implementation(const FSimulationTickParams& Params, FMoverTickEndData& OutputState) override;
 	virtual bool AttemptJump(float UpwardsSpeed, FMoverTickEndData& Output) override;
 
 	virtual void UpdateConstraintSettings(Chaos::FCharacterGroundConstraint& Constraint) const override;
diff --git a/Engine/Plugins/Experimental/Mover/Source/Mover/Public/PhysicsMover/Modes/PhysicsDrivenWalkingMode.h b/Engine/Plugins/Experimental/Mover/Source/Mover/Public/PhysicsMover/Modes/PhysicsDrivenWalkingMode.h
--- a/Engine/Plugins/Experimental/Mover/Source/Mover/Public/PhysicsMover/Modes/PhysicsDrivenWalkingMode.h
+++ b/Engine/Plugins/Experimental/Mover/Source/Mover/Public/PhysicsMover/Modes/PhysicsDrivenWalkingMode.h
@@ -22,7 +22,7 @@ class MOVER_API UPhysicsDrivenWalkingMode : public UWalkingMode, public IPhysics
 
 public:
 	virtual void OnRegistered(const FName ModeName) override;
-	virtual void OnSimulationTick(const FSimulationTickParams& Params, FMoverTickEndData& OutputState) override;
+	virtual void SimulationTick_Implementation(const FSimulationTickParams& Params, FMoverTickEndData& OutputState) override;
 
 	virtual void UpdateConstraintSettings(Chaos::FCharacterGroundConstraint& Constraint) const override;
 	virtual void OnContactModification_Internal(const FPhysicsMoverSimulationContactModifierParams& Params, Chaos::FCollisionContactModifier& Modifier) const override;
diff --git a/Engine/Plugins/Experimental/Mover/Source/Mover/Public/PhysicsMover/PathedMovement/PathedMovementMode.h b/Engine/Plugins/Experimental/Mover/Source/Mover/Public/PhysicsMover/PathedMovement/PathedMovementMode.h
--- a/Engine/Plugins/Experimental/Mover/Source/Mover/Public/PhysicsMover/PathedMovement/PathedMovementMode.h
+++ b/Engine/Plugins/Experimental/Mover/Source/Mover/Public/PhysicsMover/PathedMovement/PathedMovementMode.h
@@ -34,7 +34,7 @@ public:
 	virtual void PostEditChangeChainProperty(struct FPropertyChangedChainEvent& PropertyChangedChainEvent) override;
 #endif
 	
-	virtual void OnSimulationTick(const FSimulationTickParams& Params, FMoverTickEndData& OutputState) override final;
+	virtual void SimulationTick_Implementation(const FSimulationTickParams& Params, FMoverTickEndData& OutputState) override final;
 
 	virtual void OnProcessInput_Internal(int32 PhysicsStep, float DeltaTime, const FPhysicsMoverAsyncInput& Input) const;
 	virtual void OnPreSimulate_Internal(const FSimulationTickParams& Params, FMoverTickEndData& OutputState) const;
diff --git a/Engine/Plugins/Experimental/Mover/Source/Mover/Public/PhysicsMover/Transitions/PhysicsJumpCheck.h b/Engine/Plugins/Experimental/Mover/Source/Mover/Public/PhysicsMover/Transitions/PhysicsJumpCheck.h
--- a/Engine/Plugins/Experimental/Mover/Source/Mover/Public/PhysicsMover/Transitions/PhysicsJumpCheck.h
+++ b/Engine/Plugins/Experimental/Mover/Source/Mover/Public/PhysicsMover/Transitions/PhysicsJumpCheck.h
@@ -15,8 +15,8 @@ class MOVER_API UPhysicsJumpCheck : public UBaseMovementModeTransition
 
 
 public:
-	virtual FTransitionEvalResult OnEvaluate(const FSimulationTickParams& Params) const override;
-	virtual void OnTrigger(const FSimulationTickParams& Params) override;
+	virtual FTransitionEvalResult Evaluate_Implementation(const FSimulationTickParams& Params) const override;
+	virtual void Trigger_Implementation(const FSimulationTickParams& Params) override;
 
 	/** Instantaneous speed induced in an actor upon jumping */
 	UPROPERTY(BlueprintReadWrite, EditAnywhere, Category = Parameters, meta = (ClampMin = "0", UIMin = "0", ForceUnits = "cm/s"))
diff --git a/Engine/Plugins/Experimental/MoverExamples/Source/MoverExamples/Private/CharacterVariants/Ziplining/ZipliningMode.cpp b/Engine/Plugins/Experimental/MoverExamples/Source/MoverExamples/Private/CharacterVariants/Ziplining/ZipliningMode.cpp
--- a/Engine/Plugins/Experimental/MoverExamples/Source/MoverExamples/Private/CharacterVariants/Ziplining/ZipliningMode.cpp
+++ b/Engine/Plugins/Experimental/MoverExamples/Source/MoverExamples/Private/CharacterVariants/Ziplining/ZipliningMode.cpp
@@ -71,7 +71,7 @@ UZipliningMode::UZipliningMode(const FObjectInitializer& ObjectInitializer)
 }
 
 
-void UZipliningMode::OnGenerateMove(const FMoverTickStartData& StartState, const FMoverTimeStep& TimeStep, FProposedMove& OutProposedMove) const
+void UZipliningMode::GenerateMove_Implementation(const FMoverTickStartData& StartState, const FMoverTimeStep& TimeStep, FProposedMove& OutProposedMove) const
 {
 	UMoverComponent* MoverComp = GetMoverComponent();
 
@@ -80,7 +80,7 @@ void UZipliningMode::OnGenerateMove(const FMoverTickStartData& StartState, const
 
 }
 
-void UZipliningMode::OnSimulationTick(const FSimulationTickParams& Params, FMoverTickEndData& OutputState)
+void UZipliningMode::SimulationTick_Implementation(const FSimulationTickParams& Params, FMoverTickEndData& OutputState)
 {
 	// Are we continuing a move or starting fresh?
 	const FZipliningState* StartingZipState = Params.StartState.SyncState.SyncStateCollection.FindDataByType<FZipliningState>();
diff --git a/Engine/Plugins/Experimental/MoverExamples/Source/MoverExamples/Private/CharacterVariants/Ziplining/ZipliningTransitions.cpp b/Engine/Plugins/Experimental/MoverExamples/Source/MoverExamples/Private/CharacterVariants/Ziplining/ZipliningTransitions.cpp
--- a/Engine/Plugins/Experimental/MoverExamples/Source/MoverExamples/Private/CharacterVariants/Ziplining/ZipliningTransitions.cpp
+++ b/Engine/Plugins/Experimental/MoverExamples/Source/MoverExamples/Private/CharacterVariants/Ziplining/ZipliningTransitions.cpp
@@ -15,7 +15,7 @@ UZiplineStartTransition::UZiplineStartTransition(const FObjectInitializer& Objec
 {
 }
 
-FTransitionEvalResult UZiplineStartTransition::OnEvaluate(const FSimulationTickParams& Params) const
+FTransitionEvalResult UZiplineStartTransition::Evaluate_Implementation(const FSimulationTickParams& Params) const
 {
 	FTransitionEvalResult EvalResult = FTransitionEvalResult::NoTransition;
 
@@ -59,7 +59,7 @@ UZiplineEndTransition::UZiplineEndTransition(const FObjectInitializer& ObjectIni
 }
 
 
-FTransitionEvalResult UZiplineEndTransition::OnEvaluate(const FSimulationTickParams& Params) const
+FTransitionEvalResult UZiplineEndTransition::Evaluate_Implementation(const FSimulationTickParams& Params) const
 {
 	FTransitionEvalResult EvalResult = FTransitionEvalResult::NoTransition;
 
@@ -74,7 +74,7 @@ FTransitionEvalResult UZiplineEndTransition::OnEvaluate(const FSimulationTickPar
 	return EvalResult;
 }
 
-void UZiplineEndTransition::OnTrigger(const FSimulationTickParams& Params)
+void UZiplineEndTransition::Trigger_Implementation(const FSimulationTickParams& Params)
 {
 	//TODO: create a small jump, using current directionality
 }
diff --git a/Engine/Plugins/Experimental/MoverExamples/Source/MoverExamples/Private/MovementBases/FollowPathMode.cpp b/Engine/Plugins/Experimental/MoverExamples/Source/MoverExamples/Private/MovementBases/FollowPathMode.cpp
--- a/Engine/Plugins/Experimental/MoverExamples/Source/MoverExamples/Private/MovementBases/FollowPathMode.cpp
+++ b/Engine/Plugins/Experimental/MoverExamples/Source/MoverExamples/Private/MovementBases/FollowPathMode.cpp
@@ -13,12 +13,12 @@ UFollowPathMode::UFollowPathMode(const FObjectInitializer& ObjectInitializer)
 {
 }
 
-void UFollowPathMode::OnGenerateMove(const FMoverTickStartData& StartState, const FMoverTimeStep& TimeStep, FProposedMove& OutProposedMove) const
+void UFollowPathMode::GenerateMove_Implementation(const FMoverTickStartData& StartState, const FMoverTimeStep& TimeStep, FProposedMove& OutProposedMove) const
 {
 	// TODO: Rework to maintain the split between generating and executing moves
 }
 
-void UFollowPathMode::OnSimulationTick(const FSimulationTickParams& Params, FMoverTickEndData& OutputState)
+void UFollowPathMode::SimulationTick_Implementation(const FSimulationTickParams& Params, FMoverTickEndData& OutputState)
 {
 	const FMoverTickStartData& StartState = Params.StartState;
 	USceneComponent* UpdatedComponent = Params.MovingComps.UpdatedComponent.Get();
diff --git a/Engine/Plugins/Experimental/MoverExamples/Source/MoverExamples/Private/MovementBases/FollowSplineMode.cpp b/Engine/Plugins/Experimental/MoverExamples/Source/MoverExamples/Private/MovementBases/FollowSplineMode.cpp
--- a/Engine/Plugins/Experimental/MoverExamples/Source/MoverExamples/Private/MovementBases/FollowSplineMode.cpp
+++ b/Engine/Plugins/Experimental/MoverExamples/Source/MoverExamples/Private/MovementBases/FollowSplineMode.cpp
@@ -46,12 +46,12 @@ UFollowSplineMode::UFollowSplineMode(const FObjectInitializer& ObjectInitializer
 {
 }
 
-void UFollowSplineMode::OnGenerateMove(const FMoverTickStartData& StartState, const FMoverTimeStep& TimeStep, FProposedMove& OutProposedMove) const
+void UFollowSplineMode::GenerateMove_Implementation(const FMoverTickStartData& StartState, const FMoverTimeStep& TimeStep, FProposedMove& OutProposedMove) const
 {
 	// TODO: Rework to maintain the split between generating and executing moves
 }
 
-void UFollowSplineMode::OnSimulationTick(const FSimulationTickParams& Params, FMoverTickEndData& OutputState)
+void UFollowSplineMode::SimulationTick_Implementation(const FSimulationTickParams& Params, FMoverTickEndData& OutputState)
 {
 	const FMoverTickStartData& StartState = Params.StartState;
 	const FMovingComponentSet& MovingComps = Params.MovingComps;
diff --git a/Engine/Plugins/Experimental/MoverExamples/Source/MoverExamples/Public/CharacterVariants/Ziplining/ZipliningMode.h b/Engine/Plugins/Experimental/MoverExamples/Source/MoverExamples/Public/CharacterVariants/Ziplining/ZipliningMode.h
--- a/Engine/Plugins/Experimental/MoverExamples/Source/MoverExamples/Public/CharacterVariants/Ziplining/ZipliningMode.h
+++ b/Engine/Plugins/Experimental/MoverExamples/Source/MoverExamples/Public/CharacterVariants/Ziplining/ZipliningMode.h
@@ -20,8 +20,8 @@ class UZipliningMode : public UBaseMovementMode
 	GENERATED_UCLASS_BODY()
 
 public:
-	virtual void OnGenerateMove(const FMoverTickStartData& StartState, const FMoverTimeStep& TimeStep, FProposedMove& OutProposedMove) const override;
-	virtual void OnSimulationTick(const FSimulationTickParams& Params, FMoverTickEndData& OutputState) override;
+	virtual void GenerateMove_Implementation(const FMoverTickStartData& StartState, const FMoverTimeStep& TimeStep, FProposedMove& OutProposedMove) const override;
+	virtual void SimulationTick_Implementation(const FSimulationTickParams& Params, FMoverTickEndData& OutputState) override;
 
 	// Maximum speed 
 	UPROPERTY(EditAnywhere, Category = "Ziplining", meta = (ClampMin = "1", UIMin = "1", ForceUnits = "cm/s"))
diff --git a/Engine/Plugins/Experimental/MoverExamples/Source/MoverExamples/Public/CharacterVariants/Ziplining/ZipliningTransitions.h b/Engine/Plugins/Experimental/MoverExamples/Source/MoverExamples/Public/CharacterVariants/Ziplining/ZipliningTransitions.h
--- a/Engine/Plugins/Experimental/MoverExamples/Source/MoverExamples/Public/CharacterVariants/Ziplining/ZipliningTransitions.h
+++ b/Engine/Plugins/Experimental/MoverExamples/Source/MoverExamples/Public/CharacterVariants/Ziplining/ZipliningTransitions.h
@@ -16,7 +16,7 @@ class UZiplineStartTransition : public UBaseMovementModeTransition
 	GENERATED_UCLASS_BODY()
 
 public:
-	virtual FTransitionEvalResult OnEvaluate(const FSimulationTickParams& Params) const override;
+	virtual FTransitionEvalResult Evaluate_Implementation(const FSimulationTickParams& Params) const override;
 
 	UPROPERTY(EditAnywhere, Category = "Ziplining")
 	FName ZipliningModeName = ExtendedModeNames::Ziplining;
@@ -31,8 +31,8 @@ class UZiplineEndTransition : public UBaseMovementModeTransition
 	GENERATED_UCLASS_BODY()
 
 public:
-	virtual FTransitionEvalResult OnEvaluate(const FSimulationTickParams& Params) const override;
-	virtual void OnTrigger(const FSimulationTickParams& Params) override;
+	virtual FTransitionEvalResult Evaluate_Implementation(const FSimulationTickParams& Params) const override;
+	virtual void Trigger_Implementation(const FSimulationTickParams& Params) override;
 
 	// Mode to enter when exiting the zipline
 	UPROPERTY(EditAnywhere, Category = "Ziplining")
diff --git a/Engine/Plugins/Experimental/MoverExamples/Source/MoverExamples/Public/MovementBases/FollowPathMode.h b/Engine/Plugins/Experimental/MoverExamples/Source/MoverExamples/Public/MovementBases/FollowPathMode.h
--- a/Engine/Plugins/Experimental/MoverExamples/Source/MoverExamples/Public/MovementBases/FollowPathMode.h
+++ b/Engine/Plugins/Experimental/MoverExamples/Source/MoverExamples/Public/MovementBases/FollowPathMode.h
@@ -36,11 +36,8 @@ class MOVEREXAMPLES_API UFollowPathMode : public UBaseMovementMode
 
 
 public:
-	UFUNCTION(BlueprintCallable, Category = Mover)
-	virtual void OnGenerateMove(const FMoverTickStartData& StartState, const FMoverTimeStep& TimeStep, FProposedMove& OutProposedMove) const override;
-
-	UFUNCTION(BlueprintCallable, Category = Mover)
-	virtual void OnSimulationTick(const FSimulationTickParams& Params, FMoverTickEndData& OutputState) override;
+	virtual void GenerateMove_Implementation(const FMoverTickStartData& StartState, const FMoverTimeStep& TimeStep, FProposedMove& OutProposedMove) const override;
+	virtual void SimulationTick_Implementation(const FSimulationTickParams& Params, FMoverTickEndData& OutputState) override;
 
 	// List of ordered path locations to visit
 	UPROPERTY(EditAnywhere, BlueprintReadWrite, Category = Pathing)
diff --git a/Engine/Plugins/Experimental/MoverExamples/Source/MoverExamples/Public/MovementBases/FollowSplineMode.h b/Engine/Plugins/Experimental/MoverExamples/Source/MoverExamples/Public/MovementBases/FollowSplineMode.h
--- a/Engine/Plugins/Experimental/MoverExamples/Source/MoverExamples/Public/MovementBases/FollowSplineMode.h
+++ b/Engine/Plugins/Experimental/MoverExamples/Source/MoverExamples/Public/MovementBases/FollowSplineMode.h
@@ -55,11 +55,9 @@ class MOVEREXAMPLES_API UFollowSplineMode : public UBaseMovementMode
 	GENERATED_UCLASS_BODY()
 
 public:
-	UFUNCTION(BlueprintCallable, Category = Mover)
-	virtual void OnGenerateMove(const FMoverTickStartData& StartState, const FMoverTimeStep& TimeStep, FProposedMove& OutProposedMove) const override;
+	virtual void GenerateMove_Implementation(const FMoverTickStartData& StartState, const FMoverTimeStep& TimeStep, FProposedMove& OutProposedMove) const override;
 
-	UFUNCTION(BlueprintCallable, Category = Mover)
-	virtual void OnSimulationTick(const FSimulationTickParams& Params, FMoverTickEndData& OutputState) override;
+	virtual void SimulationTick_Implementation(const FSimulationTickParams& Params, FMoverTickEndData& OutputState) override;
 
 	UFUNCTION(BlueprintCallable, Category = "Mover|Spline")
 	void SetControlSpline(const AActor* SplineProviderActor, FSplineOffsetRangeInput Offset = FSplineOffsetRangeInput());
