Mover: Adding IsValid check to NavMoverComponent GetUpdatedComponent

diff --git a/Engine/Plugins/Experimental/Mover/Source/Mover/Public/DefaultMovementSet/NavMoverComponent.h b/Engine/Plugins/Experimental/Mover/Source/Mover/Public/DefaultMovementSet/NavMoverComponent.h
--- a/Engine/Plugins/Experimental/Mover/Source/Mover/Public/DefaultMovementSet/NavMoverComponent.h
+++ b/Engine/Plugins/Experimental/Mover/Source/Mover/Public/DefaultMovementSet/NavMoverComponent.h
@@ -71,7 +71,7 @@ public:
 	virtual UObject* GetOwnerAsObject() const override { return GetOwner(); }
 	
 	/** Get the component this movement component is updating */
-	virtual TObjectPtr<UObject> GetUpdatedObject() const override { return MoverComponent->GetUpdatedComponent(); }
+	virtual TObjectPtr<UObject> GetUpdatedObject() const override { return MoverComponent.IsValid() ? MoverComponent->GetUpdatedComponent() : nullptr; }
 
 	/** Get axis-aligned cylinder around this actor, used for simple collision checks in nav movement */
 	virtual void GetSimpleCollisionCylinder(float& CollisionRadius, float& CollisionHalfHeight) const override;
