Mover: Updating NavMoverComponent GetLocation functions to get rid of unneeded casting

diff --git a/Engine/Plugins/Experimental/Mover/Source/Mover/Private/DefaultMovementSet/NavMoverComponent.cpp b/Engine/Plugins/Experimental/Mover/Source/Mover/Private/DefaultMovementSet/NavMoverComponent.cpp
--- a/Engine/Plugins/Experimental/Mover/Source/Mover/Private/DefaultMovementSet/NavMoverComponent.cpp
+++ b/Engine/Plugins/Experimental/Mover/Source/Mover/Private/DefaultMovementSet/NavMoverComponent.cpp
@@ -111,20 +111,41 @@ bool UNavMoverComponent::ConsumeNavMovementData(FVector& OutMoveInputIntent, FVe
 
 FVector UNavMoverComponent::GetLocation() const
 {
-	const USceneComponent* UpdatedComponent = Cast<USceneComponent>(GetUpdatedObject());
-	return UpdatedComponent ? UpdatedComponent->GetComponentLocation() : FVector(FLT_MAX);
+	if (MoverComponent.IsValid())
+	{
+		if (const USceneComponent* UpdatedComponent = MoverComponent->GetUpdatedComponent())
+		{
+			return UpdatedComponent->GetComponentLocation();
+		}
+	}
+	
+	return FVector(FLT_MAX);
 }
 
 FVector UNavMoverComponent::GetFeetLocation() const
 {
-	const USceneComponent* UpdatedComponent = Cast<USceneComponent>(GetUpdatedObject());
-	return UpdatedComponent ? (UpdatedComponent->GetComponentLocation() - FVector(0,0,UpdatedComponent->Bounds.BoxExtent.Z)) : FNavigationSystem::InvalidLocation;
+	if (MoverComponent.IsValid())
+	{
+		if (const USceneComponent* UpdatedComponent = MoverComponent->GetUpdatedComponent())
+		{
+			return UpdatedComponent->GetComponentLocation() - FVector(0,0,UpdatedComponent->Bounds.BoxExtent.Z);
+		}
+	}
+
+	return FNavigationSystem::InvalidLocation;
 }
 
 FVector UNavMoverComponent::GetFeetLocationAt(FVector ComponentLocation) const
 {
-	const USceneComponent* UpdatedComponent = Cast<USceneComponent>(GetUpdatedObject());
-	return UpdatedComponent ? (ComponentLocation - FVector(0, 0, UpdatedComponent->Bounds.BoxExtent.Z)) : FNavigationSystem::InvalidLocation;
+	if (MoverComponent.IsValid())
+	{
+		if (const USceneComponent* UpdatedComponent = MoverComponent->GetUpdatedComponent())
+		{
+			return ComponentLocation - FVector(0, 0, UpdatedComponent->Bounds.BoxExtent.Z);
+		}
+	}
+	
+	return FNavigationSystem::InvalidLocation;
 }
 
 FBasedPosition UNavMoverComponent::GetFeetLocationBased() const
