[Chaos Mover] Fix crash when IsDataValid() is called on CDO's of physics-driven movement mode BP child classes (CDO's don't have an owning outer MoverComponent)

diff --git a/Engine/Plugins/Experimental/Mover/Source/Mover/Private/PhysicsMover/Modes/PhysicsDrivenFallingMode.cpp b/Engine/Plugins/Experimental/Mover/Source/Mover/Private/PhysicsMover/Modes/PhysicsDrivenFallingMode.cpp
--- a/Engine/Plugins/Experimental/Mover/Source/Mover/Private/PhysicsMover/Modes/PhysicsDrivenFallingMode.cpp
+++ b/Engine/Plugins/Experimental/Mover/Source/Mover/Private/PhysicsMover/Modes/PhysicsDrivenFallingMode.cpp
@@ -41,13 +41,7 @@ void UPhysicsDrivenFallingMode::UpdateConstraintSettings(Chaos::FCharacterGround
 EDataValidationResult UPhysicsDrivenFallingMode::IsDataValid(FDataValidationContext& Context) const
 {
 	EDataValidationResult Result = Super::IsDataValid(Context);
-	const UClass* BackendClass = GetMoverComponent()->BackendClass;
-	if (BackendClass && !BackendClass->IsChildOf<UMoverNetworkPhysicsLiaisonComponent>())
-	{
-		Context.AddError(LOCTEXT("PhysicsMovementModeHasValidPhysicsLiaison", "Physics movement modes need to have a backend class that supports physics (UMoverNetworkPhysicsLiaisonComponent)."));
-		Result = EDataValidationResult::Invalid;
-	}
-		
+	PhysicsMovementModeUtils::ValidateBackendClass(GetMoverComponent(), Context, Result);
 	return Result;
 }
 #endif // WITH_EDITOR
diff --git a/Engine/Plugins/Experimental/Mover/Source/Mover/Private/PhysicsMover/Modes/PhysicsDrivenFlyingMode.cpp b/Engine/Plugins/Experimental/Mover/Source/Mover/Private/PhysicsMover/Modes/PhysicsDrivenFlyingMode.cpp
--- a/Engine/Plugins/Experimental/Mover/Source/Mover/Private/PhysicsMover/Modes/PhysicsDrivenFlyingMode.cpp
+++ b/Engine/Plugins/Experimental/Mover/Source/Mover/Private/PhysicsMover/Modes/PhysicsDrivenFlyingMode.cpp
@@ -36,13 +36,7 @@ void UPhysicsDrivenFlyingMode::UpdateConstraintSettings(Chaos::FCharacterGroundC
 EDataValidationResult UPhysicsDrivenFlyingMode::IsDataValid(FDataValidationContext& Context) const
 {
 	EDataValidationResult Result = Super::IsDataValid(Context);
-	const UClass* BackendClass = GetMoverComponent()->BackendClass;
-	if (BackendClass && !BackendClass->IsChildOf<UMoverNetworkPhysicsLiaisonComponent>())
-	{
-		Context.AddError(LOCTEXT("PhysicsMovementModeHasValidPhysicsLiaison", "Physics movement modes need to have a backend class that supports physics (UMoverNetworkPhysicsLiaisonComponent)."));
-		Result = EDataValidationResult::Invalid;
-	}
-		
+	PhysicsMovementModeUtils::ValidateBackendClass(GetMoverComponent(), Context, Result);
 	return Result;
 }
 #endif // WITH_EDITOR
diff --git a/Engine/Plugins/Experimental/Mover/Source/Mover/Private/PhysicsMover/Modes/PhysicsDrivenSwimmingMode.cpp b/Engine/Plugins/Experimental/Mover/Source/Mover/Private/PhysicsMover/Modes/PhysicsDrivenSwimmingMode.cpp
--- a/Engine/Plugins/Experimental/Mover/Source/Mover/Private/PhysicsMover/Modes/PhysicsDrivenSwimmingMode.cpp
+++ b/Engine/Plugins/Experimental/Mover/Source/Mover/Private/PhysicsMover/Modes/PhysicsDrivenSwimmingMode.cpp
@@ -42,13 +42,7 @@ void UPhysicsDrivenSwimmingMode::UpdateConstraintSettings(Chaos::FCharacterGroun
 EDataValidationResult UPhysicsDrivenSwimmingMode::IsDataValid(FDataValidationContext& Context) const
 {
 	EDataValidationResult Result = Super::IsDataValid(Context);
-	const UClass* BackendClass = GetMoverComponent()->BackendClass;
-	if (BackendClass && !BackendClass->IsChildOf<UMoverNetworkPhysicsLiaisonComponent>())
-	{
-		Context.AddError(LOCTEXT("PhysicsMovementModeHasValidPhysicsLiaison", "Physics movement modes need to have a backend class that supports physics (UMoverNetworkPhysicsLiaisonComponent)."));
-		Result = EDataValidationResult::Invalid;
-	}
-		
+	PhysicsMovementModeUtils::ValidateBackendClass(GetMoverComponent(), Context, Result);
 	return Result;
 }
 #endif // WITH_EDITOR
diff --git a/Engine/Plugins/Experimental/Mover/Source/Mover/Private/PhysicsMover/Modes/PhysicsDrivenWalkingMode.cpp b/Engine/Plugins/Experimental/Mover/Source/Mover/Private/PhysicsMover/Modes/PhysicsDrivenWalkingMode.cpp
--- a/Engine/Plugins/Experimental/Mover/Source/Mover/Private/PhysicsMover/Modes/PhysicsDrivenWalkingMode.cpp
+++ b/Engine/Plugins/Experimental/Mover/Source/Mover/Private/PhysicsMover/Modes/PhysicsDrivenWalkingMode.cpp
@@ -252,13 +252,7 @@ void UPhysicsDrivenWalkingMode::FloorCheck(const FMoverDefaultSyncState& SyncSta
 EDataValidationResult UPhysicsDrivenWalkingMode::IsDataValid(FDataValidationContext& Context) const
 {
 	EDataValidationResult Result = Super::IsDataValid(Context);
-	const UClass* BackendClass = GetMoverComponent()->BackendClass;
-	if (BackendClass && !BackendClass->IsChildOf<UMoverNetworkPhysicsLiaisonComponent>())
-	{
-		Context.AddError(LOCTEXT("PhysicsMovementModeHasValidPhysicsLiaison", "Physics movement modes need to have a backend class that supports physics (UMoverNetworkPhysicsLiaisonComponent)."));
-		Result = EDataValidationResult::Invalid;
-	}
-		
+	PhysicsMovementModeUtils::ValidateBackendClass(GetMoverComponent(), Context, Result);
 	return Result;
 }
 #endif // WITH_EDITOR
diff --git a/Engine/Plugins/Experimental/Mover/Source/Mover/Private/PhysicsMover/PhysicsMoverSimulationTypes.cpp b/Engine/Plugins/Experimental/Mover/Source/Mover/Private/PhysicsMover/PhysicsMoverSimulationTypes.cpp
--- a/Engine/Plugins/Experimental/Mover/Source/Mover/Private/PhysicsMover/PhysicsMoverSimulationTypes.cpp
+++ b/Engine/Plugins/Experimental/Mover/Source/Mover/Private/PhysicsMover/PhysicsMoverSimulationTypes.cpp
@@ -1,7 +1,10 @@
 // Copyright Epic Games, Inc. All Rights Reserved.
 
 #include "PhysicsMover/PhysicsMoverSimulationTypes.h"
+
+#include "Backends/MoverNetworkPhysicsLiaison.h"
 #include "HAL/IConsoleManager.h"
+#include "Misc/DataValidation.h"
 
 //////////////////////////////////////////////////////////////////////////
 // Debug
@@ -23,6 +26,25 @@ FAutoConsoleVariableRef CVarPhysicsDrivenMotionMinStepUpDistance(TEXT("p.mover.p
 FAutoConsoleVariableRef CVarPhysicsDrivenMotionMaxCharacterGroundMassRatio(TEXT("p.mover.physics.MaxCharacterGroundMassRatio"),
 	GPhysicsDrivenMotionDebugParams.MaxCharacterGroundMassRatio, TEXT("Maximum ratio between character mass and ground mass as seen by the ground constraint."));
 
+//////////////////////////////////////////////////////////////////////////
+// PhysicsMovementModeUtils
+
+namespace PhysicsMovementModeUtils
+{
+	void ValidateBackendClass(UMoverComponent* MoverComponent, FDataValidationContext& Context, EDataValidationResult& Result)
+	{
+		if (MoverComponent)
+		{
+			const UClass* BackendClass = MoverComponent->BackendClass;
+			if (BackendClass && !BackendClass->IsChildOf<UMoverNetworkPhysicsLiaisonComponent>())
+			{
+				Context.AddError(NSLOCTEXT("PhysicsMovementModeUtils", "PhysicsMovementModeHasValidPhysicsLiaison", "Physics movement modes need to have a backend class that supports physics (UMoverNetworkPhysicsLiaisonComponent)."));
+				Result = EDataValidationResult::Invalid;
+			}
+		}
+	}	
+}
+
 //////////////////////////////////////////////////////////////////////////
 // FMovementSettingsInputs
 
diff --git a/Engine/Plugins/Experimental/Mover/Source/Mover/Public/PhysicsMover/PhysicsMoverSimulationTypes.h b/Engine/Plugins/Experimental/Mover/Source/Mover/Public/PhysicsMover/PhysicsMoverSimulationTypes.h
--- a/Engine/Plugins/Experimental/Mover/Source/Mover/Public/PhysicsMover/PhysicsMoverSimulationTypes.h
+++ b/Engine/Plugins/Experimental/Mover/Source/Mover/Public/PhysicsMover/PhysicsMoverSimulationTypes.h
@@ -99,6 +99,12 @@ public:
 	virtual void OnContactModification_Internal(const FPhysicsMoverSimulationContactModifierParams& Params, Chaos::FCollisionContactModifier& Modifier) const {};
 };
 
+class FDataValidationContext;
+namespace PhysicsMovementModeUtils
+{
+	void ValidateBackendClass(UMoverComponent* MoverComponent, FDataValidationContext& Context, EDataValidationResult& Result);
+}
+
 //////////////////////////////////////////////////////////////////////////
 // FMovementSettingsInputs
 
