[Mover] Changed MovementMode::Activate & Deactivate to be native virtuals with a BPIE instead of BPNE's, since BP usage of these is much more likely to be for adding additional behavior than stomping what the parent does.
- Also changed the deprecated virtuals in MovementMode and MovementModeTransition to be marked final, since overriding a deprecated virtual doesn't trigger the warning (only calls do). This means it's a breaking change, but with a clear paper trail to easily fix things up.
- Fixed crash in UBaseMovementMode::GetWorld() when opening the graph of a movement mode BP in the editor (a CDO doesn't have an outer MoverComponent)

diff --git a/Engine/Plugins/Experimental/Mover/Config/BaseMover.ini b/Engine/Plugins/Experimental/Mover/Config/BaseMover.ini
--- a/Engine/Plugins/Experimental/Mover/Config/BaseMover.ini
+++ b/Engine/Plugins/Experimental/Mover/Config/BaseMover.ini
@@ -8,13 +8,13 @@
 +StructRedirects=(OldName="/Script/Mover.KinematicDefaultInputs",NewName="/Script/Mover.CharacterDefaultInputs")
 +StructRedirects=(OldName="/Script/Mover.LayeredMove_JumpImpulse",NewName="/Script/Mover.LayeredMove_JumpImpulseOverDuration")
 
-
-
 +FunctionRedirects=(OldName="/Script/Mover.BaseMovementModeTransition.K2_OnTrigger",NewName="Trigger")
 +FunctionRedirects=(OldName="/Script/Mover.BaseMovementModeTransition.K2_OnEvaluate",NewName="Evaluate")
 +FunctionRedirects=(OldName="/Script/Mover.BaseMovementModeTransition.GetMoverComponent",NewName="K2_GetMoverComponent")
 +FunctionRedirects=(OldName="/Script/Mover.BaseMovementMode.K2_OnGenerateMove",NewName="GenerateMove")
 +FunctionRedirects=(OldName="/Script/Mover.BaseMovementMode.K2_OnSimulationTick",NewName="SimulationTick")
-+FunctionRedirects=(OldName="/Script/Mover.BaseMovementMode.K2_OnActivate",NewName="Activate")
-+FunctionRedirects=(OldName="/Script/Mover.BaseMovementMode.K2_OnDeactivate",NewName="Deactivate")
++FunctionRedirects=(OldName="/Script/Mover.BaseMovementMode.K2_OnActivate",NewName="K2_OnActivated")
++FunctionRedirects=(OldName="/Script/Mover.BaseMovementMode.Activate",NewName="K2_OnActivated")
++FunctionRedirects=(OldName="/Script/Mover.BaseMovementMode.K2_OnDeactivate",NewName="K2_OnDeactivated")
++FunctionRedirects=(OldName="/Script/Mover.BaseMovementMode.Deactivate",NewName="K2_OnDeactivated")
 +FunctionRedirects=(OldName="/Script/Mover.BaseMovementMode.GetMoverComponent",NewName="K2_GetMoverComponent")
\ No newline at end of file
diff --git a/Engine/Plugins/Experimental/Mover/Source/Mover/Private/DefaultMovementSet/Modes/AsyncNavWalkingMode.cpp b/Engine/Plugins/Experimental/Mover/Source/Mover/Private/DefaultMovementSet/Modes/AsyncNavWalkingMode.cpp
--- a/Engine/Plugins/Experimental/Mover/Source/Mover/Private/DefaultMovementSet/Modes/AsyncNavWalkingMode.cpp
+++ b/Engine/Plugins/Experimental/Mover/Source/Mover/Private/DefaultMovementSet/Modes/AsyncNavWalkingMode.cpp
@@ -342,8 +342,10 @@ void UAsyncNavWalkingMode::SetCollisionForNavWalking(bool bEnable)
 	}
 }
 
-void UAsyncNavWalkingMode::Activate_Implementation()
+void UAsyncNavWalkingMode::Activate()
 {
+	Super::Activate();
+	
 	if (UMoverComponent* MoverComp = GetMoverComponent())
 	{
 		// TEMP solution for our inability to change collision settings outside of the game thread
@@ -358,12 +360,14 @@ void UAsyncNavWalkingMode::Activate_Implementation()
 	NavMeshProjectionTimer = (NavMeshProjectionInterval > 0.f) ? FMath::FRandRange(-NavMeshProjectionInterval, 0.f) : 0.f;
 }
 
-void UAsyncNavWalkingMode::Deactivate_Implementation()
+void UAsyncNavWalkingMode::Deactivate()
 {
 	if (UMoverComponent* MoverComp = GetMoverComponent())
 	{
 		MoverComp->OnPostFinalize.AddDynamic(this, &UAsyncNavWalkingMode::PostFinalizeRestoreOriginalCollision);
 	}
+	
+	Super::Deactivate();
 }
 
 const INavigationDataInterface* UAsyncNavWalkingMode::GetNavData() const
diff --git a/Engine/Plugins/Experimental/Mover/Source/Mover/Private/DefaultMovementSet/Modes/NavWalkingMode.cpp b/Engine/Plugins/Experimental/Mover/Source/Mover/Private/DefaultMovementSet/Modes/NavWalkingMode.cpp
--- a/Engine/Plugins/Experimental/Mover/Source/Mover/Private/DefaultMovementSet/Modes/NavWalkingMode.cpp
+++ b/Engine/Plugins/Experimental/Mover/Source/Mover/Private/DefaultMovementSet/Modes/NavWalkingMode.cpp
@@ -345,14 +345,16 @@ void UNavWalkingMode::SetCollisionForNavWalking(bool bEnable)
 	}
 }
 
-void UNavWalkingMode::Activate_Implementation()
+void UNavWalkingMode::Activate()
 {
+	Super::Activate();
 	SetCollisionForNavWalking(true);
 }
 
-void UNavWalkingMode::Deactivate_Implementation()
+void UNavWalkingMode::Deactivate()
 {
 	SetCollisionForNavWalking(false);
+	Super::Deactivate();
 }
 
 const INavigationDataInterface* UNavWalkingMode::GetNavData() const
diff --git a/Engine/Plugins/Experimental/Mover/Source/Mover/Private/MovementMode.cpp b/Engine/Plugins/Experimental/Mover/Source/Mover/Private/MovementMode.cpp
--- a/Engine/Plugins/Experimental/Mover/Source/Mover/Private/MovementMode.cpp
+++ b/Engine/Plugins/Experimental/Mover/Source/Mover/Private/MovementMode.cpp
@@ -15,6 +15,13 @@
 
 UWorld* UBaseMovementMode::GetWorld() const
 {
+#if WITH_EDITOR
+	// In the editor, GetWorld() is called on the CDO as part of checking ImplementsGetWorld(). Only the CDO can exist without being outer'd to a MoverComponent.
+	if (IsTemplate())
+	{
+		return nullptr;
+	}
+#endif
 	return GetOuterUMoverComponent()->GetWorld();
 }
 
@@ -39,27 +46,23 @@ void UBaseMovementMode::OnUnregistered()
 	K2_OnUnregistered();
 }
 
-PRAGMA_DISABLE_DEPRECATION_WARNINGS
-void UBaseMovementMode::Activate_Implementation()
+void UBaseMovementMode::Activate()
 {
-	OnActivate();
+	K2_OnActivated();
 }
 
-void UBaseMovementMode::Deactivate_Implementation()
+void UBaseMovementMode::Deactivate()
 {
-	OnDeactivate();
+	K2_OnDeactivated();
 }
 
 void UBaseMovementMode::GenerateMove_Implementation(const FMoverTickStartData& StartState, const FMoverTimeStep& TimeStep, FProposedMove& OutProposedMove) const
 {
-	OnGenerateMove(StartState, TimeStep, OutProposedMove);
 }
 
 void UBaseMovementMode::SimulationTick_Implementation(const FSimulationTickParams& Params, FMoverTickEndData& OutputState)
 {
-	OnSimulationTick(Params, OutputState);
 }
-PRAGMA_ENABLE_DEPRECATION_WARNINGS
 
 UMoverComponent* UBaseMovementMode::K2_GetMoverComponent() const
 {
diff --git a/Engine/Plugins/Experimental/Mover/Source/Mover/Private/MovementModeTransition.cpp b/Engine/Plugins/Experimental/Mover/Source/Mover/Private/MovementModeTransition.cpp
--- a/Engine/Plugins/Experimental/Mover/Source/Mover/Private/MovementModeTransition.cpp
+++ b/Engine/Plugins/Experimental/Mover/Source/Mover/Private/MovementModeTransition.cpp
@@ -36,21 +36,13 @@ UMoverComponent* UBaseMovementModeTransition::K2_GetMoverComponent() const
 	return GetTypedOuter<UMoverComponent>();
 }
 
-PRAGMA_DISABLE_DEPRECATION_WARNINGS
 FTransitionEvalResult UBaseMovementModeTransition::Evaluate_Implementation(const FSimulationTickParams& Params) const
 {
-	return OnEvaluate(Params);
+	return FTransitionEvalResult::NoTransition;
 }
 
 void UBaseMovementModeTransition::Trigger_Implementation(const FSimulationTickParams& Params)
 {
-	OnTrigger(Params);
-}
-PRAGMA_ENABLE_DEPRECATION_WARNINGS
-
-FTransitionEvalResult UBaseMovementModeTransition::OnEvaluate(const FSimulationTickParams& Params) const
-{
-	return FTransitionEvalResult::NoTransition;
 }
 
 #if WITH_EDITOR
diff --git a/Engine/Plugins/Experimental/Mover/Source/Mover/Public/DefaultMovementSet/Modes/AsyncNavWalkingMode.h b/Engine/Plugins/Experimental/Mover/Source/Mover/Public/DefaultMovementSet/Modes/AsyncNavWalkingMode.h
--- a/Engine/Plugins/Experimental/Mover/Source/Mover/Public/DefaultMovementSet/Modes/AsyncNavWalkingMode.h
+++ b/Engine/Plugins/Experimental/Mover/Source/Mover/Public/DefaultMovementSet/Modes/AsyncNavWalkingMode.h
@@ -27,8 +27,8 @@ class MOVER_API UAsyncNavWalkingMode : public UBaseMovementMode
 public:
 	UAsyncNavWalkingMode();
 
-	virtual void Activate_Implementation() override;
-	virtual void Deactivate_Implementation() override;
+	virtual void Activate() override;
+	virtual void Deactivate() override;
 	virtual void SimulationTick_Implementation(const FSimulationTickParams& Params, FMoverTickEndData& OutputState) override;
 	virtual void GenerateMove_Implementation(const FMoverTickStartData& StartState, const FMoverTimeStep& TimeStep, FProposedMove& OutProposedMove) const override;
 
diff --git a/Engine/Plugins/Experimental/Mover/Source/Mover/Public/DefaultMovementSet/Modes/NavWalkingMode.h b/Engine/Plugins/Experimental/Mover/Source/Mover/Public/DefaultMovementSet/Modes/NavWalkingMode.h
--- a/Engine/Plugins/Experimental/Mover/Source/Mover/Public/DefaultMovementSet/Modes/NavWalkingMode.h
+++ b/Engine/Plugins/Experimental/Mover/Source/Mover/Public/DefaultMovementSet/Modes/NavWalkingMode.h
@@ -31,8 +31,8 @@ class MOVER_API UNavWalkingMode : public UBaseMovementMode
 public:
 	UNavWalkingMode();
 	
-	virtual void Activate_Implementation() override;
-	virtual void Deactivate_Implementation() override;
+	virtual void Activate() override;
+	virtual void Deactivate() override;
 	virtual void SimulationTick_Implementation(const FSimulationTickParams& Params, FMoverTickEndData& OutputState) override;
 	virtual void GenerateMove_Implementation(const FMoverTickStartData& StartState, const FMoverTimeStep& TimeStep, FProposedMove& OutProposedMove) const override;
 
diff --git a/Engine/Plugins/Experimental/Mover/Source/Mover/Public/MovementMode.h b/Engine/Plugins/Experimental/Mover/Source/Mover/Public/MovementMode.h
--- a/Engine/Plugins/Experimental/Mover/Source/Mover/Public/MovementMode.h
+++ b/Engine/Plugins/Experimental/Mover/Source/Mover/Public/MovementMode.h
@@ -43,12 +43,9 @@ public:
 	
 	virtual void OnRegistered(const FName ModeName);
 	virtual void OnUnregistered();
-
-	UFUNCTION(BlueprintNativeEvent)
-	void Activate();
 	
-	UFUNCTION(BlueprintNativeEvent)
-	void Deactivate();
+	virtual void Activate();
+	virtual void Deactivate();
 	
 	UFUNCTION(BlueprintNativeEvent, meta = (DisplayName = "Generate Move"))
 	void GenerateMove(const FMoverTickStartData& StartState, const FMoverTimeStep& TimeStep, FProposedMove& OutProposedMove) const;
@@ -99,18 +96,24 @@ public:
 	FGameplayTagContainer GameplayTags;
 	
 protected:
+	UFUNCTION(BlueprintImplementableEvent, meta = (DisplayName = "On Activated", ScriptName = "OnActivated"))
+	void K2_OnActivated();
+
+	UFUNCTION(BlueprintImplementableEvent, meta = (DisplayName = "On Deactivated", ScriptName = "OnDeactivated"))
+	void K2_OnDeactivated();
+	
 	UFUNCTION(BlueprintImplementableEvent, meta = (DisplayName = "On Registered", ScriptName = "OnRegistered"))
 	void K2_OnRegistered(const FName ModeName);
 
 	UFUNCTION(BlueprintImplementableEvent, meta = (DisplayName = "On Unregistered", ScriptName = "OnUnregistered"))
 	void K2_OnUnregistered();
 
-	UE_DEPRECATED(5.6, "OnActivate() has been replaced with an Activate() BlueprintNativeEvent. Rename your override to Activate_Implementation().")
-	virtual void OnActivate() {}
-	UE_DEPRECATED(5.6, "OnDeactivate() has been replaced with a Deactivate() BlueprintNativeEvent. Rename your override to Deactivate_Implementation().")
-	virtual void OnDeactivate() {}
+	UE_DEPRECATED(5.6, "OnActivate() has been renamed to Activate().")
+	virtual void OnActivate() final {}
+	UE_DEPRECATED(5.6, "OnDeactivate() has been renamed to Deactivate().")
+	virtual void OnDeactivate() final  {}
 	UE_DEPRECATED(5.6, "OnGenerateMove() has been replaced with a GenerateMove() BlueprintNativeEvent. Rename your override to GenerateMove_Implementation().")
-	virtual void OnGenerateMove(const FMoverTickStartData& StartState, const FMoverTimeStep& TimeStep, FProposedMove& OutProposedMove) const {}
+	virtual void OnGenerateMove(const FMoverTickStartData& StartState, const FMoverTimeStep& TimeStep, FProposedMove& OutProposedMove) const final {}
 	UE_DEPRECATED(5.6, "OnSimulationTick() has been replaced with a SimulationTick() BlueprintNativeEvent. Rename your override to SimulationTick_Implementation().")
-	virtual void OnSimulationTick(const FSimulationTickParams& Params, FMoverTickEndData& OutputState) {}
+	virtual void OnSimulationTick(const FSimulationTickParams& Params, FMoverTickEndData& OutputState) final {}
 };
\ No newline at end of file
diff --git a/Engine/Plugins/Experimental/Mover/Source/Mover/Public/MovementModeTransition.h b/Engine/Plugins/Experimental/Mover/Source/Mover/Public/MovementModeTransition.h
--- a/Engine/Plugins/Experimental/Mover/Source/Mover/Public/MovementModeTransition.h
+++ b/Engine/Plugins/Experimental/Mover/Source/Mover/Public/MovementModeTransition.h
@@ -85,9 +85,9 @@ protected:
 	void K2_OnUnregistered();
 
 	UE_DEPRECATED(5.6, "OnEvaluate() has been replaced with an Evaluate() BlueprintNativeEvent. Rename your override to Evaluate_Implementation().")
-	virtual FTransitionEvalResult OnEvaluate(const FSimulationTickParams& Params) const;
+	virtual FTransitionEvalResult OnEvaluate(const FSimulationTickParams& Params) const final { return FTransitionEvalResult::NoTransition; }
 	UE_DEPRECATED(5.6, "OnTrigger() has been replaced with a Trigger() BlueprintNativeEvent. Rename your override to Trigger_Implementation().")
-	virtual void OnTrigger(const FSimulationTickParams& Params) {}
+	virtual void OnTrigger(const FSimulationTickParams& Params) final {}
 };
 
 
