Improve UMoverComponent memory layout.

UMoverComponent from 1168 -> 1104
FLayeredMoveGroup 80 -> 72
FLayeredMoveFinishVelocitySettings 48 -> 40

#rnx

diff --git a/Engine/Plugins/Experimental/Mover/Source/Mover/Private/LayeredMove.cpp b/Engine/Plugins/Experimental/Mover/Source/Mover/Private/LayeredMove.cpp
--- a/Engine/Plugins/Experimental/Mover/Source/Mover/Private/LayeredMove.cpp
+++ b/Engine/Plugins/Experimental/Mover/Source/Mover/Private/LayeredMove.cpp
@@ -99,9 +99,9 @@ FString FLayeredMoveBase::ToSimpleString() const
 
 
 FLayeredMoveGroup::FLayeredMoveGroup()
-	: bApplyResidualVelocity(false)
-	, ResidualVelocity(FVector::Zero())
+	: ResidualVelocity(FVector::Zero())
 	, ResidualClamping(-1.f)
+	, bApplyResidualVelocity(false)
 {
 }
 
diff --git a/Engine/Plugins/Experimental/Mover/Source/Mover/Public/LayeredMove.h b/Engine/Plugins/Experimental/Mover/Source/Mover/Public/LayeredMove.h
--- a/Engine/Plugins/Experimental/Mover/Source/Mover/Public/LayeredMove.h
+++ b/Engine/Plugins/Experimental/Mover/Source/Mover/Public/LayeredMove.h
@@ -33,19 +33,15 @@ struct FLayeredMoveFinishVelocitySettings
 	GENERATED_USTRUCT_BODY()
 
 	FLayeredMoveFinishVelocitySettings()
-		: FinishVelocityMode(ELayeredMoveFinishVelocityMode::MaintainLastRootMotionVelocity)
-		, SetVelocity(FVector::ZeroVector)
+		: SetVelocity(FVector::ZeroVector)
 		, ClampVelocity(0.f)
+		, FinishVelocityMode(ELayeredMoveFinishVelocityMode::MaintainLastRootMotionVelocity)
 	{}
 	
 	virtual ~FLayeredMoveFinishVelocitySettings() {};
 
 	virtual void NetSerialize(FArchive& Ar);
 	
-	// What mode we want to happen when a Layered Move ends, see @ELayeredMoveFinishVelocityMode
-	UPROPERTY(EditAnywhere, BlueprintReadWrite, Category = Mover)
-	ELayeredMoveFinishVelocityMode FinishVelocityMode;
-	
 	// Velocity that the actor will use if Mode == SetVelocity
 	UPROPERTY(EditAnywhere, BlueprintReadWrite, Category = Mover)
 	FVector SetVelocity;
@@ -53,6 +49,10 @@ struct FLayeredMoveFinishVelocitySettings
 	// Actor's Velocity will be clamped to this value if Mode == ClampVelocity
 	UPROPERTY(EditAnywhere, BlueprintReadWrite, Category = Mover)
 	float ClampVelocity;
+	
+	// What mode we want to happen when a Layered Move ends, see @ELayeredMoveFinishVelocityMode
+	UPROPERTY(EditAnywhere, BlueprintReadWrite, Category = Mover)
+	ELayeredMoveFinishVelocityMode FinishVelocityMode;
 };
 
 /** 
@@ -162,13 +162,6 @@ struct MOVER_API FLayeredMoveGroup
 
 	FLayeredMoveGroup();
 	virtual ~FLayeredMoveGroup() {}
-
-	/**
-	 * If true ResidualVelocity will be the next velocity used for this actor
-	 * Note: This is set automatically when a layered move has ended based off of end velocity settings - it is not meant to be set by the user see @FLayeredMoveFinishVelocitySettings
-	 */
-	UPROPERTY(NotReplicated, VisibleAnywhere, BlueprintReadOnly, Category = Mover)
-	bool bApplyResidualVelocity;
 	
 	/**
 	 * If bApplyResidualVelocity is true this actors velocity will be set to this.
@@ -176,13 +169,6 @@ struct MOVER_API FLayeredMoveGroup
      */
 	UPROPERTY(NotReplicated, VisibleAnywhere, BlueprintReadOnly, Category = Mover)
 	FVector ResidualVelocity;
-
-	/**
-	 * Clamps an actors velocity to this value when a layered move ends. This expects Value >= 0.
-	 * Note: This is set automatically when a layered move has ended based off of end velocity settings - it is not meant to be set by the user see @FLayeredMoveFinishVelocitySettings
-	 */
-	UPROPERTY(NotReplicated, VisibleAnywhere, BlueprintReadOnly, Category = Mover)
-	float ResidualClamping;
 	
 	void QueueLayeredMove(TSharedPtr<FLayeredMoveBase> Move);
 
@@ -256,6 +242,22 @@ protected:
 
 	/** Moves that are queued to become active next sim frame */
 	TArray< TSharedPtr<FLayeredMoveBase> > QueuedLayeredMoves;
+
+public:
+
+	/**
+	 * Clamps an actors velocity to this value when a layered move ends. This expects Value >= 0.
+	 * Note: This is set automatically when a layered move has ended based off of end velocity settings - it is not meant to be set by the user see @FLayeredMoveFinishVelocitySettings
+	 */
+	UPROPERTY(NotReplicated, VisibleAnywhere, BlueprintReadOnly, Category = Mover)
+	float ResidualClamping;
+
+	/**
+	 * If true ResidualVelocity will be the next velocity used for this actor
+	 * Note: This is set automatically when a layered move has ended based off of end velocity settings - it is not meant to be set by the user see @FLayeredMoveFinishVelocitySettings
+	 */
+	UPROPERTY(NotReplicated, VisibleAnywhere, BlueprintReadOnly, Category = Mover)
+	bool bApplyResidualVelocity;
 };
 
 template<>
diff --git a/Engine/Plugins/Experimental/Mover/Source/Mover/Public/MoverComponent.h b/Engine/Plugins/Experimental/Mover/Source/Mover/Public/MoverComponent.h
--- a/Engine/Plugins/Experimental/Mover/Source/Mover/Public/MoverComponent.h
+++ b/Engine/Plugins/Experimental/Mover/Source/Mover/Public/MoverComponent.h
@@ -252,10 +252,6 @@
 	// Sets planar constraint that can limit movement direction
 	UFUNCTION(BlueprintCallable, Category = Mover)
 	void SetPlanarConstraint(const FPlanarConstraint& InConstraint);
-	
-	// If enabled, the movement of the primary visual component will be smoothed via an offset from the root moving component. This is useful in fixed-tick simulations with variable rendering rates.
-	UPROPERTY(EditAnywhere, BlueprintReadWrite, Category = "Mover")
-	EMoverSmoothingMode SmoothingMode = EMoverSmoothingMode::VisualComponentOffset;
 
 public:
 
@@ -532,18 +528,15 @@
 	UPROPERTY(Transient)
 	TObjectPtr<USceneComponent> PrimaryVisualComponent;
 
+
 	/** Cached original offset from the visual component, used for cases where we want to move the visual component away from the root component (for smoothing, corrections, etc.) */
 	FTransform BaseVisualComponentTransform = FTransform::Identity;
 
-	bool bHasValidLastProducedInput = false;
 	FMoverInputCmdContext CachedLastProducedInputCmd;
-
-	bool bHasValidCachedUsedInput = false;
+	
 	FMoverInputCmdContext CachedLastUsedInputCmd;
 	
-	bool bHasValidCachedState = false;
-	FMoverSyncState CachedLastSyncState;
-	FMoverAuxStateContext CachedLastAuxState;
+	FMoverSyncState CachedLastSyncState;	
 
 	FMoverTimeStep CachedLastSimTickTimeStep;	// Saved timestep info from our last simulation tick, used during rollback handling. This will rewind during corrections.
 	FMoverTimeStep CachedNewestSimTickTimeStep;	// Saved timestep info from the newest (farthest-advanced) simulation tick. This will not rewind during corrections.
@@ -561,23 +554,46 @@
 	UPROPERTY(Transient)
 	TObjectPtr<UMoverBlackboard> SimBlackboard;
 
-	/** If enabled, this actor will be moved to follow a base actor that it's standing on. Typically disabled for physics-based movement, which handles based movement internally. */
-	UPROPERTY(EditDefaultsOnly, Category = "Mover")
-	bool bSupportsKinematicBasedMovement = true;
-
 private:
 	/** Collection of settings objects that are shared between movement modes. This list is automatically managed based on the @MovementModes contents. */
 	UPROPERTY(EditDefaultsOnly, EditFixedSize, Instanced, Category = Mover, meta = (NoResetToDefault, ObjectMustImplement = "/Script/Mover.MovementSettingsInterface"))
 	TArray<TObjectPtr<UObject>> SharedSettings;
-
-	/** Whether or not gravity is overridden on this actor. Otherwise, fall back on world settings. See @SetGravityOverride */
-	UPROPERTY(EditDefaultsOnly, Category="Mover|Gravity")
-	bool bHasGravityOverride = false;
 	
 	/** cm/s^2, only meaningful if @bHasGravityOverride is enabled.Set @SetGravityOverride */
 	UPROPERTY(EditDefaultsOnly, Category="Mover|Gravity", meta=(ForceUnits = "cm/s^2"))
 	FVector GravityAccelOverride;
 
+private:
+
+	/** Settings that can lock movement to a particular plane */
+	UPROPERTY(EditDefaultsOnly, Category = "Mover|Constraints")
+	FPlanarConstraint PlanarConstraint;
+
+public:
+
+	// If enabled, the movement of the primary visual component will be smoothed via an offset from the root moving component. This is useful in fixed-tick simulations with variable rendering rates.
+	UPROPERTY(EditAnywhere, BlueprintReadWrite, Category = "Mover")
+	EMoverSmoothingMode SmoothingMode = EMoverSmoothingMode::VisualComponentOffset;
+
+protected:
+	bool bHasValidLastProducedInput = false;
+	bool bHasValidCachedUsedInput = false;
+	bool bHasValidCachedState = false;
+
+private:
+	/** Whether or not gravity is overridden on this actor. Otherwise, fall back on world settings. See @SetGravityOverride */
+	UPROPERTY(EditDefaultsOnly, Category="Mover|Gravity")
+	bool bHasGravityOverride = false;
+
+	/** Transient flag indicating whether we are executing OnRegister(). */
+	bool bInOnRegister = false;
+
+	/** Transient flag indicating whether we are executing InitializeComponent(). */
+	bool bInInitializeComponent = false;
+
+	// Transient flag indicating we've had a rollback and haven't started simulating forward again yet
+	bool bHasRolledBack = false;
+
 	/**
 	 * A cached quaternion representing the rotation from world space to gravity relative space defined by GravityAccelOverride.
 	 */
@@ -590,18 +606,13 @@
 	UPROPERTY(Category="Mover|Gravity", VisibleAnywhere, BlueprintReadOnly, meta = (AllowPrivateAccess = true))
 	FQuat GravityToWorldTransform;
 	
-	/** Settings that can lock movement to a particular plane */
-	UPROPERTY(EditDefaultsOnly, Category = "Mover|Constraints")
-	FPlanarConstraint PlanarConstraint;
+protected:
 
-	/** Transient flag indicating whether we are executing OnRegister(). */
-	bool bInOnRegister = false;
+	/** If enabled, this actor will be moved to follow a base actor that it's standing on. Typically disabled for physics-based movement, which handles based movement internally. */
+	UPROPERTY(EditDefaultsOnly, Category = "Mover")
+	bool bSupportsKinematicBasedMovement = true;
 
-	/** Transient flag indicating whether we are executing InitializeComponent(). */
-	bool bInInitializeComponent = false;
+	FMoverAuxStateContext CachedLastAuxState;
 
-	// Transient flag indicating we've had a rollback and haven't started simulating forward again yet
-	bool bHasRolledBack = false;
-
 	friend class UBaseMovementMode;
 	friend class UMoverNetworkPhysicsLiaisonComponent;
