Mover: removing virtual keywords from a handful of struct types that don't need it

#rnx

diff --git a/Engine/Plugins/Experimental/Mover/Source/Mover/Public/LayeredMove.h b/Engine/Plugins/Experimental/Mover/Source/Mover/Public/LayeredMove.h
--- a/Engine/Plugins/Experimental/Mover/Source/Mover/Public/LayeredMove.h
+++ b/Engine/Plugins/Experimental/Mover/Source/Mover/Public/LayeredMove.h
@@ -37,10 +37,8 @@ struct FLayeredMoveFinishVelocitySettings
 		, ClampVelocity(0.f)
 		, FinishVelocityMode(ELayeredMoveFinishVelocityMode::MaintainLastRootMotionVelocity)
 	{}
-	
-	virtual ~FLayeredMoveFinishVelocitySettings() {};
 
-	virtual void NetSerialize(FArchive& Ar);
+	void NetSerialize(FArchive& Ar);
 	
 	// Velocity that the actor will use if Mode == SetVelocity
 	UPROPERTY(EditAnywhere, BlueprintReadWrite, Category = Mover)
@@ -161,7 +159,6 @@ struct MOVER_API FLayeredMoveGroup
 	GENERATED_USTRUCT_BODY()
 
 	FLayeredMoveGroup();
-	virtual ~FLayeredMoveGroup() {}
 	
 	/**
 	 * If bApplyResidualVelocity is true this actors velocity will be set to this.
diff --git a/Engine/Plugins/Experimental/Mover/Source/Mover/Public/MoverTypes.h b/Engine/Plugins/Experimental/Mover/Source/Mover/Public/MoverTypes.h
--- a/Engine/Plugins/Experimental/Mover/Source/Mover/Public/MoverTypes.h
+++ b/Engine/Plugins/Experimental/Mover/Source/Mover/Public/MoverTypes.h
@@ -185,7 +185,6 @@ struct MOVER_API FMoverDataCollection
 	GENERATED_USTRUCT_BODY()
 
 	FMoverDataCollection();
-	virtual ~FMoverDataCollection() {}
 
 	void Empty() { DataArray.Empty(); }
 
