Mover: removed unused references to primitive component. No functional changes.

#rnx

diff --git a/Engine/Plugins/Experimental/Mover/Source/Mover/Private/DefaultMovementSet/Modes/WalkingMode.cpp b/Engine/Plugins/Experimental/Mover/Source/Mover/Private/DefaultMovementSet/Modes/WalkingMode.cpp
--- a/Engine/Plugins/Experimental/Mover/Source/Mover/Private/DefaultMovementSet/Modes/WalkingMode.cpp
+++ b/Engine/Plugins/Experimental/Mover/Source/Mover/Private/DefaultMovementSet/Modes/WalkingMode.cpp
@@ -114,10 +114,9 @@ void UWalkingMode::SimulationTick_Implementation(const FSimulationTickParams& Pa
 	UMoverComponent* MoverComp = GetMoverComponent();
 	const FMoverTickStartData& StartState = Params.StartState;
 	USceneComponent* UpdatedComponent = Params.MovingComps.UpdatedComponent.Get();
-	UPrimitiveComponent* UpdatedPrimitive = Params.MovingComps.UpdatedPrimitive.Get();
 	FProposedMove ProposedMove = Params.ProposedMove;
 
-	if (!UpdatedComponent || !UpdatedPrimitive)
+	if (!UpdatedComponent)
 	{
 		return;
 	}
@@ -145,7 +144,7 @@ void UWalkingMode::SimulationTick_Implementation(const FSimulationTickParams& Pa
 	// If we don't have cached floor information, we need to search for it again
 	if (!SimBlackboard->TryGet(CommonBlackboard::LastFloorResult, CurrentFloor))
 	{
-		UFloorQueryUtils::FindFloor(MovingComponents, CommonLegacySettings->FloorSweepDistance, CommonLegacySettings->MaxWalkSlopeCosine, UpdatedPrimitive->GetComponentLocation(), CurrentFloor);
+		UFloorQueryUtils::FindFloor(MovingComponents, CommonLegacySettings->FloorSweepDistance, CommonLegacySettings->MaxWalkSlopeCosine, UpdatedComponent->GetComponentLocation(), CurrentFloor);
 	}
 
 	OutputSyncState.MoveDirectionIntent = (ProposedMove.bHasDirIntent ? ProposedMove.DirectionIntent : FVector::ZeroVector);
@@ -236,7 +235,7 @@ void UWalkingMode::SimulationTick_Implementation(const FSimulationTickParams& Pa
 		}
 
 		// Search for the floor we've ended up on
-		UFloorQueryUtils::FindFloor(MovingComponents, CommonLegacySettings->FloorSweepDistance, CommonLegacySettings->MaxWalkSlopeCosine, UpdatedPrimitive->GetComponentLocation(), CurrentFloor);
+		UFloorQueryUtils::FindFloor(MovingComponents, CommonLegacySettings->FloorSweepDistance, CommonLegacySettings->MaxWalkSlopeCosine, UpdatedComponent->GetComponentLocation(), CurrentFloor);
 
 		if (CurrentFloor.IsWalkableFloor())
 		{
@@ -259,7 +258,7 @@ void UWalkingMode::SimulationTick_Implementation(const FSimulationTickParams& Pa
 		if ((FloorCheckPolicy == EStaticFloorCheckPolicy::Always) || 
 		    (FloorCheckPolicy == EStaticFloorCheckPolicy::OnDynamicBaseOnly && StartingSyncState->GetMovementBase()))
 		{
-			UFloorQueryUtils::FindFloor(MovingComponents, CommonLegacySettings->FloorSweepDistance, CommonLegacySettings->MaxWalkSlopeCosine, UpdatedPrimitive->GetComponentLocation(), CurrentFloor);
+			UFloorQueryUtils::FindFloor(MovingComponents, CommonLegacySettings->FloorSweepDistance, CommonLegacySettings->MaxWalkSlopeCosine, UpdatedComponent->GetComponentLocation(), CurrentFloor);
 		
 			FHitResult Hit(CurrentFloor.HitResult);
 			if (Hit.bStartPenetrating)
