Mover: skipping over mode move generation if we already have an overriding move from another source

diff --git a/Engine/Plugins/Experimental/Mover/Source/Mover/Private/MovementModeStateMachine.cpp b/Engine/Plugins/Experimental/Mover/Source/Mover/Private/MovementModeStateMachine.cpp
--- a/Engine/Plugins/Experimental/Mover/Source/Mover/Private/MovementModeStateMachine.cpp
+++ b/Engine/Plugins/Experimental/Mover/Source/Mover/Private/MovementModeStateMachine.cpp
@@ -13,6 +13,18 @@
 
 #include UE_INLINE_GENERATED_CPP_BY_NAME(MovementModeStateMachine)
 
+namespace MoverComponentCVars
+{
+	static int32 SkipGenerateMoveIfOverridden = 1;
+	FAutoConsoleVariableRef CVarSkipGenerateMoveIfOverridden(
+		TEXT("mover.perf.SkipGenerateMoveIfOverridden"),
+		SkipGenerateMoveIfOverridden,
+		TEXT("If true and we have a layered move fully overriding movement, then we will skip calling OnGenerateMove on the active movement mode for better performance\n")
+	);
+
+} // end MoverComponentCVars
+
+
 const FName UNullMovementMode::NullModeName(TEXT("Null"));
 
 UNullMovementMode::UNullMovementMode(const FObjectInitializer& ObjectInitializer)
@@ -265,12 +277,25 @@ void UMovementModeStateMachine::OnSimulationTick(USceneComponent* UpdatedCompone
 		{
 			UBaseMovementMode* CurrentMode = Modes[CurrentModeName];
 			FProposedMove CombinedMove;
-			CurrentMode->GenerateMove(SubstepStartData, SubTimeStep, OUT CombinedMove);
+			bool bHasModeMoveContribution = false;
+
+			if (!MoverComponentCVars::SkipGenerateMoveIfOverridden ||
+				!(bHasLayeredMoveContributions && CombinedLayeredMove.MixMode == EMoveMixMode::OverrideAll))
+			{
+				QUICK_SCOPE_CYCLE_COUNTER(STAT_GenerateMoveFromMode);
+				CurrentMode->GenerateMove(SubstepStartData, SubTimeStep, OUT CombinedMove);
+
+				bHasModeMoveContribution = true;
+			}
 
-			if (bHasLayeredMoveContributions)
+			if (bHasModeMoveContribution && bHasLayeredMoveContributions)
 			{
 				MoverComp->MovementMixer->MixProposedMoves(CombinedLayeredMove, CombinedMove);
 			}
+			else if (bHasLayeredMoveContributions && !bHasModeMoveContribution)
+			{
+				CombinedMove = CombinedLayeredMove;
+			}
 
 			// Apply any layered move finish velocity settings
 			if (CurrentLayeredMoves.bApplyResidualVelocity)
