Mover: making stance changing, including crouching, optional. This allows us to save on perf for characters that don't use them.

diff --git a/Engine/Plugins/Experimental/Mover/Source/Mover/Private/DefaultMovementSet/CharacterMoverComponent.cpp b/Engine/Plugins/Experimental/Mover/Source/Mover/Private/DefaultMovementSet/CharacterMoverComponent.cpp
--- a/Engine/Plugins/Experimental/Mover/Source/Mover/Private/DefaultMovementSet/CharacterMoverComponent.cpp
+++ b/Engine/Plugins/Experimental/Mover/Source/Mover/Private/DefaultMovementSet/CharacterMoverComponent.cpp
@@ -11,8 +11,6 @@
 #include "MoveLibrary/MovementUtils.h"
 
 UCharacterMoverComponent::UCharacterMoverComponent()
-	: bHandleJump(true)
-	, bWantsToCrouch(false)
 {
 	// Default movement modes
 	MovementModes.Add(DefaultModeNames::Walking, CreateDefaultSubobject<UWalkingMode>(TEXT("DefaultWalkingMode")));
@@ -123,46 +121,49 @@ void UCharacterMoverComponent::OnMoverPreSimulationTick(const FMoverTimeStep& Ti
 		}
 	}
 	
-	const FStanceModifier* StanceModifier = static_cast<const FStanceModifier*>(FindMovementModifier(StanceModifierHandle));
-	// This is a fail safe in case our handle was bad - try finding the modifier by type if we can
-	if (!StanceModifier)
+	if (bHandleStanceChanges)
 	{
-		StanceModifier = FindMovementModifierByType<FStanceModifier>();
-	}
+		const FStanceModifier* StanceModifier = static_cast<const FStanceModifier*>(FindMovementModifier(StanceModifierHandle));
+		// This is a fail safe in case our handle was bad - try finding the modifier by type if we can
+		if (!StanceModifier)
+		{
+			StanceModifier = FindMovementModifierByType<FStanceModifier>();
+		}
 	
-	EStanceMode OldActiveStance = EStanceMode::Invalid;
-	if (StanceModifier)
-	{
-		OldActiveStance = StanceModifier->ActiveStance;
-	}
+		EStanceMode OldActiveStance = EStanceMode::Invalid;
+		if (StanceModifier)
+		{
+			OldActiveStance = StanceModifier->ActiveStance;
+		}
 	
-	const bool bIsCrouching = HasGameplayTag(Mover_IsCrouching, true);
-	if (bIsCrouching && (!bWantsToCrouch || !CanCrouch()))
-	{	
-		if (StanceModifier && StanceModifier->CanExpand(this))
+		const bool bIsCrouching = HasGameplayTag(Mover_IsCrouching, true);
+		if (bIsCrouching && (!bWantsToCrouch || !CanCrouch()))
+		{	
+			if (StanceModifier && StanceModifier->CanExpand(this))
+			{
+				CancelModifierFromHandle(StanceModifier->GetHandle());
+				StanceModifierHandle.Invalidate();
+
+				StanceModifier = nullptr;
+			}
+		}
+		else if (!bIsCrouching && bWantsToCrouch && CanCrouch())
 		{
-			CancelModifierFromHandle(StanceModifier->GetHandle());
-			StanceModifierHandle.Invalidate();
+			TSharedPtr<FStanceModifier> NewStanceModifier = MakeShared<FStanceModifier>();
+			StanceModifierHandle = QueueMovementModifier(NewStanceModifier);
 
-			StanceModifier = nullptr;
+			StanceModifier = NewStanceModifier.Get();
 		}
-	}
-	else if (!bIsCrouching && bWantsToCrouch && CanCrouch())
-	{
-		TSharedPtr<FStanceModifier> NewStanceModifier = MakeShared<FStanceModifier>();
-		StanceModifierHandle = QueueMovementModifier(NewStanceModifier);
-
-		StanceModifier = NewStanceModifier.Get();
-	}
 	
-	EStanceMode NewActiveStance = EStanceMode::Invalid;
-	if (StanceModifier)
-	{
-		NewActiveStance = StanceModifier->ActiveStance;
-	}
+		EStanceMode NewActiveStance = EStanceMode::Invalid;
+		if (StanceModifier)
+		{
+			NewActiveStance = StanceModifier->ActiveStance;
+		}
 
-	if (OldActiveStance != NewActiveStance)
-	{
-		OnStanceChanged.Broadcast(OldActiveStance, NewActiveStance);
+		if (OldActiveStance != NewActiveStance)
+		{
+			OnStanceChanged.Broadcast(OldActiveStance, NewActiveStance);
+		}
 	}
 }
diff --git a/Engine/Plugins/Experimental/Mover/Source/Mover/Public/DefaultMovementSet/CharacterMoverComponent.h b/Engine/Plugins/Experimental/Mover/Source/Mover/Public/DefaultMovementSet/CharacterMoverComponent.h
--- a/Engine/Plugins/Experimental/Mover/Source/Mover/Public/DefaultMovementSet/CharacterMoverComponent.h
+++ b/Engine/Plugins/Experimental/Mover/Source/Mover/Public/DefaultMovementSet/CharacterMoverComponent.h
@@ -13,6 +13,7 @@
  */
 DECLARE_DYNAMIC_MULTICAST_DELEGATE_TwoParams(FMover_OnStanceChanged, EStanceMode, OldStance, EStanceMode, NewStance);
 
+
 UCLASS(BlueprintType, meta = (BlueprintSpawnableComponent))
 class MOVER_API UCharacterMoverComponent : public UMoverComponent
 {
@@ -23,6 +24,14 @@ public:
 	
 	virtual void BeginPlay() override;
 	
+	// Whether this component should directly handle jumping or not 
+	UPROPERTY(EditAnywhere, BlueprintReadWrite, Category = "Mover|Character")
+	bool bHandleJump = true;
+
+	// Whether this component should directly handle stance changes, including crouching input
+	UPROPERTY(EditAnywhere, BlueprintReadWrite, Category = "Mover|Character")
+	bool bHandleStanceChanges = true;
+
 	/** Returns true if currently crouching */ 
 	UFUNCTION(BlueprintCallable, Category = Mover)
 	virtual bool IsCrouching() const;
@@ -58,10 +67,6 @@ public:
 	// Perform jump on actor
 	UFUNCTION(BlueprintCallable, Category = Mover)
 	virtual bool Jump();
-
-	// Whether this component should directly handle jumping or not 
-	UPROPERTY(EditAnywhere, BlueprintReadWrite, Category=Mover)
-    bool bHandleJump;
 	
 	// Whether this actor can currently crouch or not 
 	UFUNCTION(BlueprintCallable, Category = Mover)
