Mover plugin changes in support of extending to new movement sets: exporting some magic numbers, exposing a public way to set updated component/visual component to blueprint

diff --git a/Engine/Plugins/Experimental/Mover/Source/Mover/Public/MoveLibrary/MovementUtils.h b/Engine/Plugins/Experimental/Mover/Source/Mover/Public/MoveLibrary/MovementUtils.h
--- a/Engine/Plugins/Experimental/Mover/Source/Mover/Public/MoveLibrary/MovementUtils.h
+++ b/Engine/Plugins/Experimental/Mover/Source/Mover/Public/MoveLibrary/MovementUtils.h
@@ -17,8 +17,8 @@ class UMoverComponent;
 
 namespace UE::MoverUtils
 {
-	extern const double SMALL_MOVE_DISTANCE;
-	extern const double VERTICAL_SLOPE_NORMAL_MAX_DOT;
+	MOVER_API extern const double SMALL_MOVE_DISTANCE;
+	MOVER_API extern const double VERTICAL_SLOPE_NORMAL_MAX_DOT;
 }
 
 /** Encapsulates detailed trajectory sample info, from a move that has already occurred or one projected into the future */
diff --git a/Engine/Plugins/Experimental/Mover/Source/Mover/Public/MoverComponent.h b/Engine/Plugins/Experimental/Mover/Source/Mover/Public/MoverComponent.h
--- a/Engine/Plugins/Experimental/Mover/Source/Mover/Public/MoverComponent.h
+++ b/Engine/Plugins/Experimental/Mover/Source/Mover/Public/MoverComponent.h
@@ -271,7 +271,12 @@ public:	// Queries
 	// Get the transform of the root component that our Mover simulation is moving
 	FTransform GetUpdatedComponentTransform() const;
 
+	// Sets which component we're using as the root of our movement
+	UFUNCTION(BlueprintCallable, Category = Mover)
+	void SetUpdatedComponent(USceneComponent* NewUpdatedComponent);
+
 	// Access the root component of the actor that our Mover simulation is moving
+	UFUNCTION(BlueprintCallable, Category = Mover)
 	USceneComponent* GetUpdatedComponent() const;
 
 	// Typed accessor to root moving component
@@ -283,6 +288,7 @@ public:	// Queries
 	}
 
 	// Access the primary visual component of the actor
+	UFUNCTION(BlueprintCallable, Category = Mover)
 	USceneComponent* GetPrimaryVisualComponent() const;
 
 	// Typed accessor to primary visual component
@@ -471,8 +477,6 @@ public:
 	void HandleImpact(FMoverOnImpactParams& ImpactParams);
 
 protected:
-	// Basic "Update Component/Ticking"
-	void SetUpdatedComponent(USceneComponent* NewUpdatedComponent);
 	void FindDefaultUpdatedComponent();
 	void UpdateTickRegistration();
 
