[Mover] Added a check for ParticleHandle before accessing it

diff --git a/Engine/Plugins/Experimental/Mover/Source/Mover/Private/Backends/MoverNetworkPhysicsLiaisonBase.cpp b/Engine/Plugins/Experimental/Mover/Source/Mover/Private/Backends/MoverNetworkPhysicsLiaisonBase.cpp
--- a/Engine/Plugins/Experimental/Mover/Source/Mover/Private/Backends/MoverNetworkPhysicsLiaisonBase.cpp
+++ b/Engine/Plugins/Experimental/Mover/Source/Mover/Private/Backends/MoverNetworkPhysicsLiaisonBase.cpp
@@ -686,9 +686,11 @@ void UMoverNetworkPhysicsLiaisonComponentBase::PerformProcessInputs_Internal(int
 	FMoverAuxStateContext UnusedAuxState;
 	GetMoverComponent().OnPreMovement.Broadcast(MoverTimeStep, Input.InputCmd, Input.SyncState, UnusedAuxState);
 
-	Chaos::FPBDRigidParticleHandle* ParticleHandle = GetControlledParticle_Internal();
-	FMoverDefaultSyncState& InputSyncState = Input.SyncState.SyncStateCollection.FindOrAddMutableDataByType<FMoverDefaultSyncState>();
-	InputSyncState.SetTransforms_WorldSpace(ParticleHandle->GetX(), FRotator(ParticleHandle->GetR()), ParticleHandle->GetV());
+	if (Chaos::FPBDRigidParticleHandle* ParticleHandle = GetControlledParticle_Internal())
+	{
+		FMoverDefaultSyncState& InputSyncState = Input.SyncState.SyncStateCollection.FindOrAddMutableDataByType<FMoverDefaultSyncState>();
+		InputSyncState.SetTransforms_WorldSpace(ParticleHandle->GetX(), FRotator(ParticleHandle->GetR()), ParticleHandle->GetV());
+	}
 }
 
 void UMoverNetworkPhysicsLiaisonComponentBase::PreSimulate_Internal(const FPhysicsMoverSimulationTickParams& TickParams, const FPhysicsMoverAsyncInput& Input, OUT FPhysicsMoverAsyncOutput& Output) const
