[Mover] [Crash] Safer CanExpand function for StanceModifier

diff --git a/Engine/Plugins/Experimental/Mover/Source/Mover/Private/PhysicsMover/MovementModifiers/PhysicsStanceModifier.cpp b/Engine/Plugins/Experimental/Mover/Source/Mover/Private/PhysicsMover/MovementModifiers/PhysicsStanceModifier.cpp
--- a/Engine/Plugins/Experimental/Mover/Source/Mover/Private/PhysicsMover/MovementModifiers/PhysicsStanceModifier.cpp
+++ b/Engine/Plugins/Experimental/Mover/Source/Mover/Private/PhysicsMover/MovementModifiers/PhysicsStanceModifier.cpp
@@ -3,6 +3,7 @@
 #include "PhysicsMover/MovementModifiers/PhysicsStanceModifier.h"
 
 #include "Components/CapsuleComponent.h"
+#include "DefaultMovementSet/CharacterMoverComponent.h"
 #include "DefaultMovementSet/Settings/StanceSettings.h"
 #include "MoveLibrary/MovementUtils.h"
 #include "PhysicsMover/Modes/PhysicsDrivenWalkingMode.h"
@@ -19,6 +20,68 @@ void FPhysicsStanceModifier::AdjustCapsule(UMoverComponent* MoverComp, float Old
 	UpdateTargetHeight<UPhysicsDrivenFallingMode>(MoverComp, OldHalfHeight, NewHalfHeight);
 }
 
+bool FPhysicsStanceModifier::CanExpand_Internal(UMoverComponent* MoverComponent, USceneComponent* UpdatedComp, const FMoverSyncState& SyncState) const
+{
+	float StandingHalfHeight = 90;
+	float CurrentHalfHeight = 55;
+
+	bool bEncroached = true;
+
+	if (const UCharacterMoverComponent* CharMoverComp = Cast<UCharacterMoverComponent>(MoverComponent))
+	{
+		const USceneComponent* UpdatedComponent = UpdatedComp;
+		const UPrimitiveComponent* UpdatedCompAsPrimitive = Cast<UPrimitiveComponent>(UpdatedComponent);
+
+		if (const UCapsuleComponent* OriginalCapsule = UMovementUtils::GetOriginalComponentType<UCapsuleComponent>(CharMoverComp->GetOwner()))
+		{
+			StandingHalfHeight = OriginalCapsule->GetScaledCapsuleHalfHeight();
+		}
+
+		if (const UCapsuleComponent* CapsuleComponent = Cast<UCapsuleComponent>(UpdatedCompAsPrimitive))
+		{
+			CurrentHalfHeight = CapsuleComponent->GetScaledCapsuleHalfHeight();
+		}
+
+		const float HalfHeightDifference = StandingHalfHeight - CurrentHalfHeight;
+
+		const FMoverDefaultSyncState* DefaultSyncState = SyncState.SyncStateCollection.FindDataByType<FMoverDefaultSyncState>();
+
+		const FVector PawnLocation = DefaultSyncState->GetLocation_WorldSpace();
+		const FQuat PawnRot = DefaultSyncState->GetOrientation_WorldSpace().Quaternion();
+		float PawnRadius = 0.0f;
+		float PawnHalfHeight = 0.0f;
+		UpdatedCompAsPrimitive->CalcBoundingCylinder(PawnRadius, PawnHalfHeight);
+
+		// TODO: Compensate for the difference between current capsule size and standing size
+		const FCollisionShape StandingCapsuleShape = FCollisionShape::MakeCapsule(PawnRadius, StandingHalfHeight);
+		const ECollisionChannel CollisionChannel = UpdatedCompAsPrimitive->GetCollisionObjectType();
+
+		// TODO: @Harsha Switch to physics thread safe IsOnGround_Internal() method when available.
+		bool bShouldMaintainBase = false;
+		if (const TObjectPtr<UBaseMovementMode>* CurrentMode = CharMoverComp->MovementModes.Find(SyncState.MovementMode))
+		{
+			const UBaseMovementMode* ActiveMode = CurrentMode->Get();
+			if (ActiveMode && ActiveMode->HasGameplayTag(Mover_IsOnGround, true))
+			{
+				bShouldMaintainBase = true;
+			}
+		}
+
+		if (!bShouldMaintainBase)
+		{
+			// Expand in place
+			bEncroached = UMovementUtils::OverlapTest(UpdatedComponent, UpdatedCompAsPrimitive, PawnLocation, PawnRot, CollisionChannel, StandingCapsuleShape, CharMoverComp->GetOwner());
+		}
+		else
+		{
+			// Expand while keeping base location the same.
+			const FVector StandingLocation = PawnLocation + (HalfHeightDifference + .01f) * CharMoverComp->GetUpDirection();
+			bEncroached = UMovementUtils::OverlapTest(UpdatedComponent, UpdatedCompAsPrimitive, StandingLocation, PawnRot, CollisionChannel, StandingCapsuleShape, CharMoverComp->GetOwner());
+		}
+	}
+	return !bEncroached;
+}
+
 void FPhysicsStanceModifier::OnPostSimulationTick(const FStanceModifier* Modifier, UMoverComponent* MoverComp, UPrimitiveComponent* UpdatedPrimitive,  bool bIsCrouching, bool& bPostProcessed, OUT bool& bStanceChanged)
 {
 	if(const UStanceSettings* StanceSettings = MoverComp->FindSharedSettings<UStanceSettings>())
diff --git a/Engine/Plugins/Experimental/Mover/Source/Mover/Private/PhysicsMover/PhysicsCharacterMoverComponent.cpp b/Engine/Plugins/Experimental/Mover/Source/Mover/Private/PhysicsMover/PhysicsCharacterMoverComponent.cpp
--- a/Engine/Plugins/Experimental/Mover/Source/Mover/Private/PhysicsMover/PhysicsCharacterMoverComponent.cpp
+++ b/Engine/Plugins/Experimental/Mover/Source/Mover/Private/PhysicsMover/PhysicsCharacterMoverComponent.cpp
@@ -35,7 +35,7 @@ void UPhysicsCharacterMoverComponent::BeginPlay()
 
 void UPhysicsCharacterMoverComponent::OnMoverPreMovement(const FMoverTimeStep& TimeStep, const FMoverInputCmdContext& InputCmd, const FMoverSyncState& SyncState, const FMoverAuxStateContext& AuxState)
 {
-	const FStanceModifier* StanceModifier = static_cast<const FStanceModifier*>(FindMovementModifier(StanceModifierHandle));
+	const FPhysicsStanceModifier* StanceModifier = static_cast<const FPhysicsStanceModifier*>(FindMovementModifier(StanceModifierHandle));
 	// This is a fail safe in case our handle was bad - try finding the modifier by type if we can
 	if (!StanceModifier)
 	{
@@ -51,7 +51,7 @@ void UPhysicsCharacterMoverComponent::OnMoverPreMovement(const FMoverTimeStep& T
 	const bool bIsCrouching = HasGameplayTag(Mover_IsCrouching, true);
 	if (bIsCrouching && (!bWantsToCrouch || !CanCrouch()))
 	{
-		if (StanceModifier && StanceModifier->CanExpand(this))
+		if (StanceModifier && StanceModifier->CanExpand_Internal(this, UpdatedComponent, SyncState))
 		{
 			CancelModifierFromHandle(StanceModifier->GetHandle());
 			StanceModifierHandle.Invalidate();
@@ -61,7 +61,7 @@ void UPhysicsCharacterMoverComponent::OnMoverPreMovement(const FMoverTimeStep& T
 	}
 	else if (!bIsCrouching && bWantsToCrouch && CanCrouch())
 	{
-		TSharedPtr<FStanceModifier> NewStanceModifier = MakeShared<FPhysicsStanceModifier>();
+		TSharedPtr<FPhysicsStanceModifier> NewStanceModifier = MakeShared<FPhysicsStanceModifier>();
 		StanceModifierHandle = QueueMovementModifier(NewStanceModifier);
 
 		StanceModifier = NewStanceModifier.Get();
@@ -76,7 +76,7 @@ void UPhysicsCharacterMoverComponent::OnMoverPreMovement(const FMoverTimeStep& T
 
 void UPhysicsCharacterMoverComponent::OnMoverPostSimulationTick(const FMoverTimeStep& TimeStep)
 {
-	const FStanceModifier* StanceModifier = static_cast<const FStanceModifier*>(FindMovementModifier(StanceModifierHandle));
+	const FPhysicsStanceModifier* StanceModifier = static_cast<const FPhysicsStanceModifier*>(FindMovementModifier(StanceModifierHandle));
 	if (!StanceModifier)
 	{
 		StanceModifier = FindMovementModifierByType<FPhysicsStanceModifier>();
diff --git a/Engine/Plugins/Experimental/Mover/Source/Mover/Public/PhysicsMover/MovementModifiers/PhysicsStanceModifier.h b/Engine/Plugins/Experimental/Mover/Source/Mover/Public/PhysicsMover/MovementModifiers/PhysicsStanceModifier.h
--- a/Engine/Plugins/Experimental/Mover/Source/Mover/Public/PhysicsMover/MovementModifiers/PhysicsStanceModifier.h
+++ b/Engine/Plugins/Experimental/Mover/Source/Mover/Public/PhysicsMover/MovementModifiers/PhysicsStanceModifier.h
@@ -27,6 +27,8 @@ struct MOVER_API FPhysicsStanceModifier : public FStanceModifier
 
 	virtual void AddReferencedObjects(class FReferenceCollector& Collector) override;
 
+	virtual bool CanExpand_Internal(UMoverComponent* MoverComp, USceneComponent* UpdatedComponent, const FMoverSyncState& InSyncState) const;
+
 	static void OnPostSimulationTick(const FStanceModifier* Modifier, UMoverComponent* MoverComp, UPrimitiveComponent* UpdatedPrimitive,  bool bIsCrouching, bool& bPostProcessed, OUT bool& bStanceChanged);
 
 protected:
