Mover: improving input handling options and skipping subscribing to the presim event unless one of these is enabled

diff --git a/Engine/Plugins/Experimental/Mover/Source/Mover/Private/DefaultMovementSet/CharacterMoverComponent.cpp b/Engine/Plugins/Experimental/Mover/Source/Mover/Private/DefaultMovementSet/CharacterMoverComponent.cpp
--- a/Engine/Plugins/Experimental/Mover/Source/Mover/Private/DefaultMovementSet/CharacterMoverComponent.cpp
+++ b/Engine/Plugins/Experimental/Mover/Source/Mover/Private/DefaultMovementSet/CharacterMoverComponent.cpp
@@ -24,7 +24,35 @@ void UCharacterMoverComponent::BeginPlay()
 {
 	Super::BeginPlay();
 
-	OnPreSimulationTick.AddDynamic(this, &UCharacterMoverComponent::OnMoverPreSimulationTick);
+	OnHandlerSettingChanged();
+}
+
+bool UCharacterMoverComponent::GetHandleJump() const
+{
+	return bHandleJump;
+}
+
+void UCharacterMoverComponent::SetHandleJump(bool bInHandleJump)
+{
+	if (bHandleJump != bInHandleJump)
+	{
+		bHandleJump = bInHandleJump;
+		OnHandlerSettingChanged();
+	}
+}
+
+bool UCharacterMoverComponent::GetHandleStanceChanges() const
+{
+	return bHandleStanceChanges;
+}
+
+void UCharacterMoverComponent::SetHandleStanceChanges(bool bInHandleStanceChanges)
+{
+	if (bHandleStanceChanges != bInHandleStanceChanges)
+	{
+		bHandleStanceChanges = bInHandleStanceChanges;
+		OnHandlerSettingChanged();
+	}
 }
 
 bool UCharacterMoverComponent::IsCrouching() const
@@ -167,3 +195,17 @@ void UCharacterMoverComponent::OnMoverPreSimulationTick(const FMoverTimeStep& Ti
 		}
 	}
 }
+
+void UCharacterMoverComponent::OnHandlerSettingChanged()
+{
+	const bool bIsHandlingAnySettings = bHandleJump || bHandleStanceChanges;
+
+	if (bIsHandlingAnySettings)
+	{
+		OnPreSimulationTick.AddUniqueDynamic(this, &UCharacterMoverComponent::OnMoverPreSimulationTick);
+	}
+	else
+	{
+		OnPreSimulationTick.RemoveDynamic(this, &UCharacterMoverComponent::OnMoverPreSimulationTick);
+	}
+}
\ No newline at end of file
diff --git a/Engine/Plugins/Experimental/Mover/Source/Mover/Public/DefaultMovementSet/CharacterMoverComponent.h b/Engine/Plugins/Experimental/Mover/Source/Mover/Public/DefaultMovementSet/CharacterMoverComponent.h
--- a/Engine/Plugins/Experimental/Mover/Source/Mover/Public/DefaultMovementSet/CharacterMoverComponent.h
+++ b/Engine/Plugins/Experimental/Mover/Source/Mover/Public/DefaultMovementSet/CharacterMoverComponent.h
@@ -7,7 +7,7 @@
 #include "CharacterMoverComponent.generated.h"
 
 /**
- * Fires when a stance is changed
+ * Fires when a stance is changed, if stance handling is enabled (see @SetHandleStanceChanges)
  * Note: If a stance was just Activated it will fire with an invalid OldStance
  *		 If a stance was just Deactivated it will fire with an invalid NewStance
  */
@@ -23,14 +23,23 @@ public:
 	UCharacterMoverComponent();
 	
 	virtual void BeginPlay() override;
-	
-	// Whether this component should directly handle jumping or not 
-	UPROPERTY(EditAnywhere, BlueprintReadWrite, Category = "Mover|Character")
-	bool bHandleJump = true;
 
-	// Whether this component should directly handle stance changes, including crouching input
-	UPROPERTY(EditAnywhere, BlueprintReadWrite, Category = "Mover|Character")
-	bool bHandleStanceChanges = true;
+
+	// Returns whether this component is tasked with handling jump input or not
+	UFUNCTION(BlueprintGetter)
+	bool GetHandleJump() const;
+
+	// If true, this component will handle default character inputs for jumping
+	UFUNCTION(BlueprintSetter)
+	void SetHandleJump(bool bInHandleJump);
+
+	// Returns whether this component is tasked with handling character stance changes, including crouching
+	UFUNCTION(BlueprintGetter)
+	bool GetHandleStanceChanges() const;
+
+	// If true, this component will process stancing changes and crouching inputs
+	UFUNCTION(BlueprintSetter)
+	void SetHandleStanceChanges(bool bInHandleStanceChanges);
 
 	/** Returns true if currently crouching */ 
 	UFUNCTION(BlueprintCallable, Category = Mover)
@@ -88,10 +97,20 @@ protected:
 	UFUNCTION()
 	virtual void OnMoverPreSimulationTick(const FMoverTimeStep& TimeStep, const FMoverInputCmdContext& InputCmd);
 	
+	virtual void OnHandlerSettingChanged();
+
 	// ID used to keep track of the modifier responsible for crouching
 	FMovementModifierHandle StanceModifierHandle;
 
 	/** If true, try to crouch (or keep crouching) on next update. If false, try to stop crouching on next update. */
-	UPROPERTY(Category = "Mover|Crouch", VisibleInstanceOnly, BlueprintReadOnly)
-	bool bWantsToCrouch;
+	UPROPERTY(BlueprintReadOnly, Transient, Category = "Mover|Crouch")
+	uint8 bWantsToCrouch : 1 = 0;
+
+	// Whether this component should directly handle jumping or not 
+	UPROPERTY(EditAnywhere, BlueprintGetter = GetHandleJump, BlueprintSetter = SetHandleJump, Category = "Mover|Character")
+	uint8 bHandleJump : 1 = 1;
+
+	// Whether this component should directly handle stance changes, including crouching input
+	UPROPERTY(EditAnywhere, BlueprintGetter = GetHandleStanceChanges, BlueprintSetter = SetHandleStanceChanges, Category = "Mover|Character")
+	uint8 bHandleStanceChanges : 1 = 1;
 };
