Fixes Initialization of collision params in UPhysicsMovementUtils::FloorSweep_Internal

#rnx

diff --git a/Engine/Plugins/Experimental/Mover/Source/Mover/Private/PhysicsMover/PhysicsMovementUtils.cpp b/Engine/Plugins/Experimental/Mover/Source/Mover/Private/PhysicsMover/PhysicsMovementUtils.cpp
--- a/Engine/Plugins/Experimental/Mover/Source/Mover/Private/PhysicsMover/PhysicsMovementUtils.cpp
+++ b/Engine/Plugins/Experimental/Mover/Source/Mover/Private/PhysicsMover/PhysicsMovementUtils.cpp
@@ -47,6 +47,8 @@ void UPhysicsMovementUtils::FloorSweep_Internal(
 		ResponseParams.CollisionResponse.SetResponse(ECC_Destructible, ECR_Block);
 		ResponseParams.CollisionResponse.SetResponse(ECC_PhysicsBody, ECR_Block);
 
+		UpdatedPrimitive->InitSweepCollisionParams(QueryParams, ResponseParams);
+
 		TArray<FHitResult> Hits;
 
 		const float DeltaPosVertLength = DeltaPos.Dot(UpDir);
