[Pathed Mover] Explicitly init the PathOrigin to inifinity to satisfy the automated property init tests and keep it clear when it's been set to an actually valid transform

diff --git a/Engine/Plugins/Experimental/Mover/Source/Mover/Public/PhysicsMover/PathedMovement/PathedMovementTypes.h b/Engine/Plugins/Experimental/Mover/Source/Mover/Public/PhysicsMover/PathedMovement/PathedMovementTypes.h
--- a/Engine/Plugins/Experimental/Mover/Source/Mover/Public/PhysicsMover/PathedMovement/PathedMovementTypes.h
+++ b/Engine/Plugins/Experimental/Mover/Source/Mover/Public/PhysicsMover/PathedMovement/PathedMovementTypes.h
@@ -53,7 +53,7 @@ struct FMutablePathedMovementProperties
 	EPathedPhysicsPlaybackBehavior PlaybackBehavior = EPathedPhysicsPlaybackBehavior::OneShot;
 
 	UPROPERTY()
-	FTransform PathOrigin = FTransform(NoInit);
+	FTransform PathOrigin = FTransform(FVector(HUGE_VAL));
 
 	bool IsMoving() const;
 	bool IsLooping() const;
