Mover: Providing a gameplay tag that can be used to skip proposed movement in RootMotionAttributeLayeredMoves

diff --git a/Engine/Plugins/Experimental/Mover/Source/Mover/Private/DefaultMovementSet/LayeredMoves/RootMotionAttributeLayeredMove.cpp b/Engine/Plugins/Experimental/Mover/Source/Mover/Private/DefaultMovementSet/LayeredMoves/RootMotionAttributeLayeredMove.cpp
--- a/Engine/Plugins/Experimental/Mover/Source/Mover/Private/DefaultMovementSet/LayeredMoves/RootMotionAttributeLayeredMove.cpp
+++ b/Engine/Plugins/Experimental/Mover/Source/Mover/Private/DefaultMovementSet/LayeredMoves/RootMotionAttributeLayeredMove.cpp
@@ -123,11 +123,10 @@ bool FLayeredMove_RootMotionAttribute::GenerateMove(const FMoverTickStartData& S
 		const FMoverDefaultSyncState* SyncState = SimState.SyncState.SyncStateCollection.FindDataByType<FMoverDefaultSyncState>();
 		const FCharacterDefaultInputs* InputCmd = SimState.InputCmd.InputCollection.FindDataByType<FCharacterDefaultInputs>();
 
-		// TODO: Provide better options to control when we should skip proposed movement
 		if (!bHasValidWarpingContext &&
-				((InputCmd && InputCmd->bIsJumpJustPressed) || MoverComp->HasGameplayTag(Mover_IsFalling, false)))
+				((InputCmd && InputCmd->bIsJumpJustPressed) || MoverComp->HasGameplayTag(Mover_SkipAnimRootMotion, false)))
 		{
-			return false;	// do not perform root motion while jumping or falling, so that we can have air control, unless we're under montage influence
+			return false;	// do not perform root motion while we have the SkipAnimRootMotion tag. i.e. jumping or falling, so that we can have air control, unless we're under montage influence
 		}
 
 		// Note that we're forcing the use of the sync state's actor transform data. This is necessary when the movement simulation 
diff --git a/Engine/Plugins/Experimental/Mover/Source/Mover/Private/DefaultMovementSet/Modes/AsyncFallingMode.cpp b/Engine/Plugins/Experimental/Mover/Source/Mover/Private/DefaultMovementSet/Modes/AsyncFallingMode.cpp
--- a/Engine/Plugins/Experimental/Mover/Source/Mover/Private/DefaultMovementSet/Modes/AsyncFallingMode.cpp
+++ b/Engine/Plugins/Experimental/Mover/Source/Mover/Private/DefaultMovementSet/Modes/AsyncFallingMode.cpp
@@ -30,6 +30,7 @@ UAsyncFallingMode::UAsyncFallingMode(const FObjectInitializer& ObjectInitializer
 
 	GameplayTags.AddTag(Mover_IsInAir);
 	GameplayTags.AddTag(Mover_IsFalling);
+	GameplayTags.AddTag(Mover_SkipAnimRootMotion);
 }
 
 
diff --git a/Engine/Plugins/Experimental/Mover/Source/Mover/Private/DefaultMovementSet/Modes/FallingMode.cpp b/Engine/Plugins/Experimental/Mover/Source/Mover/Private/DefaultMovementSet/Modes/FallingMode.cpp
--- a/Engine/Plugins/Experimental/Mover/Source/Mover/Private/DefaultMovementSet/Modes/FallingMode.cpp
+++ b/Engine/Plugins/Experimental/Mover/Source/Mover/Private/DefaultMovementSet/Modes/FallingMode.cpp
@@ -27,6 +27,7 @@ UFallingMode::UFallingMode(const FObjectInitializer& ObjectInitializer)
 
 	GameplayTags.AddTag(Mover_IsInAir);
 	GameplayTags.AddTag(Mover_IsFalling);
+	GameplayTags.AddTag(Mover_SkipAnimRootMotion);
 }
 
 
diff --git a/Engine/Plugins/Experimental/Mover/Source/Mover/Private/MoverTypes.cpp b/Engine/Plugins/Experimental/Mover/Source/Mover/Private/MoverTypes.cpp
--- a/Engine/Plugins/Experimental/Mover/Source/Mover/Private/MoverTypes.cpp
+++ b/Engine/Plugins/Experimental/Mover/Source/Mover/Private/MoverTypes.cpp
@@ -15,6 +15,7 @@ UE_DEFINE_GAMEPLAY_TAG_COMMENT(Mover_IsFlying, "Mover.IsFlying", "Default Mover
 UE_DEFINE_GAMEPLAY_TAG_COMMENT(Mover_IsSwimming, "Mover.IsSwimming", "Default Mover state flag indicating character is swimming.");
 UE_DEFINE_GAMEPLAY_TAG_COMMENT(Mover_IsCrouching, "Mover.Stance.IsCrouching", "Default Mover state flag indicating character is crouching.");
 UE_DEFINE_GAMEPLAY_TAG_COMMENT(Mover_IsNavWalking, "Mover.IsNavWalking", "Default Mover state flag indicating character is NavWalking.");
+UE_DEFINE_GAMEPLAY_TAG_COMMENT(Mover_SkipAnimRootMotion, "Mover.SkipAnimRootMotion", "Default Mover state flag indicating Animation Root Motion proposed movement should be skipped.");
 
 FMoverOnImpactParams::FMoverOnImpactParams() 
 	: AttemptedMoveDelta(0) 
diff --git a/Engine/Plugins/Experimental/Mover/Source/Mover/Public/MoverTypes.h b/Engine/Plugins/Experimental/Mover/Source/Mover/Public/MoverTypes.h
--- a/Engine/Plugins/Experimental/Mover/Source/Mover/Public/MoverTypes.h
+++ b/Engine/Plugins/Experimental/Mover/Source/Mover/Public/MoverTypes.h
@@ -18,6 +18,7 @@ MOVER_API UE_DECLARE_GAMEPLAY_TAG_EXTERN(Mover_IsFlying);
 MOVER_API UE_DECLARE_GAMEPLAY_TAG_EXTERN(Mover_IsSwimming);
 MOVER_API UE_DECLARE_GAMEPLAY_TAG_EXTERN(Mover_IsCrouching);
 MOVER_API UE_DECLARE_GAMEPLAY_TAG_EXTERN(Mover_IsNavWalking);
+MOVER_API UE_DECLARE_GAMEPLAY_TAG_EXTERN(Mover_SkipAnimRootMotion);
 
 /** List of Mover tick phases where different parts of work are being performed */
 UENUM(BlueprintType)
