[Pathed Mover] Let FMutablePathedMovementProperties::PathOrigin init nice and clean with a default ctor, as automation tests were ticked about both ENoInit and HUGE_VAL
- Instead of checking the transform for validity before comparing in UMoverPathedPhysicsLiaisonComponent::PerformProcessInputs_Internal, check to see if the input sync state arrived with an existing movement state on it at all. If not, it's the first frame, and we can act accordingly

diff --git a/Engine/Plugins/Experimental/Mover/Source/Mover/Private/Backends/MoverPathedPhysicsLiaison.cpp b/Engine/Plugins/Experimental/Mover/Source/Mover/Private/Backends/MoverPathedPhysicsLiaison.cpp
--- a/Engine/Plugins/Experimental/Mover/Source/Mover/Private/Backends/MoverPathedPhysicsLiaison.cpp
+++ b/Engine/Plugins/Experimental/Mover/Source/Mover/Private/Backends/MoverPathedPhysicsLiaison.cpp
@@ -231,15 +231,24 @@ void UMoverPathedPhysicsLiaisonComponent::PerformProcessInputs_Internal(int32 Ph
 	{
 		Chaos::FPBDRigidParticleHandle& ParticleHandle = *GetControlledParticle_Internal();
 		Chaos::FJointConstraint& JointConstraint = *GetJointConstraint();
-		
-		FPathedPhysicsMovementState& InputMoveState = Input.SyncState.SyncStateCollection.FindOrAddMutableDataByType<FPathedPhysicsMovementState>();
-		FMutablePathedMovementProperties& CurProps = InputMoveState.MutableProps;
+
+		bool bIsFirstProcess = false;
+		FPathedPhysicsMovementState* InputMoveState = Input.SyncState.SyncStateCollection.FindMutableDataByType<FPathedPhysicsMovementState>();
+		if (!InputMoveState)
+		{
+			// If there isn't an existing FPathedPhysicsMovementState, this is the first time we're processing input and establishing the MoveState on the SyncState
+			bIsFirstProcess = true;
+			InputMoveState = &Input.SyncState.SyncStateCollection.FindOrAddMutableDataByType<FPathedPhysicsMovementState>();
+			InputMoveState->MutableProps = PathedMovementInputs->Props;
+		}
+
+		FMutablePathedMovementProperties& CurProps = InputMoveState->MutableProps;
 		const FMutablePathedMovementProperties& InputProps = PathedMovementInputs->Props;
 
 		Chaos::FPBDRigidsEvolution& Evolution = *GetPhysicsSolver()->GetEvolution();
 
 		// Toggle whether the joint is enabled and update the controlled particle to be kinematic/dynamic accordingly
-		if (InputProps.bIsJointEnabled != CurProps.bIsJointEnabled)
+		if (bIsFirstProcess || InputProps.bIsJointEnabled != CurProps.bIsJointEnabled)
 		{
 			if (InputProps.bIsJointEnabled)
 			{
@@ -252,7 +261,7 @@ void UMoverPathedPhysicsLiaisonComponent::PerformProcessInputs_Internal(int32 Ph
 				{
 					if (PathedMode)
 					{
-						TargetLastFrame = PathedMode->CalcTargetRelativeTransform(InputMoveState.CurrentProgress);
+						TargetLastFrame = PathedMode->CalcTargetRelativeTransform(InputMoveState->CurrentProgress);
 					}
 					else
 					{
@@ -271,11 +280,11 @@ void UMoverPathedPhysicsLiaisonComponent::PerformProcessInputs_Internal(int32 Ph
 			}
 		}
 		
-
-		if (InputMoveState.MutableProps.PathOrigin.IsValid())
+		// Has the path origin shifted?
+		if (!bIsFirstProcess)
 		{
-			const FVector OriginLocationDelta = InputMoveState.MutableProps.PathOrigin.GetLocation() - InputProps.PathOrigin.GetLocation();
-			const FQuat OriginRotationDelta = InputMoveState.MutableProps.PathOrigin.GetRotation() - InputProps.PathOrigin.GetRotation();
+			const FVector OriginLocationDelta = CurProps.PathOrigin.GetLocation() - InputProps.PathOrigin.GetLocation();
+			const FQuat OriginRotationDelta = CurProps.PathOrigin.GetRotation() - InputProps.PathOrigin.GetRotation();
 			if (!OriginLocationDelta.IsNearlyZero() || !FRotator(OriginRotationDelta).IsNearlyZero())
 			{
 				// Teleport both the controlled particle and the joint endpoint by the change in origin
diff --git a/Engine/Plugins/Experimental/Mover/Source/Mover/Public/PhysicsMover/PathedMovement/PathedMovementTypes.h b/Engine/Plugins/Experimental/Mover/Source/Mover/Public/PhysicsMover/PathedMovement/PathedMovementTypes.h
--- a/Engine/Plugins/Experimental/Mover/Source/Mover/Public/PhysicsMover/PathedMovement/PathedMovementTypes.h
+++ b/Engine/Plugins/Experimental/Mover/Source/Mover/Public/PhysicsMover/PathedMovement/PathedMovementTypes.h
@@ -53,7 +53,7 @@ struct FMutablePathedMovementProperties
 	EPathedPhysicsPlaybackBehavior PlaybackBehavior = EPathedPhysicsPlaybackBehavior::OneShot;
 
 	UPROPERTY()
-	FTransform PathOrigin = FTransform(NoInit);
+	FTransform PathOrigin;
 
 	bool IsMoving() const;
 	bool IsLooping() const;
