Mover - Add an option to transitions to only run on the first step

diff --git a/Engine/Plugins/Experimental/Mover/Source/Mover/Private/MovementModeStateMachine.cpp b/Engine/Plugins/Experimental/Mover/Source/Mover/Private/MovementModeStateMachine.cpp
--- a/Engine/Plugins/Experimental/Mover/Source/Mover/Private/MovementModeStateMachine.cpp
+++ b/Engine/Plugins/Experimental/Mover/Source/Mover/Private/MovementModeStateMachine.cpp
@@ -181,6 +181,7 @@ void UMovementModeStateMachine::OnSimulationTick(USceneComponent* UpdatedCompone
 
 	AdvanceToNextMode();
 
+	int SubStepCount = 0;
 	const int32 MaxConsecutiveFullRefundedSubsteps = GetDefault<UMoverDeveloperSettings>()->MaxTimesToRefundSubstep;
 	int32 NumConsecutiveFullRefundedSubsteps = 0;
 	while (SubTimeStep.BaseSimTimeMs < EndingSimTimeMs)
@@ -324,7 +325,7 @@ void UMovementModeStateMachine::OnSimulationTick(USceneComponent* UpdatedCompone
 
 				for (UBaseMovementModeTransition* Transition : CurrentMode->Transitions)
 				{
-					if (IsValid(Transition))
+					if (IsValid(Transition) && ((SubStepCount == 0) || !Transition->bFirstSubStepOnly))
 					{
 						EvalResult = Transition->Evaluate(SimTickParams);
 
@@ -409,6 +410,8 @@ void UMovementModeStateMachine::OnSimulationTick(USceneComponent* UpdatedCompone
 
 		SubstepStartData.SyncState = OutputState.SyncState;
 		SubstepStartData.AuxState  = OutputState.AuxState;
+
+		++SubStepCount;
 	}
 
 	FApplyMovementEffectParams EffectParams;
diff --git a/Engine/Plugins/Experimental/Mover/Source/Mover/Public/MovementModeTransition.h b/Engine/Plugins/Experimental/Mover/Source/Mover/Public/MovementModeTransition.h
--- a/Engine/Plugins/Experimental/Mover/Source/Mover/Public/MovementModeTransition.h
+++ b/Engine/Plugins/Experimental/Mover/Source/Mover/Public/MovementModeTransition.h
@@ -76,6 +76,12 @@ public:
 	/** Whether this transition should reenter a mode if it evaluates true and wants to transition into a mode the actor is already in */
 	UPROPERTY(EditDefaultsOnly, BlueprintReadWrite, Category = Mover)
 	bool bAllowModeReentry = false;
+
+	/** Whether this transition should only apply to the first step of the update. If true, modes reached after transitions or mode changes
+	* in the current update will not consider this transition
+	*/
+	UPROPERTY(EditDefaultsOnly, BlueprintReadWrite, Category = Mover)
+	bool bFirstSubStepOnly = false;
 	
 protected:
 	UFUNCTION(BlueprintImplementableEvent, Category = MovementModeTransition, meta = (DisplayName = "On Registered", ScriptName = "OnRegistered"))
