Physics Mover - Add launch input and handling

diff --git a/Engine/Plugins/Experimental/Mover/Source/Mover/Private/PhysicsMover/InstantMovementEffects/ApplyVelocityPhysicsMovementEffect.cpp b/Engine/Plugins/Experimental/Mover/Source/Mover/Private/PhysicsMover/InstantMovementEffects/ApplyVelocityPhysicsMovementEffect.cpp
new file mode 100644
--- /dev/null
+++ b/Engine/Plugins/Experimental/Mover/Source/Mover/Private/PhysicsMover/InstantMovementEffects/ApplyVelocityPhysicsMovementEffect.cpp
@@ -0,0 +1,104 @@
+// Copyright Epic Games, Inc. All Rights Reserved.
+
+#include "PhysicsMover/InstantMovementEffects/ApplyVelocityPhysicsMovementEffect.h"
+
+#include "Chaos/ParticleHandle.h"
+#include "MoverComponent.h"
+#include "MoverDataModelTypes.h"
+#include "MoverSimulationTypes.h"
+#include "MoveLibrary/MoverBlackboard.h"
+#include "PhysicsMover/PhysicsMovementUtils.h"
+
+FApplyVelocityPhysicsEffect::FApplyVelocityPhysicsEffect()
+	: VelocityToApply(FVector::ZeroVector)
+	, bAdditiveVelocity(false)
+	, ForceMovementMode(NAME_None)
+{
+}
+
+bool FApplyVelocityPhysicsEffect::ApplyMovementEffect(FApplyMovementEffectParams& ApplyEffectParams, FMoverSyncState& OutputState)
+{
+	// Get the position and orientation. Start by looking in the sync state.
+	// If not there, get from the particle corresponding to the updated component
+	FVector Position;
+	FRotator Orientation;
+	if (const FMoverDefaultSyncState* CurrentSyncState = OutputState.SyncStateCollection.FindDataByType<FMoverDefaultSyncState>())
+	{
+		Position = CurrentSyncState->GetLocation_WorldSpace();
+		Orientation = CurrentSyncState->GetOrientation_WorldSpace();
+	}
+	else if (Chaos::FPBDRigidParticleHandle* ParticleHandle = UPhysicsMovementUtils::GetRigidParticleHandleFromComponent(ApplyEffectParams.UpdatedPrimitive))
+	{
+		Position = ParticleHandle->GetX();
+		Orientation = FRotator(ParticleHandle->GetR());
+	}
+	else
+	{
+		return false;
+	}
+
+	FMoverDefaultSyncState& OutputSyncState = OutputState.SyncStateCollection.FindOrAddMutableDataByType<FMoverDefaultSyncState>();
+
+	OutputState.MovementMode = ForceMovementMode;
+
+	FRelativeBaseInfo MovementBaseInfo;
+	if (const UMoverBlackboard* SimBlackboard = ApplyEffectParams.MoverComp->GetSimBlackboard())
+	{
+		SimBlackboard->TryGet(CommonBlackboard::LastFoundDynamicMovementBase, MovementBaseInfo);
+	}
+
+	FVector Velocity = VelocityToApply;
+	if (bAdditiveVelocity)
+	{
+		if (const FMoverDefaultSyncState* SyncState = ApplyEffectParams.StartState->SyncState.SyncStateCollection.FindDataByType<FMoverDefaultSyncState>())
+		{
+			Velocity += SyncState->GetVelocity_WorldSpace();
+		}
+	}
+
+	OutputSyncState.SetTransforms_WorldSpace(
+		Position,
+		Orientation,
+		Velocity,
+		MovementBaseInfo.MovementBase.Get(),
+		MovementBaseInfo.BoneName);
+
+	return true;
+}
+
+FInstantMovementEffect* FApplyVelocityPhysicsEffect::Clone() const
+{
+	FApplyVelocityPhysicsEffect* CopyPtr = new FApplyVelocityPhysicsEffect(*this);
+	return CopyPtr;
+}
+
+void FApplyVelocityPhysicsEffect::NetSerialize(FArchive& Ar)
+{
+	Super::NetSerialize(Ar);
+
+	SerializePackedVector<10, 16>(VelocityToApply, Ar);
+	Ar.SerializeBits(&bAdditiveVelocity, 1);
+
+	bool bUsingForcedMovementMode = !ForceMovementMode.IsNone();
+	Ar.SerializeBits(&bUsingForcedMovementMode, 1);
+
+	if (bUsingForcedMovementMode)
+	{
+		Ar << ForceMovementMode;
+	}
+}
+
+UScriptStruct* FApplyVelocityPhysicsEffect::GetScriptStruct() const
+{
+	return FApplyVelocityPhysicsEffect::StaticStruct();
+}
+
+FString FApplyVelocityPhysicsEffect::ToSimpleString() const
+{
+	return FString::Printf(TEXT("ApplyVelocity"));
+}
+
+void FApplyVelocityPhysicsEffect::AddReferencedObjects(FReferenceCollector& Collector)
+{
+	Super::AddReferencedObjects(Collector);
+}
\ No newline at end of file
diff --git a/Engine/Plugins/Experimental/Mover/Source/Mover/Private/PhysicsMover/PhysicsMovementUtils.cpp b/Engine/Plugins/Experimental/Mover/Source/Mover/Private/PhysicsMover/PhysicsMovementUtils.cpp
--- a/Engine/Plugins/Experimental/Mover/Source/Mover/Private/PhysicsMover/PhysicsMovementUtils.cpp
+++ b/Engine/Plugins/Experimental/Mover/Source/Mover/Private/PhysicsMover/PhysicsMovementUtils.cpp
@@ -111,17 +111,28 @@ void UPhysicsMovementUtils::FloorSweep_Internal(
 	}
 }
 
-const Chaos::FPBDRigidParticleHandle* UPhysicsMovementUtils::GetRigidParticleHandleFromHitResult(const FHitResult& HitResult)
+Chaos::FPBDRigidParticleHandle* UPhysicsMovementUtils::GetRigidParticleHandleFromHitResult(const FHitResult& HitResult)
 {
 	if (IPhysicsComponent* PhysicsComp = Cast<IPhysicsComponent>(HitResult.Component))
 	{
 		if (Chaos::FPhysicsObjectHandle PhysicsObject = PhysicsComp->GetPhysicsObjectById(HitResult.Item))
 		{
 			Chaos::FReadPhysicsObjectInterface_Internal Interface = Chaos::FPhysicsObjectInternalInterface::GetRead();
-			if (const Chaos::FGeometryParticleHandle* Particle = Interface.GetParticle(PhysicsObject))
-			{
-				return Particle->CastToRigidParticle();
-			}
+			return Interface.GetRigidParticle(PhysicsObject);
+		}
+	}
+
+	return nullptr;
+}
+
+Chaos::FPBDRigidParticleHandle* UPhysicsMovementUtils::GetRigidParticleHandleFromComponent(UPrimitiveComponent* PrimComp)
+{
+	if (IPhysicsComponent* PhysComp = Cast<IPhysicsComponent>(PrimComp))
+	{
+		if (Chaos::FPhysicsObject* PhysicsObject = PhysComp->GetPhysicsObjectByName(NAME_None))
+		{
+			Chaos::FReadPhysicsObjectInterface_Internal Interface = Chaos::FPhysicsObjectInternalInterface::GetRead();
+			return Interface.GetRigidParticle(PhysicsObject);
 		}
 	}
 
diff --git a/Engine/Plugins/Experimental/Mover/Source/Mover/Private/PhysicsMover/PhysicsMoverSimulationTypes.cpp b/Engine/Plugins/Experimental/Mover/Source/Mover/Private/PhysicsMover/PhysicsMoverSimulationTypes.cpp
--- a/Engine/Plugins/Experimental/Mover/Source/Mover/Private/PhysicsMover/PhysicsMoverSimulationTypes.cpp
+++ b/Engine/Plugins/Experimental/Mover/Source/Mover/Private/PhysicsMover/PhysicsMoverSimulationTypes.cpp
@@ -138,3 +138,31 @@ void FMoverAIInputs::Interpolate(const FMoverDataStructBase& From, const FMoverD
 void FMoverAIInputs::Merge(const FMoverDataStructBase& From)
 {
 }
+
+//////////////////////////////////////////////////////////////////////////
+// FMoverLaunchInputs
+
+FMoverDataStructBase* FMoverLaunchInputs::Clone() const
+{
+	FMoverLaunchInputs* CopyPtr = new FMoverLaunchInputs(*this);
+	return CopyPtr;
+}
+
+bool FMoverLaunchInputs::NetSerialize(FArchive& Ar, UPackageMap* Map, bool& bOutSuccess)
+{
+	Super::NetSerialize(Ar, Map, bOutSuccess);
+
+	Ar << LaunchVelocity;
+	Ar << Mode;
+
+	bOutSuccess = true;
+	return true;
+}
+
+void FMoverLaunchInputs::ToString(FAnsiStringBuilderBase& Out) const
+{
+	Super::ToString(Out);
+
+	Out.Appendf("LaunchVelocity: X=%.2f Y=%.2f Z=%.2f\n", LaunchVelocity.X, LaunchVelocity.Y, LaunchVelocity.Z);
+	Out.Appendf("Mode: %u\n", Mode);
+}
diff --git a/Engine/Plugins/Experimental/Mover/Source/Mover/Private/PhysicsMover/Transitions/PhysicsJumpCheck.cpp b/Engine/Plugins/Experimental/Mover/Source/Mover/Private/PhysicsMover/Transitions/PhysicsJumpCheck.cpp
--- a/Engine/Plugins/Experimental/Mover/Source/Mover/Private/PhysicsMover/Transitions/PhysicsJumpCheck.cpp
+++ b/Engine/Plugins/Experimental/Mover/Source/Mover/Private/PhysicsMover/Transitions/PhysicsJumpCheck.cpp
@@ -7,6 +7,8 @@
 #include "DefaultMovementSet/InstantMovementEffects/BasicInstantMovementEffects.h"
 #include "MoverComponent.h"
 #include "MoveLibrary/FloorQueryUtils.h"
+#include "PhysicsMover/InstantMovementEffects/ApplyVelocityPhysicsMovementEffect.h"
+#include "PhysicsMover/PhysicsMovementUtils.h"
 
 #if WITH_EDITOR
 #include "Misc/DataValidation.h"
@@ -14,18 +16,6 @@
 
 namespace Private
 {
-	static Chaos::FPBDRigidParticleHandle* GetParticleHandleFromComponent(UPrimitiveComponent* PrimComp)
-	{
-		if (IPhysicsComponent* PhysComp = Cast<IPhysicsComponent>(PrimComp))
-		{
-			Chaos::FPhysicsObject* PhysicsObject = PhysComp->GetPhysicsObjectByName(NAME_None);
-			Chaos::FReadPhysicsObjectInterface_Internal Interface = Chaos::FPhysicsObjectInternalInterface::GetRead();
-			return Interface.GetRigidParticle(PhysicsObject);
-		}
-
-		return nullptr;
-	}
-
 	static void ApplyImpulse(Chaos::FPBDRigidParticleHandle* Particle, const Chaos::FVec3& Impulse, const Chaos::FVec3& Location)
 	{
 		const Chaos::FVec3 Offset = Location - Particle->GetX();
@@ -37,6 +27,7 @@ namespace Private
 UPhysicsJumpCheck::UPhysicsJumpCheck(const FObjectInitializer& ObjectInitializer)
 	: Super(ObjectInitializer)
 {
+	bFirstSubStepOnly = true;
 	TransitionToMode = DefaultModeNames::Falling;
 }
 
@@ -58,21 +49,22 @@ FTransitionEvalResult UPhysicsJumpCheck::Evaluate_Implementation(const FSimulati
 
 void UPhysicsJumpCheck::Trigger_Implementation(const FSimulationTickParams& Params)
 {
-	TSharedPtr<FJumpImpulseEffect> JumpMove = MakeShared<FJumpImpulseEffect>();
+	TSharedPtr<FApplyVelocityPhysicsEffect> JumpMove = MakeShared<FApplyVelocityPhysicsEffect>();
 
 	if (UMoverComponent* MoverComp = Params.MovingComps.MoverComponent.Get())
 	{
-		JumpMove->UpwardsSpeed = JumpUpwardsSpeed;
+		JumpMove->VelocityToApply = FVector(JumpUpwardsSpeed * MoverComp->GetUpDirection());
+		JumpMove->bAdditiveVelocity = true;
 		Params.MovingComps.MoverComponent->QueueInstantMovementEffect(JumpMove);
 
 		if (FractionalGroundReactionImpulse > 0.0f)
 		{
-			Chaos::FPBDRigidParticleHandle* CharacterParticle = Private::GetParticleHandleFromComponent(Params.MovingComps.UpdatedPrimitive.Get());
+			Chaos::FPBDRigidParticleHandle* CharacterParticle = UPhysicsMovementUtils::GetRigidParticleHandleFromComponent(Params.MovingComps.UpdatedPrimitive.Get());
 			Chaos::FPBDRigidParticleHandle* GroundParticle = nullptr;
 			FHitResult HitResult;
 			if (MoverComp->TryGetFloorCheckHitResult(HitResult))
 			{
-				GroundParticle = Private::GetParticleHandleFromComponent(HitResult.GetComponent());
+				GroundParticle = UPhysicsMovementUtils::GetRigidParticleHandleFromHitResult(HitResult);
 			}
 
 			if (CharacterParticle && GroundParticle)
diff --git a/Engine/Plugins/Experimental/Mover/Source/Mover/Private/PhysicsMover/Transitions/PhysicsLaunchCheck.cpp b/Engine/Plugins/Experimental/Mover/Source/Mover/Private/PhysicsMover/Transitions/PhysicsLaunchCheck.cpp
new file mode 100644
--- /dev/null
+++ b/Engine/Plugins/Experimental/Mover/Source/Mover/Private/PhysicsMover/Transitions/PhysicsLaunchCheck.cpp
@@ -0,0 +1,62 @@
+// Copyright Epic Games, Inc. All Rights Reserved.
+
+#include "PhysicsMover/Transitions/PhysicsLaunchCheck.h"
+
+#include "Chaos/PhysicsObject.h"
+#include "Chaos/PhysicsObjectInternalInterface.h"
+#include "MoverComponent.h"
+#include "MoveLibrary/FloorQueryUtils.h"
+#include "PhysicsMover/InstantMovementEffects/ApplyVelocityPhysicsMovementEffect.h"
+#include "PhysicsMover/PhysicsMoverSimulationTypes.h"
+
+#if WITH_EDITOR
+#include "Misc/DataValidation.h"
+#endif
+
+UPhysicsLaunchCheck::UPhysicsLaunchCheck(const FObjectInitializer& ObjectInitializer)
+	: Super(ObjectInitializer)
+{
+	bFirstSubStepOnly = true;
+	TransitionToMode = NAME_None;
+}
+
+FTransitionEvalResult UPhysicsLaunchCheck::Evaluate_Implementation(const FSimulationTickParams& Params) const
+{
+	FTransitionEvalResult EvalResult;
+
+	if (const FMoverLaunchInputs* LaunchInputs = Params.StartState.InputCmd.InputCollection.FindDataByType<FMoverLaunchInputs>())
+	{
+		if (TransitionToMode.IsNone())
+		{
+			EvalResult.NextMode = Params.StartState.SyncState.MovementMode;
+		}
+		else
+		{
+			EvalResult.NextMode = TransitionToMode;
+		}
+	}
+
+	return EvalResult;
+}
+
+void UPhysicsLaunchCheck::Trigger_Implementation(const FSimulationTickParams& Params)
+{
+	if (UMoverComponent* MoverComp = Params.MovingComps.MoverComponent.Get())
+	{
+		if (const FMoverLaunchInputs* LaunchInputs = Params.StartState.InputCmd.InputCollection.FindDataByType<FMoverLaunchInputs>())
+		{
+			TSharedPtr<FApplyVelocityPhysicsEffect> ApplyVelocityEffect = MakeShared<FApplyVelocityPhysicsEffect>();
+			ApplyVelocityEffect->bAdditiveVelocity = LaunchInputs->Mode == EMoverLaunchVelocityMode::Additive;
+			ApplyVelocityEffect->VelocityToApply = LaunchInputs->LaunchVelocity;
+			ApplyVelocityEffect->ForceMovementMode = TransitionToMode;
+			MoverComp->QueueInstantMovementEffect(ApplyVelocityEffect);
+		}
+	}
+}
+
+#if WITH_EDITOR
+EDataValidationResult UPhysicsLaunchCheck::IsDataValid(FDataValidationContext& Context) const
+{
+	return Super::IsDataValid(Context);
+}
+#endif // WITH_EDITOR
diff --git a/Engine/Plugins/Experimental/Mover/Source/Mover/Public/PhysicsMover/InstantMovementEffects/ApplyVelocityPhysicsMovementEffect.h b/Engine/Plugins/Experimental/Mover/Source/Mover/Public/PhysicsMover/InstantMovementEffects/ApplyVelocityPhysicsMovementEffect.h
new file mode 100644
--- /dev/null
+++ b/Engine/Plugins/Experimental/Mover/Source/Mover/Public/PhysicsMover/InstantMovementEffects/ApplyVelocityPhysicsMovementEffect.h
@@ -0,0 +1,42 @@
+// Copyright Epic Games, Inc. All Rights Reserved.
+
+#pragma once
+
+#include "InstantMovementEffect.h"
+#include "ApplyVelocityPhysicsMovementEffect.generated.h"
+
+/** Apply Velocity: provides a velocity to the actor after (optionally) forcing them into a particular movement mode
+  * Note: this only applies the velocity for one tick!
+  */
+USTRUCT(BlueprintType, Category = "Mover|Instant Movement Effect|Physics", DisplayName = "Apply Velocity Physics Instant Movement Effect")
+struct MOVER_API FApplyVelocityPhysicsEffect : public FInstantMovementEffect
+{
+	GENERATED_BODY()
+
+	FApplyVelocityPhysicsEffect();
+	virtual ~FApplyVelocityPhysicsEffect() {}
+
+	// Velocity to apply to the actor.
+	UPROPERTY(EditAnywhere, BlueprintReadWrite, Category = Mover, meta = (ForceUnits = "cm/s"))
+	FVector VelocityToApply;
+
+	// If true VelocityToApply will be added to current velocity on this actor. If false velocity will be set directly to VelocityToApply
+	UPROPERTY(EditAnywhere, BlueprintReadWrite, Category = Mover)
+	bool bAdditiveVelocity;
+
+	// Optional movement mode name to force the actor into before applying the impulse velocity.
+	UPROPERTY(EditAnywhere, BlueprintReadWrite, Category = Mover)
+	FName ForceMovementMode;
+
+	virtual bool ApplyMovementEffect(FApplyMovementEffectParams& ApplyEffectParams, FMoverSyncState& OutputState) override;
+
+	virtual FInstantMovementEffect* Clone() const override;
+
+	virtual void NetSerialize(FArchive& Ar) override;
+
+	virtual UScriptStruct* GetScriptStruct() const override;
+
+	virtual FString ToSimpleString() const override;
+
+	virtual void AddReferencedObjects(class FReferenceCollector& Collector) override;
+};
\ No newline at end of file
diff --git a/Engine/Plugins/Experimental/Mover/Source/Mover/Public/PhysicsMover/PhysicsMovementUtils.h b/Engine/Plugins/Experimental/Mover/Source/Mover/Public/PhysicsMover/PhysicsMovementUtils.h
--- a/Engine/Plugins/Experimental/Mover/Source/Mover/Public/PhysicsMover/PhysicsMovementUtils.h
+++ b/Engine/Plugins/Experimental/Mover/Source/Mover/Public/PhysicsMover/PhysicsMovementUtils.h
@@ -21,7 +21,9 @@ public:
 		float QueryRadius, float QueryDistance, float MaxWalkSlopeCosine, float TargetHeight, FFloorCheckResult& OutFloorResult, FWaterCheckResult& OutWaterResult);
 
 	// If the hit result hit something, return the particle handle
-	static const Chaos::FPBDRigidParticleHandle* GetRigidParticleHandleFromHitResult(const FHitResult& HitResult);
+	static Chaos::FPBDRigidParticleHandle* GetRigidParticleHandleFromHitResult(const FHitResult& HitResult);
+
+	static Chaos::FPBDRigidParticleHandle* GetRigidParticleHandleFromComponent(UPrimitiveComponent* PrimComp);
 
 	// Checks if any hit is with water and, if so, fills in the OutWaterResult
 	static bool GetWaterResultFromHitResults(const TArray<FHitResult>& Hits, const FVector& Location, const float TargetHeight, FWaterCheckResult& OutWaterResult);
diff --git a/Engine/Plugins/Experimental/Mover/Source/Mover/Public/PhysicsMover/PhysicsMoverSimulationTypes.h b/Engine/Plugins/Experimental/Mover/Source/Mover/Public/PhysicsMover/PhysicsMoverSimulationTypes.h
--- a/Engine/Plugins/Experimental/Mover/Source/Mover/Public/PhysicsMover/PhysicsMoverSimulationTypes.h
+++ b/Engine/Plugins/Experimental/Mover/Source/Mover/Public/PhysicsMover/PhysicsMoverSimulationTypes.h
@@ -192,4 +192,54 @@ struct TStructOpsTypeTraits< FMoverAIInputs > : public TStructOpsTypeTraitsBase2
 	{
 		WithCopy = true
 	};
+};
+
+UENUM(BlueprintType)
+enum class EMoverLaunchVelocityMode : uint8
+{
+	Additive,
+	Override
+};
+
+USTRUCT(BlueprintType)
+struct MOVER_API FMoverLaunchInputs : public FMoverDataStructBase
+{
+	GENERATED_BODY()
+
+public:
+	// Launch velocity in cm/s
+	UPROPERTY(BlueprintReadWrite, Category = Mover)
+	FVector LaunchVelocity;
+
+	UPROPERTY(BlueprintReadWrite, Category = Mover)
+	EMoverLaunchVelocityMode Mode;
+
+	FMoverLaunchInputs()
+		: LaunchVelocity(FVector::ZeroVector)
+		, Mode(EMoverLaunchVelocityMode::Additive)
+	{
+	}
+
+	virtual ~FMoverLaunchInputs() {}
+
+	// @return newly allocated copy of this FMovementSettingsInputs. Must be overridden by child classes
+	virtual FMoverDataStructBase* Clone() const override;
+
+	virtual bool NetSerialize(FArchive& Ar, UPackageMap* Map, bool& bOutSuccess) override;
+
+	virtual UScriptStruct* GetScriptStruct() const override { return StaticStruct(); }
+
+	virtual void ToString(FAnsiStringBuilderBase& Out) const override;
+
+	virtual void AddReferencedObjects(FReferenceCollector& Collector) override { Super::AddReferencedObjects(Collector); }
+};
+
+template<>
+struct TStructOpsTypeTraits< FMoverLaunchInputs > : public TStructOpsTypeTraitsBase2< FMoverLaunchInputs >
+{
+	enum
+	{
+		WithNetSerializer = true,
+		WithCopy = true
+	};
 };
\ No newline at end of file
diff --git a/Engine/Plugins/Experimental/Mover/Source/Mover/Public/PhysicsMover/Transitions/PhysicsLaunchCheck.h b/Engine/Plugins/Experimental/Mover/Source/Mover/Public/PhysicsMover/Transitions/PhysicsLaunchCheck.h
new file mode 100644
--- /dev/null
+++ b/Engine/Plugins/Experimental/Mover/Source/Mover/Public/PhysicsMover/Transitions/PhysicsLaunchCheck.h
@@ -0,0 +1,32 @@
+// Copyright Epic Games, Inc. All Rights Reserved.
+
+#pragma once
+
+#include "MovementModeTransition.h"
+#include "PhysicsLaunchCheck.generated.h"
+
+/**
+ * Transition that handles launching based on input for a physics-based character
+ */
+UCLASS(Blueprintable, BlueprintType)
+class MOVER_API UPhysicsLaunchCheck : public UBaseMovementModeTransition
+{
+	GENERATED_BODY()
+
+public:
+	UPhysicsLaunchCheck(const class FObjectInitializer& ObjectInitializer);
+
+	virtual FTransitionEvalResult Evaluate_Implementation(const FSimulationTickParams& Params) const override;
+	virtual void Trigger_Implementation(const FSimulationTickParams& Params) override;
+
+#if WITH_EDITOR
+	virtual EDataValidationResult IsDataValid(class FDataValidationContext& Context) const override;
+#endif // WITH_EDITOR
+
+	/**
+	Optional: Name of movement mode to transition to when launch is activated.
+	If set to None launching will not change the current mode
+	*/
+	UPROPERTY(BlueprintReadWrite, EditAnywhere, Category = Parameters)
+	FName TransitionToMode;
+};
\ No newline at end of file
