[Mover] Don't use the generated GetOuterUMoverComponent() for UBaseMovementMode::GetMoverComponent() since it does a check. When called on the CDO (ex: for data validation on a BP) there's no outer mover comp, which is fine, and we don't want to crash.

diff --git a/Engine/Plugins/Experimental/Mover/Source/Mover/Public/MovementMode.h b/Engine/Plugins/Experimental/Mover/Source/Mover/Public/MovementMode.h
--- a/Engine/Plugins/Experimental/Mover/Source/Mover/Public/MovementMode.h
+++ b/Engine/Plugins/Experimental/Mover/Source/Mover/Public/MovementMode.h
@@ -57,12 +57,20 @@ public:
 	UFUNCTION(BlueprintCallable, Category=Mover, meta=(DisplayName="Get Mover Component", ScriptName = GetMoverComponent))
 	UMoverComponent* K2_GetMoverComponent() const;
 
+	/**
+	 * Gets the outer mover component of the indicated type. Does not check on the type or the presence of the MoverComp outer. Safe to call on CDOs.
+	 * Note: Since UBaseMovementMode is declared "Within = MoverComponent", all instances of a mode except the CDO are guaranteed to have a valid MoverComponent outer.
+	 */
 	template<typename MoverT = UMoverComponent UE_REQUIRES(std::is_base_of_v<MoverT, UMoverComponent>)>
 	MoverT* GetMoverComponent() const
 	{
-		return Cast<MoverT>(GetOuterUMoverComponent());
+		return Cast<MoverT>(GetOuter());
 	}
 
+	/**
+	 * Gets the outer mover component of the indicated type, checked for validity.
+	 * Note: Since UBaseMovementMode is declared "Within = MoverComponent", all instances of a mode except the CDO are guaranteed to have a valid MoverComponent outer.
+	 */
 	template<typename MoverT = UMoverComponent UE_REQUIRES(std::is_base_of_v<MoverT, UMoverComponent>)>
 	MoverT& GetMoverComponentChecked() const
 	{
