Mover: Adding an external gameplay tag container allowing outside systems to add gameplay tags for Mover evaluations.

diff --git a/Engine/Plugins/Experimental/Mover/Source/Mover/Private/MoverComponent.cpp b/Engine/Plugins/Experimental/Mover/Source/Mover/Private/MoverComponent.cpp
--- a/Engine/Plugins/Experimental/Mover/Source/Mover/Private/MoverComponent.cpp
+++ b/Engine/Plugins/Experimental/Mover/Source/Mover/Private/MoverComponent.cpp
@@ -681,6 +681,21 @@ const FMovementModifierBase* UMoverComponent::FindMovementModifierByType(const U
 
 bool UMoverComponent::HasGameplayTag(FGameplayTag TagToFind, bool bExactMatch) const
 {
+	if (bExactMatch)
+	{
+		if (ExternalGameplayTags.HasTagExact(TagToFind))
+		{
+			return true;
+		}	
+	}
+	else
+	{
+		if (ExternalGameplayTags.HasTag(TagToFind))
+		{
+			return true;
+		}
+	}
+	
 	if (bHasValidCachedState)
 	{
 		// Search Movement Modes
@@ -715,6 +730,26 @@ bool UMoverComponent::HasGameplayTag(FGameplayTag TagToFind, bool bExactMatch) c
 	return false;
 }
 
+void UMoverComponent::AddGameplayTag(FGameplayTag TagToAdd)
+{
+	ExternalGameplayTags.AddTag(TagToAdd);
+}
+
+void UMoverComponent::AddGameplayTags(const FGameplayTagContainer& TagsToAdd)
+{
+	ExternalGameplayTags.AppendTags(TagsToAdd);
+}
+
+void UMoverComponent::RemoveGameplayTag(FGameplayTag TagToRemove)
+{
+	ExternalGameplayTags.RemoveTag(TagToRemove);
+}
+
+void UMoverComponent::RemoveGameplayTags(const FGameplayTagContainer& TagsToRemove)
+{
+	ExternalGameplayTags.RemoveTags(TagsToRemove);
+}
+
 void UMoverComponent::SetFrameStateFromContext(const FMoverSyncState* SyncState, const FMoverAuxStateContext* AuxState, bool bRebaseBasedState)
 {
 
diff --git a/Engine/Plugins/Experimental/Mover/Source/Mover/Public/MoverComponent.h b/Engine/Plugins/Experimental/Mover/Source/Mover/Public/MoverComponent.h
--- a/Engine/Plugins/Experimental/Mover/Source/Mover/Public/MoverComponent.h
+++ b/Engine/Plugins/Experimental/Mover/Source/Mover/Public/MoverComponent.h
@@ -461,6 +461,36 @@ public:	// Queries
  	 */
 	UFUNCTION(BlueprintPure, Category = Mover, meta = (Keywords = "HasTag"))
 	bool HasGameplayTag(FGameplayTag TagToFind, bool bExactMatch) const;
+
+	/**
+  	 * Adds a gameplay tag to this Mover Component.
+  	 * Note: Duplicate tags will not be added
+  	 * @param TagToAdd			Tag to add to the Mover Component
+  	 */
+	UFUNCTION(BlueprintCallable, Category = Mover, meta = (Keywords = "Add Tag"))
+	void AddGameplayTag(FGameplayTag TagToAdd);
+
+	/**
+   	 * Adds a series of gameplay tags to this Mover Component
+   	 * Note: Duplicate tags will not be added
+   	 * @param TagsToAdd			Tags to add/append to the Mover Component
+   	 */
+	UFUNCTION(BlueprintCallable, Category = Mover, meta = (Keywords = "Add Tag"))
+	void AddGameplayTags(UPARAM(Ref) const FGameplayTagContainer& TagsToAdd);
+	
+	/**
+   	 * Removes a gameplay tag from this Mover Component
+   	 * @param TagToRemove			Tag to remove from the Mover Component
+   	 */
+	UFUNCTION(BlueprintCallable, Category = Mover, meta = (Keywords = "Remove Tag"))
+	void RemoveGameplayTag(FGameplayTag TagToRemove);
+
+	/**
+	 * Removes gameplay tags from this Mover Component
+	 * @param TagsToRemove			Tags to remove from the Mover Component
+	 */
+	UFUNCTION(BlueprintCallable, Category = Mover, meta = (Keywords = "Remove Tag"))
+	void RemoveGameplayTags(UPARAM(Ref) const FGameplayTagContainer& TagsToRemove);
 	
 protected:
 
@@ -536,6 +566,10 @@ protected:
 	/** Cached original offset from the visual component, used for cases where we want to move the visual component away from the root component (for smoothing, corrections, etc.) */
 	FTransform BaseVisualComponentTransform = FTransform::Identity;
 
+	// TODO: Look at possibility of replacing this with a FGameplayTagCountContainer that could possibly represent both internal and external tags
+	/** A list of gameplay tags associated with this Mover Component added from sources outside of Mover */
+	FGameplayTagContainer ExternalGameplayTags;
+	
 	FMoverInputCmdContext CachedLastProducedInputCmd;
 	
 	FMoverInputCmdContext CachedLastUsedInputCmd;
