Remove TickComponent from UMoverComponent

diff --git a/Engine/Plugins/Experimental/Mover/Source/Mover/Private/MoverComponent.cpp b/Engine/Plugins/Experimental/Mover/Source/Mover/Private/MoverComponent.cpp
--- a/Engine/Plugins/Experimental/Mover/Source/Mover/Private/MoverComponent.cpp
+++ b/Engine/Plugins/Experimental/Mover/Source/Mover/Private/MoverComponent.cpp
@@ -60,7 +60,7 @@ static constexpr float ROTATOR_TOLERANCE = (1e-3);
 UMoverComponent::UMoverComponent()
 {
 	PrimaryComponentTick.TickGroup = TG_PrePhysics;
-	PrimaryComponentTick.bCanEverTick = true;
+	PrimaryComponentTick.bCanEverTick = false;
 
 	BasedMovementTickFunction.bCanEverTick = true;
 	BasedMovementTickFunction.bStartWithTickEnabled = false;
@@ -238,11 +238,6 @@ void UMoverComponent::BeginPlay()
 	}
 }
 
-void UMoverComponent::TickComponent(float DeltaTime, enum ELevelTick TickType, FActorComponentTickFunction* ThisTickFunction)
-{
-	Super::TickComponent(DeltaTime, TickType, ThisTickFunction);
-}
-
 void UMoverComponent::BindProcessGeneratedMovement(FMover_ProcessGeneratedMovement ProcessGeneratedMovementEvent)
 {
 	ProcessGeneratedMovement = ProcessGeneratedMovementEvent;
diff --git a/Engine/Plugins/Experimental/Mover/Source/Mover/Public/MoverComponent.h b/Engine/Plugins/Experimental/Mover/Source/Mover/Public/MoverComponent.h
--- a/Engine/Plugins/Experimental/Mover/Source/Mover/Public/MoverComponent.h
+++ b/Engine/Plugins/Experimental/Mover/Source/Mover/Public/MoverComponent.h
@@ -74,8 +74,6 @@ public:
 	virtual void PostLoad() override;
 	virtual void BeginPlay() override;
 
-	virtual void TickComponent(float DeltaTime, enum ELevelTick TickType, FActorComponentTickFunction* ThisTickFunction) override;
-
 	// Broadcast before each simulation tick.
 	// Note - Guaranteed to run on the game thread (even in async simulation).
 	UPROPERTY(BlueprintAssignable, Category = Mover)
