[Mover]

- Improve padding on FProposedMove so that the padding is at the bottom
- Remove virtual FMovementModifierGroup to save a VTable pointer

#rnx

diff --git a/Engine/Plugins/Experimental/Mover/Source/Mover/Public/MoveLibrary/MovementUtilsTypes.h b/Engine/Plugins/Experimental/Mover/Source/Mover/Public/MoveLibrary/MovementUtilsTypes.h
--- a/Engine/Plugins/Experimental/Mover/Source/Mover/Public/MoveLibrary/MovementUtilsTypes.h
+++ b/Engine/Plugins/Experimental/Mover/Source/Mover/Public/MoveLibrary/MovementUtilsTypes.h
@@ -29,34 +29,34 @@ struct MOVER_API FProposedMove
 
 	FProposedMove() : 
 		bHasDirIntent(false)
-	{}
-
-	// Determines how this move should resolve with other moves
-	UPROPERTY(EditAnywhere, BlueprintReadWrite, Category = Mover)
-	EMoveMixMode MixMode = EMoveMixMode::AdditiveVelocity;
+	{}	
 
 	/**
 	 * Indicates that we should switch to a particular movement mode before the next simulation step is performed.
 	 * Note: If this is set from a layered move the preferred mode will only be set at the beginning of the layered move, not continuously.
 	 */
 	UPROPERTY(EditAnywhere, BlueprintReadWrite, Category = Mover)
-	FName	PreferredMode = NAME_None;
-
-	// Signals whether there was any directional intent specified
-	UPROPERTY(EditAnywhere, BlueprintReadWrite, Category = Mover)
-	uint8	 bHasDirIntent : 1;
+	FName PreferredMode = NAME_None;	
 
 	// Directional, per-axis magnitude [-1, 1] in world space (length of 1 indicates max speed intent). Only valid if bHasDirIntent is set.
 	UPROPERTY(EditAnywhere, BlueprintReadWrite, Category = Mover)
-	FVector  DirectionIntent = FVector::ZeroVector;
+	FVector DirectionIntent = FVector::ZeroVector;
 
 	// Units per second, world space, possibly mapped onto walking surface
 	UPROPERTY(EditAnywhere, BlueprintReadWrite, Category = Mover)
-	FVector  LinearVelocity = FVector::ZeroVector;
+	FVector LinearVelocity = FVector::ZeroVector;
 	
 	// Degrees per second, local space
 	UPROPERTY(EditAnywhere, BlueprintReadWrite, Category = Mover)
 	FRotator AngularVelocity = FRotator::ZeroRotator;
+	
+	// Signals whether there was any directional intent specified
+	UPROPERTY(EditAnywhere, BlueprintReadWrite, Category = Mover)
+	uint8 bHasDirIntent : 1;
+	
+	// Determines how this move should resolve with other moves
+	UPROPERTY(EditAnywhere, BlueprintReadWrite, Category = Mover)
+	EMoveMixMode MixMode = EMoveMixMode::AdditiveVelocity;
 };
 
 
diff --git a/Engine/Plugins/Experimental/Mover/Source/Mover/Public/MovementModifier.h b/Engine/Plugins/Experimental/Mover/Source/Mover/Public/MovementModifier.h
--- a/Engine/Plugins/Experimental/Mover/Source/Mover/Public/MovementModifier.h
+++ b/Engine/Plugins/Experimental/Mover/Source/Mover/Public/MovementModifier.h
@@ -182,7 +182,6 @@ struct MOVER_API FMovementModifierGroup
 	GENERATED_BODY()
 
 	FMovementModifierGroup() {}
-	virtual ~FMovementModifierGroup() {}
 	
 	bool HasAnyMoves() const { return (!ActiveModifiers.IsEmpty() || !QueuedModifiers.IsEmpty()); }
 
