Mover: fixed cases where static const FNames were more appropriate, saving a bit of performance

#rnx

diff --git a/Engine/Plugins/Experimental/Mover/Source/Mover/Private/MoveLibrary/AsyncMovementUtils.cpp b/Engine/Plugins/Experimental/Mover/Source/Mover/Private/MoveLibrary/AsyncMovementUtils.cpp
--- a/Engine/Plugins/Experimental/Mover/Source/Mover/Private/MoveLibrary/AsyncMovementUtils.cpp
+++ b/Engine/Plugins/Experimental/Mover/Source/Mover/Private/MoveLibrary/AsyncMovementUtils.cpp
@@ -19,8 +19,8 @@ static float PenetrationOverlapCheckInflationCVar = 0.1f;
 
 bool UAsyncMovementUtils::TestDepenetratingMove(const FMovingComponentSet& MovingComps, const FVector& StartLocation, const FVector& TargetLocation, const FQuat& StartRotation, const FQuat& TargetRotation, bool bShouldSweep, FHitResult& OutHit, FMovementRecord& InOutMoveRecord)
 {
-	const FName TestSafeMoveSubstepName = "TestSafeMove";
-	const FName TestSafeDepenetrationSubstepName = "TestSafeResolvePenetration";
+	static const FName TestSafeMoveSubstepName = "TestSafeMove";
+	static const FName TestSafeDepenetrationSubstepName = "TestSafeResolvePenetration";
 
 	const USceneComponent* UpdatedComponent = MovingComps.UpdatedComponent.Get();
 
