Mover: avoid unnecessary slow method of checking if we have a valid ModeFSM

#rnx

diff --git a/Engine/Plugins/Experimental/Mover/Source/Mover/Private/MoverComponent.cpp b/Engine/Plugins/Experimental/Mover/Source/Mover/Private/MoverComponent.cpp
--- a/Engine/Plugins/Experimental/Mover/Source/Mover/Private/MoverComponent.cpp
+++ b/Engine/Plugins/Experimental/Mover/Source/Mover/Private/MoverComponent.cpp
@@ -433,7 +433,7 @@ void UMoverComponent::SimulationTick(const FMoverTimeStep& InTimeStep, const FMo
 
 	OnPreSimulationTick.Broadcast(MoverTimeStep, SimInput.InputCmd);
 
-	if (!ModeFSM->IsValidLowLevel())
+	if (!ModeFSM)
 	{
 		SimOutput.SyncState = SimInput.SyncState;
 		SimOutput.AuxState = SimInput.AuxState;
