Replaced ground normal with surface normal to avoid infinite positive velocity when colliding with slopes.

#rnx

diff --git a/Engine/Plugins/Experimental/Mover/Source/Mover/Private/PhysicsMover/Modes/PhysicsDrivenFallingMode.cpp b/Engine/Plugins/Experimental/Mover/Source/Mover/Private/PhysicsMover/Modes/PhysicsDrivenFallingMode.cpp
--- a/Engine/Plugins/Experimental/Mover/Source/Mover/Private/PhysicsMover/Modes/PhysicsDrivenFallingMode.cpp
+++ b/Engine/Plugins/Experimental/Mover/Source/Mover/Private/PhysicsMover/Modes/PhysicsDrivenFallingMode.cpp
@@ -152,7 +152,7 @@ void UPhysicsDrivenFallingMode::SimulationTick_Implementation(const FSimulationT
 
 	constexpr float FloorDistanceTolerance = 2.0f;
 	const FVector ProjectedGroundVelocity = UPhysicsMovementUtils::ComputeIntegratedGroundVelocityFromHitResult(StartingSyncState->GetLocation_WorldSpace(), FloorResult.HitResult, DeltaSeconds);
-	const float ProjectedRelativeVerticalVelocity = UpDir.Dot(ProposedMove.LinearVelocity - ProjectedGroundVelocity);
+	const float ProjectedRelativeVerticalVelocity = FloorResult.HitResult.ImpactNormal.Dot(ProposedMove.LinearVelocity - ProjectedGroundVelocity);
 	const float ProjectedFloorDistance = FloorResult.FloorDist + ProjectedRelativeVerticalVelocity * DeltaSeconds;
 	const bool bIsFloorWithinReach = ProjectedFloorDistance < TargetHeight + FloorDistanceTolerance;
 	const bool bIsMovingUpRelativeToFloor = ProjectedRelativeVerticalVelocity > UE_KINDA_SMALL_NUMBER;
