Physics Mover - Filter ignored particles in contact modification function

diff --git a/Engine/Plugins/Experimental/Mover/Source/Mover/Private/Backends/MoverNetworkPhysicsLiaison.cpp b/Engine/Plugins/Experimental/Mover/Source/Mover/Private/Backends/MoverNetworkPhysicsLiaison.cpp
--- a/Engine/Plugins/Experimental/Mover/Source/Mover/Private/Backends/MoverNetworkPhysicsLiaison.cpp
+++ b/Engine/Plugins/Experimental/Mover/Source/Mover/Private/Backends/MoverNetworkPhysicsLiaison.cpp
@@ -17,10 +17,6 @@
 #include "DefaultMovementSet/Modes/FallingMode.h"
 #include "Misc/DataValidation.h"
 
-
-#if WITH_EDITOR
-#endif
-
 #define LOCTEXT_NAMESPACE "Mover"
 
 extern FPhysicsDrivenMotionDebugParams GPhysicsDrivenMotionDebugParams;
@@ -203,18 +199,79 @@ void UMoverNetworkPhysicsLiaisonComponent::OnContactModification_Internal(const
 		return;
 	}
 
+	Chaos::FCharacterGroundConstraintHandle* ConstraintHandle = nullptr;
+	if (Constraint && Constraint->IsValid())
+	{
+		ConstraintHandle = Constraint->GetProxy<Chaos::FCharacterGroundConstraintProxy>()->GetPhysicsThreadAPI();
+	}
+
+	if (!ConstraintHandle || !ConstraintHandle->IsEnabled() || !ConstraintHandle->GetCharacterParticle())
+	{
+		return;
+	}
+
 	const UMoverComponent& MoverComp = GetMoverComponent();
-	if (MoverComp.MovementModes.Contains(Input.SyncState.MovementMode))
+
+	// Global
+
+	TArray<const Chaos::FGeometryParticleHandle*> IgnoreParticles;
+	if (MoverComp.UpdatedCompAsPrimitive)
 	{
-		if (const IPhysicsCharacterMovementModeInterface* PhysicsMode = Cast<const IPhysicsCharacterMovementModeInterface>(MoverComp.MovementModes[Input.SyncState.MovementMode]))
+		for (const UPrimitiveComponent* PrimComp : MoverComp.UpdatedCompAsPrimitive->GetMoveIgnoreComponents())
 		{
-			Chaos::FCharacterGroundConstraintHandle* ConstraintHandle = Constraint->GetProxy<Chaos::FCharacterGroundConstraintProxy>()->GetPhysicsThreadAPI();
-			if (!ConstraintHandle || !ConstraintHandle->IsEnabled() || !ConstraintHandle->GetCharacterParticle())
+			TArray<Chaos::FPhysicsObject*> PhysObjs = PrimComp->GetAllPhysicsObjects();
+			Chaos::FReadPhysicsObjectInterface_Internal Interface = Chaos::FPhysicsObjectInternalInterface::GetRead();
+			for (Chaos::FPhysicsObject* PhysObj : PhysObjs)
 			{
-				return;
+				IgnoreParticles.Add(Interface.GetParticle(PhysObj));
 			}
+		}
 
-			const FPhysicsMoverSimulationContactModifierParams Params{ConstraintHandle, MoverComp.UpdatedCompAsPrimitive};
+		for (const AActor* Actor : MoverComp.UpdatedCompAsPrimitive->GetMoveIgnoreActors())
+		{
+			for (UActorComponent* ActorComp : Actor->GetComponents())
+			{
+				if (IPhysicsComponent* PhysComp = Cast<IPhysicsComponent>(ActorComp))
+				{
+					TArray<Chaos::FPhysicsObject*> PhysObjs = PhysComp->GetAllPhysicsObjects();
+					Chaos::FReadPhysicsObjectInterface_Internal Interface = Chaos::FPhysicsObjectInternalInterface::GetRead();
+					for (Chaos::FPhysicsObject* PhysObj : PhysObjs)
+					{
+						if (const Chaos::FGeometryParticleHandle* ParticleHandle = Interface.GetParticle(PhysObj))
+						{
+							IgnoreParticles.Add(ParticleHandle);
+						}
+					}
+				}
+			}
+		}
+	}
+
+	if (!IgnoreParticles.IsEmpty())
+	{
+		Chaos::FPBDRigidParticleHandle* CharacterParticle = ConstraintHandle->GetCharacterParticle()->CastToRigidParticle();
+		Chaos::FContactPairModifierParticleRange Contacts = Modifier.GetContacts(CharacterParticle);
+		for (auto& Contact : Contacts)
+		{
+			const Chaos::FGeometryParticleHandle* OtherParticle = Contact.GetOtherParticle(CharacterParticle);
+			for (const Chaos::FGeometryParticleHandle* Particle : IgnoreParticles)
+			{
+				if (Particle == OtherParticle)
+				{
+					Contact.Disable();
+					break;
+				}
+			}
+		}
+	}
+
+	// Per Mode
+
+	if (MoverComp.MovementModes.Contains(Input.SyncState.MovementMode))
+	{
+		if (const IPhysicsCharacterMovementModeInterface* PhysicsMode = Cast<const IPhysicsCharacterMovementModeInterface>(MoverComp.MovementModes[Input.SyncState.MovementMode]))
+		{
+			const FPhysicsMoverSimulationContactModifierParams Params{ ConstraintHandle, MoverComp.UpdatedCompAsPrimitive };
 			PhysicsMode->OnContactModification_Internal(Params, Modifier);
 		}
 	}
