Newton - Fixes for launching

#changelist validated
#virtualized

diff --git a/Engine/Plugins/Experimental/Mover/Source/Mover/Public/PhysicsMover/PhysicsMoverSimulationTypes.h b/Engine/Plugins/Experimental/Mover/Source/Mover/Public/PhysicsMover/PhysicsMoverSimulationTypes.h
--- a/Engine/Plugins/Experimental/Mover/Source/Mover/Public/PhysicsMover/PhysicsMoverSimulationTypes.h
+++ b/Engine/Plugins/Experimental/Mover/Source/Mover/Public/PhysicsMover/PhysicsMoverSimulationTypes.h
@@ -25,7 +25,7 @@ class UPrimitiveComponent;
 
 struct FPhysicsDrivenMotionDebugParams
 {
-	float TeleportThreshold = 100.0f;
+	float TeleportThreshold = 1000.0f;
 	float MinStepUpDistance = 5.0f;
 	float MaxCharacterGroundMassRatio = 1.5f;
 	bool EnableMultithreading = true;
