Mover: pulling static const FNames out of function scope to avoid init guard cost

#rnx

diff --git a/Engine/Plugins/Experimental/Mover/Source/Mover/Private/MoveLibrary/AsyncMovementUtils.cpp b/Engine/Plugins/Experimental/Mover/Source/Mover/Private/MoveLibrary/AsyncMovementUtils.cpp
--- a/Engine/Plugins/Experimental/Mover/Source/Mover/Private/MoveLibrary/AsyncMovementUtils.cpp
+++ b/Engine/Plugins/Experimental/Mover/Source/Mover/Private/MoveLibrary/AsyncMovementUtils.cpp
@@ -16,12 +16,11 @@ static float HitDistanceToleranceCVar = 0.f;
 static float InitialOverlapToleranceCVar = 0.f;
 static float PenetrationOverlapCheckInflationCVar = 0.1f;
 
+static const FName TestSafeMoveSubstepName = "TestSafeMove";
+static const FName TestSafeDepenetrationSubstepName = "TestSafeResolvePenetration";
 
 bool UAsyncMovementUtils::TestDepenetratingMove(const FMovingComponentSet& MovingComps, const FVector& StartLocation, const FVector& TargetLocation, const FQuat& StartRotation, const FQuat& TargetRotation, bool bShouldSweep, FHitResult& OutHit, FMovementRecord& InOutMoveRecord)
 {
-	static const FName TestSafeMoveSubstepName = "TestSafeMove";
-	static const FName TestSafeDepenetrationSubstepName = "TestSafeResolvePenetration";
-
 	const USceneComponent* UpdatedComponent = MovingComps.UpdatedComponent.Get();
 
 	if (!UpdatedComponent)
@@ -283,6 +282,7 @@ static void PullBackHit(FHitResult& Hit, const FVector& Start, const FVector& En
 }
 
 
+static const FName TestSweepTraceTagName = "SweepTestMoverComponent";
 
 // Returns true if any of the sweeping movement could occur
 bool UAsyncMovementUtils::TestMoveComponent_Internal(const FMovingComponentSet& MovingComps, const FVector& StartLocation, const FVector& TargetLocation, const FQuat& StartRotation, const FQuat& TargetRotation, bool bShouldSweep, EMoveComponentFlags MoveCompFlags, FHitResult& OutHit)
@@ -338,13 +338,11 @@ bool UAsyncMovementUtils::TestMoveComponent_Internal(const FMovingComponentSet&
 	UWorld* const MyWorld = UpdatedPrimitive->GetWorld();
 	if (MyWorld && bCollisionEnabled && (DeltaSizeSq > 0.f))
 	{
-		static const FName TraceTagName = TEXT("SweepTestMoverComponent");
-
 		FComponentQueryParams Params(SCENE_QUERY_STAT(SweepTestMoverComponent), OwningActor);
 		FCollisionResponseParams ResponseParam;
 		UpdatedPrimitive->InitSweepCollisionParams(Params, ResponseParam);
 		Params.bIgnoreTouches = true;
-		Params.TraceTag = TraceTagName;
+		Params.TraceTag = TestSweepTraceTagName;
 
 		const bool bHadBlockingHit = MyWorld->ComponentSweepMulti(AllHits, UpdatedPrimitive, TraceStart, TraceEnd, InitialRotationQuat, Params);
 
diff --git a/Engine/Plugins/Experimental/Mover/Source/MoverEditor/Private/K2Node_PlayMontageOnMoverActor.cpp b/Engine/Plugins/Experimental/Mover/Source/MoverEditor/Private/K2Node_PlayMontageOnMoverActor.cpp
--- a/Engine/Plugins/Experimental/Mover/Source/MoverEditor/Private/K2Node_PlayMontageOnMoverActor.cpp
+++ b/Engine/Plugins/Experimental/Mover/Source/MoverEditor/Private/K2Node_PlayMontageOnMoverActor.cpp
@@ -31,13 +31,13 @@ FText UK2Node_PlayMontageOnMoverActor::GetMenuCategory() const
 	return LOCTEXT("PlayMontageCategory", "Animation|Montage");
 }
 
+static const FName NAME_OnNotifyBegin = "OnNotifyBegin";
+static const FName NAME_OnNotifyEnd = "OnNotifyEnd";
+
 void UK2Node_PlayMontageOnMoverActor::GetPinHoverText(const UEdGraphPin& Pin, FString& HoverTextOut) const
 {
 	Super::GetPinHoverText(Pin, HoverTextOut);
 
-	static const FName NAME_OnNotifyBegin = FName(TEXT("OnNotifyBegin"));
-	static const FName NAME_OnNotifyEnd = FName(TEXT("OnNotifyEnd"));
-
 	if (Pin.PinName == NAME_OnNotifyBegin)
 	{
 		FText ToolTipText = LOCTEXT("K2Node_PlayMontageOnMoverActor_OnNotifyBegin_Tooltip", "Event called when using a PlayMontageNotify or PlayMontageNotifyWindow Notify in a Montage.");
