[Mover] [BugFix] Check if ground is dynamic before applying ground reaction impulse

diff --git a/Engine/Plugins/Experimental/Mover/Source/Mover/Private/PhysicsMover/Transitions/PhysicsJumpCheck.cpp b/Engine/Plugins/Experimental/Mover/Source/Mover/Private/PhysicsMover/Transitions/PhysicsJumpCheck.cpp
--- a/Engine/Plugins/Experimental/Mover/Source/Mover/Private/PhysicsMover/Transitions/PhysicsJumpCheck.cpp
+++ b/Engine/Plugins/Experimental/Mover/Source/Mover/Private/PhysicsMover/Transitions/PhysicsJumpCheck.cpp
@@ -67,7 +67,7 @@ void UPhysicsJumpCheck::Trigger_Implementation(const FSimulationTickParams& Para
 				GroundParticle = UPhysicsMovementUtils::GetRigidParticleHandleFromHitResult(HitResult);
 			}
 
-			if (CharacterParticle && GroundParticle)
+			if (GroundParticle && GroundParticle->IsDynamic() && CharacterParticle)
 			{
 				const Chaos::FVec3 Impulse = -FractionalGroundReactionImpulse * CharacterParticle->M() * JumpUpwardsSpeed * MoverComp->GetUpDirection();
 				Private::ApplyImpulse(GroundParticle, Impulse, HitResult.ImpactPoint);
