Mover: providing options for how nav walking mode should handle it when the proposed move would land us off the nav mesh. Previously, we would always revert to normal walking mode. This is still the default, but now we have the option to not move or to move anyways.

diff --git a/Engine/Plugins/Experimental/Mover/Source/Mover/Private/DefaultMovementSet/Modes/AsyncNavWalkingMode.cpp b/Engine/Plugins/Experimental/Mover/Source/Mover/Private/DefaultMovementSet/Modes/AsyncNavWalkingMode.cpp
--- a/Engine/Plugins/Experimental/Mover/Source/Mover/Private/DefaultMovementSet/Modes/AsyncNavWalkingMode.cpp
+++ b/Engine/Plugins/Experimental/Mover/Source/Mover/Private/DefaultMovementSet/Modes/AsyncNavWalkingMode.cpp
@@ -234,10 +234,28 @@ void UAsyncNavWalkingMode::SimulationTick_Implementation(const FSimulationTickPa
 		const bool bHasNavigationData = FindNavFloor(TargetFeetLocation, DestNavLocation, NavData);
 		if (!bHasNavigationData)
 		{
-			// Can't find nav mesh at this location, so fall back to normal walking mode 
-			OutputState.MovementEndState.NextModeName = DefaultModeNames::Walking;
-			OutputState.MovementEndState.RemainingMs = Params.TimeStep.StepMs;
-			MoveRecord.SetDeltaSeconds(0.0f);
+			// Can't find nav mesh at this location, so we need to do something else
+			switch (BehaviorOffNavMesh)
+			{
+				default:	// fall through
+				case EOffNavMeshBehavior::SwitchToWalking:
+					UE_LOG(LogMover, Verbose, TEXT("%s could not find valid navigation data at location %s. Switching to walking mode."), *GetNameSafe(MoverComp->GetOwner()), *TargetFeetLocation.ToCompactString());
+					OutputState.MovementEndState.NextModeName = DefaultModeNames::Walking;
+					OutputState.MovementEndState.RemainingMs = Params.TimeStep.StepMs;
+					MoveRecord.SetDeltaSeconds(0.0f);
+					break;
+
+				case EOffNavMeshBehavior::MoveWithoutNavMesh:
+					// allow the full move to occur 
+					LocationInProgress = StartingSyncState->GetLocation_WorldSpace() + (ProposedMove.LinearVelocity* DeltaSeconds);
+					RotationInProgress = TargetRotation;
+					break;
+
+				case EOffNavMeshBehavior::DoNotMove:
+					UE_LOG(LogMover, Verbose, TEXT("%s could not find valid navigation data at location %s. Cannot move."), *GetNameSafe(MoverComp->GetOwner()), *TargetFeetLocation.ToCompactString());
+					// nothing to be done
+					break;
+			}
 
 			CaptureFinalState(LocationInProgress, RotationInProgress.Rotator(), /*bDidAttemptMovement=*/true, CurrentFloor, MoveRecord, OutputSyncState);
 			return;
diff --git a/Engine/Plugins/Experimental/Mover/Source/Mover/Private/DefaultMovementSet/Modes/NavWalkingMode.cpp b/Engine/Plugins/Experimental/Mover/Source/Mover/Private/DefaultMovementSet/Modes/NavWalkingMode.cpp
--- a/Engine/Plugins/Experimental/Mover/Source/Mover/Private/DefaultMovementSet/Modes/NavWalkingMode.cpp
+++ b/Engine/Plugins/Experimental/Mover/Source/Mover/Private/DefaultMovementSet/Modes/NavWalkingMode.cpp
@@ -229,10 +229,28 @@ void UNavWalkingMode::SimulationTick_Implementation(const FSimulationTickParams&
 		const bool bHasNavigationData = FindNavFloor(AdjustedDest, DestNavLocation, NavData);
 		if (!bHasNavigationData)
 		{
-			// Can't find nav mesh, so fall back to normal walking mode and let it process the intended movement
-			OutputState.MovementEndState.NextModeName = DefaultModeNames::Walking;
-			OutputState.MovementEndState.RemainingMs = Params.TimeStep.StepMs;
-			MoveRecord.SetDeltaSeconds(0.0f);
+			// Can't find nav mesh at this location, so we need to do something else
+			switch (BehaviorOffNavMesh)
+			{
+				default:	// fall through
+				case EOffNavMeshBehavior::SwitchToWalking:
+					UE_LOG(LogMover, Verbose, TEXT("%s could not find valid navigation data at location %s. Switching to walking mode."), *GetNameSafe(MoverComp->GetOwner()), *TargetFeetLocation.ToCompactString());
+					OutputState.MovementEndState.NextModeName = DefaultModeNames::Walking;
+					OutputState.MovementEndState.RemainingMs = Params.TimeStep.StepMs;
+					MoveRecord.SetDeltaSeconds(0.0f);
+					break;
+
+				case EOffNavMeshBehavior::MoveWithoutNavMesh:
+					// allow the full move to occur 
+					// TODO: Need to actually move the updated component
+					ensureMsgf(false, TEXT("NavWalkingMode does not yet support MoveWithoutNavMesh"));
+					break;
+
+				case EOffNavMeshBehavior::DoNotMove:
+					UE_LOG(LogMover, Verbose, TEXT("%s could not find valid navigation data at location %s. Cannot move."), *GetNameSafe(MoverComp->GetOwner()), *TargetFeetLocation.ToCompactString());
+					// nothing to be done
+					break;
+			}
 			
 			CaptureFinalState(UpdatedComponent, true, CurrentFloor, MoveRecord, OutputSyncState);
 			return;
diff --git a/Engine/Plugins/Experimental/Mover/Source/Mover/Public/DefaultMovementSet/Modes/AsyncNavWalkingMode.h b/Engine/Plugins/Experimental/Mover/Source/Mover/Public/DefaultMovementSet/Modes/AsyncNavWalkingMode.h
--- a/Engine/Plugins/Experimental/Mover/Source/Mover/Public/DefaultMovementSet/Modes/AsyncNavWalkingMode.h
+++ b/Engine/Plugins/Experimental/Mover/Source/Mover/Public/DefaultMovementSet/Modes/AsyncNavWalkingMode.h
@@ -64,6 +64,11 @@ public:
 	UPROPERTY(Category = "NavMesh Movement", EditAnywhere, BlueprintReadWrite, meta = (editcondition = "bProjectNavMeshWalking", ClampMin = "0", UIMin = "0"))
 	float NavMeshProjectionInterpSpeed;
 
+	/** What should we do if we stray off the nav mesh? */
+	UPROPERTY(Category = "NavMesh Movement", EditAnywhere, BlueprintReadWrite)
+	EOffNavMeshBehavior BehaviorOffNavMesh = EOffNavMeshBehavior::SwitchToWalking;
+
+
 	UPROPERTY(Transient)
 	float NavMeshProjectionTimer;
 
diff --git a/Engine/Plugins/Experimental/Mover/Source/Mover/Public/DefaultMovementSet/Modes/NavWalkingMode.h b/Engine/Plugins/Experimental/Mover/Source/Mover/Public/DefaultMovementSet/Modes/NavWalkingMode.h
--- a/Engine/Plugins/Experimental/Mover/Source/Mover/Public/DefaultMovementSet/Modes/NavWalkingMode.h
+++ b/Engine/Plugins/Experimental/Mover/Source/Mover/Public/DefaultMovementSet/Modes/NavWalkingMode.h
@@ -18,6 +18,23 @@ namespace NavWalkingBlackBoard
 	const FName ProjectedNavMeshHitResult = TEXT("ProjectedNavMeshHitResult");
 }
 
+/** Options for what to do when we find ourselves off the nav mesh */
+UENUM(BlueprintType)
+enum class EOffNavMeshBehavior : uint8
+{
+	/** We change to normal walking mode and re-attempt the move. Something will later need to switch back to nav walking. */
+	SwitchToWalking = 0        UMETA(DisplayName = "Switch to Walking Mode"),
+
+	/** We proceed to move as directed, but height may diverge from the floor until we return to a valid nav mesh. */
+	MoveWithoutNavMesh = 1     UMETA(DisplayName = "Move Without Nav Mesh"),
+
+	/** We do not move. Movement will continue later if nav mesh becomes available. */
+	DoNotMove = 2			   UMETA(DisplayName = "Do Not Move"),
+};
+
+
+
+
 /**
  * NavWalkingMode: a default movement mode for traversing surfaces and movement bases by using an active navmesh when moving the actor rather than collision checks.
  * Note: This movement mode requires a NavMoverComponent be on the actor to function properly. This mode also contains some randomization to avoid navmesh look ups
@@ -67,6 +84,11 @@ public:
 	/** Speed at which to interpolate agent navmesh offset between traces. 0: Instant (no interp) > 0: Interp speed") */
 	UPROPERTY(Category="NavMesh Movement", EditAnywhere, BlueprintReadWrite, meta=(editcondition = "bProjectNavMeshWalking", ClampMin="0", UIMin="0"))
 	float NavMeshProjectionInterpSpeed;
+
+	/** What should we do if we stray off the nav mesh? */
+	UPROPERTY(Category = "NavMesh Movement", EditAnywhere, BlueprintReadWrite)
+	EOffNavMeshBehavior BehaviorOffNavMesh = EOffNavMeshBehavior::SwitchToWalking;
+
 	
 	UPROPERTY(Transient)
 	float NavMeshProjectionTimer;
