Mover: fix stale logging that was accidentally included in previous submission

#rnx

diff --git a/Engine/Plugins/Experimental/Mover/Source/Mover/Private/DefaultMovementSet/Modes/NavWalkingMode.cpp b/Engine/Plugins/Experimental/Mover/Source/Mover/Private/DefaultMovementSet/Modes/NavWalkingMode.cpp
--- a/Engine/Plugins/Experimental/Mover/Source/Mover/Private/DefaultMovementSet/Modes/NavWalkingMode.cpp
+++ b/Engine/Plugins/Experimental/Mover/Source/Mover/Private/DefaultMovementSet/Modes/NavWalkingMode.cpp
@@ -234,7 +234,7 @@ void UNavWalkingMode::SimulationTick_Implementation(const FSimulationTickParams&
 			{
 				default:	// fall through
 				case EOffNavMeshBehavior::SwitchToWalking:
-					UE_LOG(LogMover, Verbose, TEXT("%s could not find valid navigation data at location %s. Switching to walking mode."), *GetNameSafe(MoverComp->GetOwner()), *TargetFeetLocation.ToCompactString());
+					UE_LOG(LogMover, Verbose, TEXT("%s could not find valid navigation data at location %s. Switching to walking mode."), *GetNameSafe(MoverComp->GetOwner()), *AdjustedDest.ToCompactString());
 					OutputState.MovementEndState.NextModeName = DefaultModeNames::Walking;
 					OutputState.MovementEndState.RemainingMs = Params.TimeStep.StepMs;
 					MoveRecord.SetDeltaSeconds(0.0f);
@@ -247,7 +247,7 @@ void UNavWalkingMode::SimulationTick_Implementation(const FSimulationTickParams&
 					break;
 
 				case EOffNavMeshBehavior::DoNotMove:
-					UE_LOG(LogMover, Verbose, TEXT("%s could not find valid navigation data at location %s. Cannot move."), *GetNameSafe(MoverComp->GetOwner()), *TargetFeetLocation.ToCompactString());
+					UE_LOG(LogMover, Verbose, TEXT("%s could not find valid navigation data at location %s. Cannot move."), *GetNameSafe(MoverComp->GetOwner()), *AdjustedDest.ToCompactString());
 					// nothing to be done
 					break;
 			}
