Mover navwalking changes:
- allow rotation in navwalking mode when there is no translation
- removed based movement cruft that was not active

diff --git a/Engine/Plugins/Experimental/Mover/Source/Mover/Private/DefaultMovementSet/Modes/AsyncNavWalkingMode.cpp b/Engine/Plugins/Experimental/Mover/Source/Mover/Private/DefaultMovementSet/Modes/AsyncNavWalkingMode.cpp
--- a/Engine/Plugins/Experimental/Mover/Source/Mover/Private/DefaultMovementSet/Modes/AsyncNavWalkingMode.cpp
+++ b/Engine/Plugins/Experimental/Mover/Source/Mover/Private/DefaultMovementSet/Modes/AsyncNavWalkingMode.cpp
@@ -128,7 +128,7 @@ void UAsyncNavWalkingMode::SimulationTick_Implementation(const FSimulationTickPa
 	}
 
 	const FMoverTickStartData& StartState = Params.StartState;
-	FProposedMove ProposedMove = Params.ProposedMove;
+	const FProposedMove& ProposedMove = Params.ProposedMove;
 	FVector UpDirection = MoverComp->GetUpDirection();
 
 	const FCharacterDefaultInputs* CharacterInputs = StartState.InputCmd.InputCollection.FindDataByType<FCharacterDefaultInputs>();
@@ -167,7 +167,6 @@ void UAsyncNavWalkingMode::SimulationTick_Implementation(const FSimulationTickPa
 	FVector LocationInProgress = StartLocation;
 	FQuat RotationInProgress = StartRotation;
 
-	FFloorCheckResult CurrentFloor;
 	const bool bDeltaMoveNearlyZero = OrigMoveDelta.IsNearlyZero();
 	FNavLocation DestNavLocation;
 
@@ -257,7 +256,7 @@ void UAsyncNavWalkingMode::SimulationTick_Implementation(const FSimulationTickPa
 					break;
 			}
 
-			CaptureFinalState(LocationInProgress, RotationInProgress.Rotator(), /*bDidAttemptMovement=*/true, CurrentFloor, MoveRecord, OutputSyncState);
+			CaptureFinalState(LocationInProgress, RotationInProgress.Rotator(), MoveRecord, OutputSyncState);
 			return;
 		}
 
@@ -287,6 +286,11 @@ void UAsyncNavWalkingMode::SimulationTick_Implementation(const FSimulationTickPa
 				RotationInProgress = FQuat::Slerp(StartRotation, TargetRotation, MoveHitResult.Time);
 			}
 		}
+		else
+		{
+			// not moving, but let's allow the full rotation
+			RotationInProgress = TargetRotation;
+		}
 	}
 	else
 	{
@@ -295,7 +299,7 @@ void UAsyncNavWalkingMode::SimulationTick_Implementation(const FSimulationTickPa
 		MoveRecord.SetDeltaSeconds((Params.TimeStep.StepMs - OutputState.MovementEndState.RemainingMs) * 0.001f);
 	}
 
-	CaptureFinalState(LocationInProgress, RotationInProgress.Rotator(), /*bDidAttemptMovement=*/true, CurrentFloor, MoveRecord, OutputSyncState);
+	CaptureFinalState(LocationInProgress, RotationInProgress.Rotator(), MoveRecord, OutputSyncState);
 }
 
 bool UAsyncNavWalkingMode::FindNavFloor(const FVector& TestLocation, FNavLocation& OutNavFloorLocation, const INavigationDataInterface* NavData) const
@@ -596,62 +600,20 @@ void UAsyncNavWalkingMode::OnUnregistered()
 	Super::OnUnregistered();
 }
 
-void UAsyncNavWalkingMode::CaptureFinalState(const FVector FinalLocation, const FRotator FinalRotation, bool bDidAttemptMovement, const FFloorCheckResult& FloorResult, const FMovementRecord& Record, FMoverDefaultSyncState& OutputSyncState) const
+void UAsyncNavWalkingMode::CaptureFinalState(const FVector FinalLocation, const FRotator FinalRotation, const FMovementRecord& Record, FMoverDefaultSyncState& OutputSyncState) const
 {
-	FRelativeBaseInfo PriorBaseInfo;
-
 	UMoverBlackboard* SimBlackboard = GetMoverComponent()->GetSimBlackboard_Mutable();
 
-	const bool bHasPriorBaseInfo = SimBlackboard->TryGet(CommonBlackboard::LastFoundDynamicMovementBase, PriorBaseInfo);
-
-	FRelativeBaseInfo CurrentBaseInfo = UpdateFloorAndBaseInfo(FloorResult);
-
 	// If we're on a dynamic base and we're not trying to move, keep using the same relative actor location. This prevents slow relative 
 	//  drifting that can occur from repeated floor sampling as the base moves through the world.
-	if (CurrentBaseInfo.HasRelativeInfo()
-		&& bHasPriorBaseInfo && !bDidAttemptMovement
-		&& PriorBaseInfo.UsesSameBase(CurrentBaseInfo))
-	{
-		CurrentBaseInfo.ContactLocalPosition = PriorBaseInfo.ContactLocalPosition;
-	}
-
-	if (CurrentBaseInfo.HasRelativeInfo())
-	{
-		SimBlackboard->Set(CommonBlackboard::LastFoundDynamicMovementBase, CurrentBaseInfo);
-
-		OutputSyncState.SetTransforms_WorldSpace(FinalLocation,
-			FinalRotation,
-			Record.GetRelevantVelocity(),
-			CurrentBaseInfo.MovementBase.Get(), CurrentBaseInfo.BoneName);
-	}
-	else
-	{
-		SimBlackboard->Invalidate(CommonBlackboard::LastFoundDynamicMovementBase);
+	SimBlackboard->Invalidate(CommonBlackboard::LastFoundDynamicMovementBase);
 
-		OutputSyncState.SetTransforms_WorldSpace(FinalLocation,
-			FinalRotation,
-			Record.GetRelevantVelocity(),
-			nullptr);	// no movement base
-	}
+	OutputSyncState.SetTransforms_WorldSpace(FinalLocation,
+		FinalRotation,
+		Record.GetRelevantVelocity(),
+		nullptr);	// no movement base
 }
 
-FRelativeBaseInfo UAsyncNavWalkingMode::UpdateFloorAndBaseInfo(const FFloorCheckResult& FloorResult) const
-{
-	FRelativeBaseInfo ReturnBaseInfo;
-
-	UMoverBlackboard* SimBlackboard = GetMoverComponent()->GetSimBlackboard_Mutable();
-
-	SimBlackboard->Set(CommonBlackboard::LastFloorResult, FloorResult);
-
-	if (FloorResult.IsWalkableFloor() && UBasedMovementUtils::IsADynamicBase(FloorResult.HitResult.GetComponent()))
-	{
-		ReturnBaseInfo.SetFromFloorResult(FloorResult);
-	}
-
-	return ReturnBaseInfo;
-}
-
-
 void UAsyncNavWalkingMode::PostFinalizeSetCollisionForNavWalking(const FMoverSyncState& SyncState, const FMoverAuxStateContext& AuxState)
 {
 	SetCollisionForNavWalking(true);
diff --git a/Engine/Plugins/Experimental/Mover/Source/Mover/Private/DefaultMovementSet/Modes/NavWalkingMode.cpp b/Engine/Plugins/Experimental/Mover/Source/Mover/Private/DefaultMovementSet/Modes/NavWalkingMode.cpp
--- a/Engine/Plugins/Experimental/Mover/Source/Mover/Private/DefaultMovementSet/Modes/NavWalkingMode.cpp
+++ b/Engine/Plugins/Experimental/Mover/Source/Mover/Private/DefaultMovementSet/Modes/NavWalkingMode.cpp
@@ -122,7 +122,7 @@ void UNavWalkingMode::SimulationTick_Implementation(const FSimulationTickParams&
 	const FMoverTickStartData& StartState = Params.StartState;
 	USceneComponent* UpdatedComponent = Params.MovingComps.UpdatedComponent.Get();
 	UPrimitiveComponent* UpdatedPrimitive = Params.MovingComps.UpdatedPrimitive.Get();
-	FProposedMove ProposedMove = Params.ProposedMove;
+	const FProposedMove& ProposedMove = Params.ProposedMove;
 	FVector UpDirection = MoverComp->GetUpDirection();
 
 	if (!UpdatedComponent || !UpdatedPrimitive)
@@ -150,10 +150,13 @@ void UNavWalkingMode::SimulationTick_Implementation(const FSimulationTickParams&
 	const FRotator StartingOrient = StartingSyncState->GetOrientation_WorldSpace();
 	FRotator TargetOrient = StartingOrient;
 
+	bool bIsOrientationChanging = false;
+
 	// Apply orientation changes (if any)
 	if (!UMovementUtils::IsAngularVelocityZero(ProposedMove.AngularVelocity))
 	{ 
 		TargetOrient += (ProposedMove.AngularVelocity * DeltaSeconds);
+		bIsOrientationChanging = true;
 	}
 	
 	const FQuat TargetOrientQuat = TargetOrient.Quaternion();
@@ -161,8 +164,7 @@ void UNavWalkingMode::SimulationTick_Implementation(const FSimulationTickParams&
 	
 	FVector StartingLocation = NavMoverComponent->GetFeetLocation();
 	FVector AdjustedDest = StartingLocation + OrigMoveDelta;
-	
-	FFloorCheckResult CurrentFloor; 
+
 	const bool bDeltaMoveNearlyZero = OrigMoveDelta.IsNearlyZero();
 	FNavLocation DestNavLocation;
 
@@ -251,8 +253,8 @@ void UNavWalkingMode::SimulationTick_Implementation(const FSimulationTickParams&
 					// nothing to be done
 					break;
 			}
-			
-			CaptureFinalState(UpdatedComponent, true, CurrentFloor, MoveRecord, OutputSyncState);
+
+			CaptureFinalState(UpdatedComponent, MoveRecord, OutputSyncState);
 			return;
 		}
 
@@ -272,7 +274,7 @@ void UNavWalkingMode::SimulationTick_Implementation(const FSimulationTickParams&
 
 		FVector AdjustedDelta = NewLocation - StartingLocation;
 		
-		if (!AdjustedDelta.IsNearlyZero())
+		if (!AdjustedDelta.IsNearlyZero() || bIsOrientationChanging)
 		{
 			FHitResult MoveHitResult;
 			UMovementUtils::TrySafeMoveUpdatedComponent(Params.MovingComps, AdjustedDelta, TargetOrientQuat, bSweepWhileNavWalking, MoveHitResult, ETeleportType::None, MoveRecord);
@@ -286,7 +288,7 @@ void UNavWalkingMode::SimulationTick_Implementation(const FSimulationTickParams&
 		MoveRecord.SetDeltaSeconds(0.0f);
 	}
 
-	CaptureFinalState(UpdatedComponent, true, CurrentFloor, MoveRecord, OutputSyncState);
+	CaptureFinalState(UpdatedComponent, MoveRecord, OutputSyncState);
 }
 
 bool UNavWalkingMode::FindNavFloor(const FVector& TestLocation, FNavLocation& OutNavFloorLocation, const INavigationDataInterface* NavData) const
@@ -581,61 +583,19 @@ void UNavWalkingMode::OnUnregistered()
 	Super::OnUnregistered();
 }
 
-void UNavWalkingMode::CaptureFinalState(USceneComponent* UpdatedComponent, bool bDidAttemptMovement, const FFloorCheckResult& FloorResult, const FMovementRecord& Record, FMoverDefaultSyncState& OutputSyncState) const
+void UNavWalkingMode::CaptureFinalState(USceneComponent* UpdatedComponent, const FMovementRecord& Record, FMoverDefaultSyncState& OutputSyncState) const
 {
-	FRelativeBaseInfo PriorBaseInfo;
-
 	UMoverBlackboard* SimBlackboard = GetMoverComponent()->GetSimBlackboard_Mutable();
 
-	const bool bHasPriorBaseInfo = SimBlackboard->TryGet(CommonBlackboard::LastFoundDynamicMovementBase, PriorBaseInfo);
-
-	FRelativeBaseInfo CurrentBaseInfo = UpdateFloorAndBaseInfo(FloorResult);
-
-	// If we're on a dynamic base and we're not trying to move, keep using the same relative actor location. This prevents slow relative 
-	//  drifting that can occur from repeated floor sampling as the base moves through the world.
-	if (CurrentBaseInfo.HasRelativeInfo() 
-		&& bHasPriorBaseInfo && !bDidAttemptMovement 
-		&& PriorBaseInfo.UsesSameBase(CurrentBaseInfo))
-	{
-		CurrentBaseInfo.ContactLocalPosition = PriorBaseInfo.ContactLocalPosition;
-	}
-	
-	if (CurrentBaseInfo.HasRelativeInfo())
-	{
-		SimBlackboard->Set(CommonBlackboard::LastFoundDynamicMovementBase, CurrentBaseInfo);
-
-		OutputSyncState.SetTransforms_WorldSpace( UpdatedComponent->GetComponentLocation(),
-												  UpdatedComponent->GetComponentRotation(),
-												  Record.GetRelevantVelocity(),
-												  CurrentBaseInfo.MovementBase.Get(), CurrentBaseInfo.BoneName);
-	}
-	else
-	{
-		SimBlackboard->Invalidate(CommonBlackboard::LastFoundDynamicMovementBase);
+	SimBlackboard->Invalidate(CommonBlackboard::LastFoundDynamicMovementBase);
 
-		OutputSyncState.SetTransforms_WorldSpace( UpdatedComponent->GetComponentLocation(),
-												  UpdatedComponent->GetComponentRotation(),
-												  Record.GetRelevantVelocity(),
-												  nullptr);	// no movement base
-	}
+	OutputSyncState.SetTransforms_WorldSpace(UpdatedComponent->GetComponentLocation(),
+		UpdatedComponent->GetComponentRotation(),
+		Record.GetRelevantVelocity(),
+		nullptr);	// no movement base
 
 	UpdatedComponent->ComponentVelocity = OutputSyncState.GetVelocity_WorldSpace();
 }
 
-FRelativeBaseInfo UNavWalkingMode::UpdateFloorAndBaseInfo(const FFloorCheckResult& FloorResult) const
-{
-	FRelativeBaseInfo ReturnBaseInfo;
-
-	UMoverBlackboard* SimBlackboard = GetMoverComponent()->GetSimBlackboard_Mutable();
-
-	SimBlackboard->Set(CommonBlackboard::LastFloorResult, FloorResult);
-
-	if (FloorResult.IsWalkableFloor() && UBasedMovementUtils::IsADynamicBase(FloorResult.HitResult.GetComponent()))
-	{
-		ReturnBaseInfo.SetFromFloorResult(FloorResult);
-	}
-
-	return ReturnBaseInfo;
-}
 
 
diff --git a/Engine/Plugins/Experimental/Mover/Source/Mover/Public/DefaultMovementSet/Modes/AsyncNavWalkingMode.h b/Engine/Plugins/Experimental/Mover/Source/Mover/Public/DefaultMovementSet/Modes/AsyncNavWalkingMode.h
--- a/Engine/Plugins/Experimental/Mover/Source/Mover/Public/DefaultMovementSet/Modes/AsyncNavWalkingMode.h
+++ b/Engine/Plugins/Experimental/Mover/Source/Mover/Public/DefaultMovementSet/Modes/AsyncNavWalkingMode.h
@@ -124,9 +124,7 @@ protected:
 	virtual void OnRegistered(const FName ModeName) override;
 	virtual void OnUnregistered() override;
 
-	void CaptureFinalState(const FVector FinalLocation, const FRotator FinalRotation, bool bDidAttemptMovement, const FFloorCheckResult& FloorResult, const FMovementRecord& Record, FMoverDefaultSyncState& OutputSyncState) const;
-
-	FRelativeBaseInfo UpdateFloorAndBaseInfo(const FFloorCheckResult& FloorResult) const;
+	void CaptureFinalState(const FVector FinalLocation, const FRotator FinalRotation, const FMovementRecord& Record, FMoverDefaultSyncState& OutputSyncState) const;
 
 	TWeakObjectPtr<const UCommonLegacyMovementSettings> CommonLegacySettings;
  
diff --git a/Engine/Plugins/Experimental/Mover/Source/Mover/Public/DefaultMovementSet/Modes/NavWalkingMode.h b/Engine/Plugins/Experimental/Mover/Source/Mover/Public/DefaultMovementSet/Modes/NavWalkingMode.h
--- a/Engine/Plugins/Experimental/Mover/Source/Mover/Public/DefaultMovementSet/Modes/NavWalkingMode.h
+++ b/Engine/Plugins/Experimental/Mover/Source/Mover/Public/DefaultMovementSet/Modes/NavWalkingMode.h
@@ -145,11 +145,9 @@ protected:
 	
 	virtual void OnRegistered(const FName ModeName) override;
 	virtual void OnUnregistered() override;
-	
-	void CaptureFinalState(USceneComponent* UpdatedComponent, bool bDidAttemptMovement, const FFloorCheckResult& FloorResult, const FMovementRecord& Record, FMoverDefaultSyncState& OutputSyncState) const;
 
-	FRelativeBaseInfo UpdateFloorAndBaseInfo(const FFloorCheckResult& FloorResult) const;
-	
+	void CaptureFinalState(USceneComponent* UpdatedComponent, const FMovementRecord& Record, FMoverDefaultSyncState& OutputSyncState) const;
+
 	TObjectPtr<const UCommonLegacyMovementSettings> CommonLegacySettings;
 
 	// Collision responses saved at the activation of this mode, to be restored as upon deactivation if the CDO can't be used
