Physics Mover - Improve handling of landing to avoid stall in velocity

diff --git a/Engine/Plugins/Experimental/Mover/Source/Mover/Private/PhysicsMover/Modes/PhysicsDrivenFallingMode.cpp b/Engine/Plugins/Experimental/Mover/Source/Mover/Private/PhysicsMover/Modes/PhysicsDrivenFallingMode.cpp
--- a/Engine/Plugins/Experimental/Mover/Source/Mover/Private/PhysicsMover/Modes/PhysicsDrivenFallingMode.cpp
+++ b/Engine/Plugins/Experimental/Mover/Source/Mover/Private/PhysicsMover/Modes/PhysicsDrivenFallingMode.cpp
@@ -148,8 +148,6 @@ void UPhysicsDrivenFallingMode::SimulationTick_Implementation(const FSimulationT
 		TargetVel -= (CurPhysVolume->GetGravityZ() * FVector::UpVector * DeltaSeconds);
 	}
 
-	FVector TargetPos = StartingSyncState->GetLocation_WorldSpace() + TargetVel * DeltaSeconds;
-
 	constexpr float FloorDistanceTolerance = 2.0f;
 	const FVector ProjectedGroundVelocity = UPhysicsMovementUtils::ComputeIntegratedGroundVelocityFromHitResult(StartingSyncState->GetLocation_WorldSpace(), FloorResult.HitResult, DeltaSeconds);
 	const float ProjectedRelativeVerticalVelocity = FloorResult.HitResult.ImpactNormal.Dot(ProposedMove.LinearVelocity - ProjectedGroundVelocity);
@@ -157,13 +155,15 @@ void UPhysicsDrivenFallingMode::SimulationTick_Implementation(const FSimulationT
 	const bool bIsFloorWithinReach = ProjectedFloorDistance < TargetHeight + FloorDistanceTolerance;
 	const bool bIsMovingUpRelativeToFloor = ProjectedRelativeVerticalVelocity > UE_KINDA_SMALL_NUMBER;
 
+	FVector TargetPos = StartingSyncState->GetLocation_WorldSpace();
 	if (FloorResult.IsWalkableFloor() && bIsFloorWithinReach && !bIsMovingUpRelativeToFloor)
 	{
 		OutputState.MovementEndState.NextModeName = CommonLegacySettings->GroundMovementModeName;
-		//TargetVel = FVector::VectorPlaneProject(TargetVel, FloorResult.HitResult.Normal);
-		const FPlane MovementPlane(FVector::ZeroVector, UpDir);
-		TargetVel = UMovementUtils::ConstrainToPlane(TargetVel, MovementPlane, false);
-		TargetPos -= (UpDir.Dot(TargetVel * DeltaSeconds) + (FloorResult.FloorDist - TargetHeight)) * UpDir;
+		TargetPos += UpDir * (TargetHeight - FloorResult.FloorDist);
+	}
+	else
+	{
+		TargetPos += TargetVel * DeltaSeconds;
 	}
 
 	OutputState.MovementEndState.RemainingMs = 0.0f;
diff --git a/Engine/Plugins/Experimental/Mover/Source/Mover/Private/PhysicsMover/Modes/PhysicsDrivenWalkingMode.cpp b/Engine/Plugins/Experimental/Mover/Source/Mover/Private/PhysicsMover/Modes/PhysicsDrivenWalkingMode.cpp
--- a/Engine/Plugins/Experimental/Mover/Source/Mover/Private/PhysicsMover/Modes/PhysicsDrivenWalkingMode.cpp
+++ b/Engine/Plugins/Experimental/Mover/Source/Mover/Private/PhysicsMover/Modes/PhysicsDrivenWalkingMode.cpp
@@ -332,6 +332,16 @@ void UPhysicsDrivenWalkingMode::SimulationTick_Implementation(const FSimulationT
 
 	ProposedMove.LinearVelocity = OutDeltaPos / DeltaSeconds;
 
+	// The base movement mode does not apply gravity in walking mode so apply here
+	// This is so that the gravity in this mode will be consistent with the gravity
+	// set on the mover, not the default physics gravity
+	FVector TargetVelocity = StartingSyncState->GetVelocity_WorldSpace() + UMovementUtils::ComputeVelocityFromGravity(MoverComp->GetGravityAcceleration(), DeltaSeconds);
+	if (const APhysicsVolume* CurPhysVolume = UpdatedComponent->GetPhysicsVolume())
+	{
+		// The physics simulation applies Z-only gravity acceleration via physics volumes, so we need to account for it here 
+		TargetVelocity -= (CurPhysVolume->GetGravityZ() * FVector::UpVector * DeltaSeconds);
+	}
+
 	SimBlackboard->Set(CommonBlackboard::LastFloorResult, FloorResult);
 	SimBlackboard->Set(CommonBlackboard::LastWaterResult, WaterResult);
 
@@ -344,24 +354,30 @@ void UPhysicsDrivenWalkingMode::SimulationTick_Implementation(const FSimulationT
 	else if (FloorResult.IsWalkableFloor())
 	{
 		const FVector StartGroundVelocity = UPhysicsMovementUtils::ComputeGroundVelocityFromHitResult(StartingSyncState->GetLocation_WorldSpace(), FloorResult.HitResult, DeltaSeconds);
-		
-		FVector TargetVelocity = StartingSyncState->GetVelocity_WorldSpace();
+		const FVector ProjectedGroundVelocity = UPhysicsMovementUtils::ComputeIntegratedGroundVelocityFromHitResult(StartingSyncState->GetLocation_WorldSpace(), FloorResult.HitResult, DeltaSeconds);
+		const bool bIsGroundMoving = ProjectedGroundVelocity.SizeSquared() > UE_KINDA_SMALL_NUMBER;
+
+		TargetVelocity = StartingSyncState->GetVelocity_WorldSpace();
 		FVector TargetPosition = StartingSyncState->GetLocation_WorldSpace();
-		if (bHandleVerticalLandingSeparately && FVector::Parallel(StartingSyncState->GetVelocity_WorldSpace().GetSafeNormal(), UpDir))
+		if (FloorResult.FloorDist > TargetHeight + UE_KINDA_SMALL_NUMBER)
 		{
-			TargetPosition = StartingSyncState->GetLocation_WorldSpace() - UpDir * (FloorResult.FloorDist - TargetHeight);
+			TargetPosition += UpDir * (TargetHeight - FloorResult.FloorDist);
 		}
-		else
+
+		constexpr float ParallelCosThreshold = 0.999f;
+		const bool bVerticalVelocity = FVector::Parallel(TargetVelocity.GetSafeNormal(), UpDir, ParallelCosThreshold);
+		const bool bUseProposedMove = !(bHandleVerticalLandingSeparately && bVerticalVelocity) || ProposedMove.bHasDirIntent;
+
+		if (bUseProposedMove)
 		{
 			const FVector ProposedMovePlaneVelocity = ProposedMove.LinearVelocity - ProposedMove.LinearVelocity.ProjectOnToNormal(PrevGroundNormal);
-			const FVector StartingMovePlaneVelocity = StartingSyncState->GetVelocity_WorldSpace() - StartingSyncState->GetVelocity_WorldSpace().ProjectOnToNormal(PrevGroundNormal);
 
 			// If there is velocity intent in the normal direction then use the velocity from the proposed move. Otherwise
 			// retain the previous vertical velocity
 			FVector ProposedNormalVelocity = ProposedMove.LinearVelocity - ProposedMovePlaneVelocity;
-			if (ProposedNormalVelocity.SizeSquared() <= UE_SMALL_NUMBER)
+			if (ProposedNormalVelocity.SizeSquared() > UE_KINDA_SMALL_NUMBER)
 			{
-				ProposedNormalVelocity = StartingSyncState->GetVelocity_WorldSpace() - StartingMovePlaneVelocity;
+				TargetVelocity += ProposedNormalVelocity - TargetVelocity.ProjectOnToNormal(PrevGroundNormal);
 			}
 
 			TargetPosition += ProposedMovePlaneVelocity * DeltaSeconds;
@@ -375,12 +391,13 @@ void UPhysicsDrivenWalkingMode::SimulationTick_Implementation(const FSimulationT
 		{
 			CharacterGravity = PhysVolume->GetGravityZ();
 		}
-		const FVector ProjectedVelocity = TargetVelocity + CharacterGravity * FVector::UpVector * DeltaSeconds;
-		const FVector ProjectedGroundVelocity = UPhysicsMovementUtils::ComputeIntegratedGroundVelocityFromHitResult(StartingSyncState->GetLocation_WorldSpace(), FloorResult.HitResult, DeltaSeconds);
-		const bool bIsGroundMoving = FMath::Abs(ProjectedGroundVelocity.Z) > UE_KINDA_SMALL_NUMBER;
+
+		FVector ProjectedVelocity = TargetVelocity + CharacterGravity * FVector::UpVector * DeltaSeconds;
+
 		const float ProjectedRelativeVerticalVelocity = FloorResult.HitResult.ImpactNormal.Dot(ProjectedVelocity - ProjectedGroundVelocity);
 		const float VerticalVelocityLimit = 2.0f / DeltaSeconds;
-		if ((ProjectedRelativeVerticalVelocity > VerticalVelocityLimit) && (ProjectedVelocity.Z > VerticalVelocityLimit) && bIsGroundMoving)
+
+		if ((ProjectedRelativeVerticalVelocity > VerticalVelocityLimit) && bIsGroundMoving && (ProjectedVelocity.Dot(UpDir) > VerticalVelocityLimit))
 		{
 			bIsLiftingOffSurface = true;
 		}
@@ -402,12 +419,18 @@ void UPhysicsDrivenWalkingMode::SimulationTick_Implementation(const FSimulationT
 		{
 			SimBlackboard->Set(CommonBlackboard::TimeSinceSupported, 0.0f);
 		}
-		else
+		else if (!bIsLiftingOffSurface)
 		{
+			// Falling
 			TimeSinceSupported += DeltaSeconds;
 			SimBlackboard->Set(CommonBlackboard::TimeSinceSupported, TimeSinceSupported);
 			bIsSupported = TimeSinceSupported < MaxUnsupportedTimeBeforeFalling;
 		}
+		else
+		{
+			// Moving up relative to ground
+			SimBlackboard->Set(CommonBlackboard::TimeSinceSupported, MaxUnsupportedTimeBeforeFalling);
+		}
 
 		// Apply vertical velocity to target if stepping down
 		const bool bNeedsVerticalVelocityToTarget = bIsSupported && bIsSteppingDown && (EndHeightAboutGround > 0.0f) && !bIsLiftingOffSurface;
