Mover: now sanitizing scale factor coming from a root motion attribute, as it is never desirable.

diff --git a/Engine/Plugins/Experimental/Mover/Source/Mover/Private/DefaultMovementSet/LayeredMoves/RootMotionAttributeLayeredMove.cpp b/Engine/Plugins/Experimental/Mover/Source/Mover/Private/DefaultMovementSet/LayeredMoves/RootMotionAttributeLayeredMove.cpp
--- a/Engine/Plugins/Experimental/Mover/Source/Mover/Private/DefaultMovementSet/LayeredMoves/RootMotionAttributeLayeredMove.cpp
+++ b/Engine/Plugins/Experimental/Mover/Source/Mover/Private/DefaultMovementSet/LayeredMoves/RootMotionAttributeLayeredMove.cpp
@@ -73,6 +73,7 @@ bool FLayeredMove_RootMotionAttribute::GenerateMove(const FMoverTickStartData& S
 				// NOTE this will only work for ticking modes that tick in time with the world tick, because it relies on the Mesh ticking at the same rate as the movement simulation. 
 				// For fixed-tick modes, the attribute would be better as an accumulator for the movement sim to consume, along with a time accumulation attribute.
 				LocalRootMotion = RootMotionAttribute->Value;
+				LocalRootMotion.SetScale3D(FVector::OneVector);	// sanitize any scaling factor
 				bDidAttrHaveRootMotion = true;
 			}
 		}
diff --git a/Engine/Plugins/Experimental/Mover/Source/Mover/Public/DefaultMovementSet/LayeredMoves/RootMotionAttributeLayeredMove.h b/Engine/Plugins/Experimental/Mover/Source/Mover/Public/DefaultMovementSet/LayeredMoves/RootMotionAttributeLayeredMove.h
--- a/Engine/Plugins/Experimental/Mover/Source/Mover/Public/DefaultMovementSet/LayeredMoves/RootMotionAttributeLayeredMove.h
+++ b/Engine/Plugins/Experimental/Mover/Source/Mover/Public/DefaultMovementSet/LayeredMoves/RootMotionAttributeLayeredMove.h
@@ -7,8 +7,8 @@
 #include "RootMotionAttributeLayeredMove.generated.h"
 
 /** 
- * Root Motion Attribute Move: handles root motion from a mesh's custom attribute. 
- * Currently only supports Independent ticking mode, and allows air control while jumping/falling.
+ * Root Motion Attribute Move: handles root motion from a mesh's custom attribute, ignoring scaling.
+ * Currently only supports Independent ticking mode, and allows controlled movement while jumping/falling or when a SkipAnimRootMotion tag is active.
  */
 USTRUCT(BlueprintType)
 struct MOVER_API FLayeredMove_RootMotionAttribute : public FLayeredMoveBase
