Mover cleanup: avoiding some default initialization and favoring newer ustruct body macro

#rnx

diff --git a/Engine/Plugins/Experimental/Mover/Source/Mover/Public/MoveLibrary/MovementRecord.h b/Engine/Plugins/Experimental/Mover/Source/Mover/Public/MoveLibrary/MovementRecord.h
--- a/Engine/Plugins/Experimental/Mover/Source/Mover/Public/MoveLibrary/MovementRecord.h
+++ b/Engine/Plugins/Experimental/Mover/Source/Mover/Public/MoveLibrary/MovementRecord.h
@@ -8,7 +8,7 @@
 USTRUCT(BlueprintType)
 struct MOVER_API FMovementSubstep
 {
-	GENERATED_USTRUCT_BODY()
+	GENERATED_BODY()
 
 	FMovementSubstep() {}
 	FMovementSubstep(FName InMoveName, FVector InMoveDelta, bool InIsRelevant = true)
@@ -29,7 +29,7 @@ struct MOVER_API FMovementSubstep
 USTRUCT(BlueprintType)
 struct MOVER_API FMovementRecord
 {
-	GENERATED_USTRUCT_BODY()
+	GENERATED_BODY()
 
 	void Append(FMovementSubstep Substep);
 
diff --git a/Engine/Plugins/Experimental/Mover/Source/Mover/Public/MoveLibrary/MovementUtilsTypes.h b/Engine/Plugins/Experimental/Mover/Source/Mover/Public/MoveLibrary/MovementUtilsTypes.h
--- a/Engine/Plugins/Experimental/Mover/Source/Mover/Public/MoveLibrary/MovementUtilsTypes.h
+++ b/Engine/Plugins/Experimental/Mover/Source/Mover/Public/MoveLibrary/MovementUtilsTypes.h
@@ -25,7 +25,7 @@ enum class EMoveMixMode : uint8
 USTRUCT(BlueprintType)
 struct MOVER_API FProposedMove
 {
-	GENERATED_USTRUCT_BODY()
+	GENERATED_BODY()
 
 	FProposedMove() : 
 		bHasDirIntent(false)
@@ -68,18 +68,18 @@ struct MOVER_API FProposedMove
 USTRUCT(BlueprintType)
 struct MOVER_API FMovingComponentSet
 {
-	GENERATED_USTRUCT_BODY()
+	GENERATED_BODY()
 
 	FMovingComponentSet() {}
 	FMovingComponentSet(USceneComponent* InUpdatedComponent);
 	FMovingComponentSet(UMoverComponent* InMoverComponent);
 
 	UPROPERTY(EditAnywhere, BlueprintReadWrite, Category = Mover)
-	TWeakObjectPtr<USceneComponent> UpdatedComponent = nullptr;
+	TWeakObjectPtr<USceneComponent> UpdatedComponent;
 
 	UPROPERTY(EditAnywhere, BlueprintReadWrite, Category = Mover)
-	TWeakObjectPtr<UPrimitiveComponent> UpdatedPrimitive = nullptr;
+	TWeakObjectPtr<UPrimitiveComponent> UpdatedPrimitive;
 
 	UPROPERTY(EditAnywhere, BlueprintReadWrite, Category = Mover)
-	TWeakObjectPtr<UMoverComponent> MoverComponent = nullptr;
+	TWeakObjectPtr<UMoverComponent> MoverComponent;
 };
\ No newline at end of file
diff --git a/Engine/Plugins/Experimental/Mover/Source/Mover/Public/MoverSimulationTypes.h b/Engine/Plugins/Experimental/Mover/Source/Mover/Public/MoverSimulationTypes.h
--- a/Engine/Plugins/Experimental/Mover/Source/Mover/Public/MoverSimulationTypes.h
+++ b/Engine/Plugins/Experimental/Mover/Source/Mover/Public/MoverSimulationTypes.h
@@ -39,7 +39,7 @@ namespace CommonBlackboard
 USTRUCT(BlueprintType)
 struct MOVER_API FMovementModeTickEndState
 {
-	GENERATED_USTRUCT_BODY()
+	GENERATED_BODY()
 	
 	FMovementModeTickEndState() 
 	{ 
@@ -70,7 +70,7 @@ struct MOVER_API FMovementModeTickEndState
 USTRUCT(BlueprintType)
 struct MOVER_API FMoverInputCmdContext
 {
-	GENERATED_USTRUCT_BODY()
+	GENERATED_BODY()
 
 public:
 	UPROPERTY(BlueprintReadWrite, Category = Mover)
@@ -103,7 +103,7 @@ public:
 USTRUCT(BlueprintType)
 struct MOVER_API FMoverSyncState
 {
-	GENERATED_USTRUCT_BODY()
+	GENERATED_BODY()
 
 public:
 
@@ -179,7 +179,7 @@ public:
 USTRUCT(BlueprintType)
 struct MOVER_API FMoverAuxStateContext
 {
-	GENERATED_USTRUCT_BODY()
+	GENERATED_BODY()
 
 public:
 	UScriptStruct* GetStruct() const { return StaticStruct(); }
@@ -209,7 +209,7 @@ public:
 USTRUCT(BlueprintType)
 struct MOVER_API FMoverTickStartData
 {
-	GENERATED_USTRUCT_BODY()
+	GENERATED_BODY()
 
 	FMoverTickStartData() {}
 	FMoverTickStartData(
@@ -234,7 +234,7 @@ struct MOVER_API FMoverTickStartData
 USTRUCT(BlueprintType)
 struct MOVER_API FMoverTickEndData
 {
-	GENERATED_USTRUCT_BODY()
+	GENERATED_BODY()
 
 	FMoverTickEndData() {}
 	FMoverTickEndData(
