Physics Mover - Fix crash due to use of StateMachine on GT

diff --git a/Engine/Plugins/Experimental/Mover/Source/Mover/Private/Backends/MoverNetworkPhysicsLiaison.cpp b/Engine/Plugins/Experimental/Mover/Source/Mover/Private/Backends/MoverNetworkPhysicsLiaison.cpp
--- a/Engine/Plugins/Experimental/Mover/Source/Mover/Private/Backends/MoverNetworkPhysicsLiaison.cpp
+++ b/Engine/Plugins/Experimental/Mover/Source/Mover/Private/Backends/MoverNetworkPhysicsLiaison.cpp
@@ -150,7 +150,7 @@ void UMoverNetworkPhysicsLiaisonComponent::UpdateConstraintSettings()
 		Constraint->SetCosMaxWalkableSlopeAngle(CommonMovementSettings->MaxWalkSlopeCosine);
 		Constraint->SetMaxCharacterGroundMassRatio(GPhysicsDrivenMotionDebugParams.MaxCharacterGroundMassRatio);
 
-		const UBaseMovementMode* CurrentMode = MoverComp.ModeFSM->GetCurrentMode();
+		const UBaseMovementMode* CurrentMode = MoverComp.GetActiveMode();
 		if (CurrentMode && CurrentMode->Implements<UPhysicsCharacterMovementModeInterface>())
 		{
 			const IPhysicsCharacterMovementModeInterface* PhysicsMode = CastChecked<IPhysicsCharacterMovementModeInterface>(CurrentMode);
