PhysicsMover - Remove requirement for there to be a FCharacterDefaultInputs in the inputs

diff --git a/Engine/Plugins/Experimental/Mover/Source/Mover/Private/Backends/MoverNetworkPhysicsLiaison.cpp b/Engine/Plugins/Experimental/Mover/Source/Mover/Private/Backends/MoverNetworkPhysicsLiaison.cpp
--- a/Engine/Plugins/Experimental/Mover/Source/Mover/Private/Backends/MoverNetworkPhysicsLiaison.cpp
+++ b/Engine/Plugins/Experimental/Mover/Source/Mover/Private/Backends/MoverNetworkPhysicsLiaison.cpp
@@ -285,11 +285,6 @@ bool UMoverNetworkPhysicsLiaisonComponent::HasValidState() const
 
 bool UMoverNetworkPhysicsLiaisonComponent::CanProcessInputs_Internal(const FPhysicsMoverAsyncInput& Input) const
 {
-	if (!Input.InputCmd.InputCollection.FindDataByType<FCharacterDefaultInputs>())
-	{
-		return false;
-	}
-	
 	return Super::CanProcessInputs_Internal(Input);
 }
 
@@ -330,9 +325,6 @@ bool UMoverNetworkPhysicsLiaisonComponent::CanSimulate_Internal(const FPhysicsMo
 
 void UMoverNetworkPhysicsLiaisonComponent::PerformPreSimulate_Internal(const FPhysicsMoverSimulationTickParams& TickParams, const FPhysicsMoverAsyncInput& Input, FPhysicsMoverAsyncOutput& Output) const
 {
-	FCharacterDefaultInputs* CharacterDefaultInputs = Input.InputCmd.InputCollection.FindMutableDataByType<FCharacterDefaultInputs>();
-	check(CharacterDefaultInputs);
-
 	Chaos::FCharacterGroundConstraintHandle* ConstraintHandle = Constraint->GetProxy<Chaos::FCharacterGroundConstraintProxy>()->GetPhysicsThreadAPI();
 	Chaos::FPBDRigidParticleHandle* CharacterParticle = ConstraintHandle->GetCharacterParticle()->CastToRigidParticle();
 	UMoverComponent& MoverComp = GetMoverComponent();
@@ -363,7 +355,8 @@ void UMoverNetworkPhysicsLiaisonComponent::PerformPreSimulate_Internal(const FPh
 	FMoverDefaultSyncState& InputSyncState = Input.SyncState.SyncStateCollection.FindOrAddMutableDataByType<FMoverDefaultSyncState>();
 	InputSyncState.SetTransforms_WorldSpace(CharacterParticle->GetX(), FRotator(CharacterParticle->GetR()), CharacterParticle->GetV() - LocalGroundVelocity + AIMoveVelocity);
 
-	if (!CharacterDefaultInputs->SuggestedMovementMode.IsNone())
+	FCharacterDefaultInputs* CharacterDefaultInputs = Input.InputCmd.InputCollection.FindMutableDataByType<FCharacterDefaultInputs>();
+	if (CharacterDefaultInputs && !CharacterDefaultInputs->SuggestedMovementMode.IsNone())
 	{
 		MoverComp.QueueNextMode(CharacterDefaultInputs->SuggestedMovementMode);
 	}
@@ -456,7 +449,10 @@ void UMoverNetworkPhysicsLiaisonComponent::PerformPreSimulate_Internal(const FPh
 			});
 	}
 
-	CharacterDefaultInputs->SuggestedMovementMode = NAME_None; // Should have already been set this frame
+	if (CharacterDefaultInputs)
+	{
+		CharacterDefaultInputs->SuggestedMovementMode = NAME_None; // Should have already been set this frame
+	}
 }
 
 Chaos::FPhysicsObject* UMoverNetworkPhysicsLiaisonComponent::GetControlledPhysicsObject() const
diff --git a/Engine/Plugins/Experimental/Mover/Source/Mover/Private/PhysicsMover/Modes/PhysicsDrivenFallingMode.cpp b/Engine/Plugins/Experimental/Mover/Source/Mover/Private/PhysicsMover/Modes/PhysicsDrivenFallingMode.cpp
--- a/Engine/Plugins/Experimental/Mover/Source/Mover/Private/PhysicsMover/Modes/PhysicsDrivenFallingMode.cpp
+++ b/Engine/Plugins/Experimental/Mover/Source/Mover/Private/PhysicsMover/Modes/PhysicsDrivenFallingMode.cpp
@@ -122,7 +122,8 @@ void UPhysicsDrivenFallingMode::SimulationTick_Implementation(const FSimulationT
 	SimBlackboard->Set(CommonBlackboard::LastFloorResult, FloorResult);
 	SimBlackboard->Set(CommonBlackboard::LastWaterResult, WaterResult);
 
-	const bool bIsMovingUp = CharacterInputs->bIsJumpJustPressed || (UpDir.Dot(ProposedMove.LinearVelocity) > 0.0f);
+	const bool bJumping = CharacterInputs && CharacterInputs->bIsJumpPressed;
+	const bool bIsMovingUp = bJumping || (UpDir.Dot(ProposedMove.LinearVelocity) > 0.0f);
 	const float ProjectedImmersionDepth = WaterResult.WaterSplineData.ImmersionDepth - (UpDir.Dot(ProposedMove.LinearVelocity) * DeltaSeconds);
 	const bool bStartSwimming = ProjectedImmersionDepth > CommonLegacySettings->SwimmingStartImmersionDepth;
 
diff --git a/Engine/Plugins/Experimental/Mover/Source/Mover/Private/PhysicsMover/Modes/PhysicsDrivenSwimmingMode.cpp b/Engine/Plugins/Experimental/Mover/Source/Mover/Private/PhysicsMover/Modes/PhysicsDrivenSwimmingMode.cpp
--- a/Engine/Plugins/Experimental/Mover/Source/Mover/Private/PhysicsMover/Modes/PhysicsDrivenSwimmingMode.cpp
+++ b/Engine/Plugins/Experimental/Mover/Source/Mover/Private/PhysicsMover/Modes/PhysicsDrivenSwimmingMode.cpp
@@ -101,7 +101,7 @@ void UPhysicsDrivenSwimmingMode::SimulationTick_Implementation(const FSimulation
 
 	const float DeltaSeconds = Params.TimeStep.StepMs * 0.001f;
 
-	if (CharacterInputs->bIsJumpJustPressed && AttemptJump(SurfaceSwimmingWaterControlSettings.JumpMultiplier*CommonLegacySettings->JumpUpwardsSpeed, OutputState))
+	if (CharacterInputs && CharacterInputs->bIsJumpJustPressed && AttemptJump(SurfaceSwimmingWaterControlSettings.JumpMultiplier*CommonLegacySettings->JumpUpwardsSpeed, OutputState))
 	{
 		OutputState.MovementEndState.RemainingMs = Params.TimeStep.StepMs;
 		return;
