Physics Mover: adjusting some code that assumes FCharacterDefaultInputs are always present

diff --git a/Engine/Plugins/Experimental/Mover/Source/Mover/Private/PhysicsMover/Transitions/PhysicsJumpCheck.cpp b/Engine/Plugins/Experimental/Mover/Source/Mover/Private/PhysicsMover/Transitions/PhysicsJumpCheck.cpp
--- a/Engine/Plugins/Experimental/Mover/Source/Mover/Private/PhysicsMover/Transitions/PhysicsJumpCheck.cpp
+++ b/Engine/Plugins/Experimental/Mover/Source/Mover/Private/PhysicsMover/Transitions/PhysicsJumpCheck.cpp
@@ -36,12 +36,12 @@ FTransitionEvalResult UPhysicsJumpCheck::Evaluate_Implementation(const FSimulati
 	FTransitionEvalResult EvalResult; 
 
 	const FMoverTickStartData& StartState = Params.StartState;
-	const FCharacterDefaultInputs* CharacterInputs = StartState.InputCmd.InputCollection.FindDataByType<FCharacterDefaultInputs>();
-	check(CharacterInputs);
-
-	if (CharacterInputs->bIsJumpJustPressed)
+	if (const FCharacterDefaultInputs* CharacterInputs = StartState.InputCmd.InputCollection.FindDataByType<FCharacterDefaultInputs>())
 	{
-		EvalResult.NextMode = TransitionToMode;
+		if (CharacterInputs->bIsJumpJustPressed)
+		{
+			EvalResult.NextMode = TransitionToMode;
+		}
 	}
 
 	return EvalResult;
