Mover removing some unneeded MoverDataCollection additions in constructors and reset functions

diff --git a/Engine/Plugins/Experimental/Mover/Source/Mover/Public/MoverSimulationTypes.h b/Engine/Plugins/Experimental/Mover/Source/Mover/Public/MoverSimulationTypes.h
--- a/Engine/Plugins/Experimental/Mover/Source/Mover/Public/MoverSimulationTypes.h
+++ b/Engine/Plugins/Experimental/Mover/Source/Mover/Public/MoverSimulationTypes.h
@@ -76,12 +76,6 @@ public:
 	UPROPERTY(BlueprintReadWrite, Category = Mover)
 	FMoverDataCollection InputCollection;
 
-	FMoverInputCmdContext()
-	{
-		// TODO: Consider eliminating this default input, considering non-character / non-controlled Mover actors (like platforms)
-		InputCollection.FindOrAddDataByType<FCharacterDefaultInputs>();
-	}
-
 	UScriptStruct* GetStruct() const { return StaticStruct(); }
 
 	void NetSerialize(const FNetSerializeParams& P)
@@ -169,7 +163,6 @@ public:
 	{
 		MovementMode = NAME_None;
 		SyncStateCollection.Empty();
-		SyncStateCollection.FindOrAddDataByType<FMoverDefaultSyncState>();
 		LayeredMoves.Reset();
 		MovementModifiers.Reset();
 	}
