[Mover] [BugFix] Fixing StanceModifiers for Sim-Proxies

diff --git a/Engine/Plugins/Experimental/Mover/Source/Mover/Private/Backends/MoverNetworkPhysicsLiaisonBase.cpp b/Engine/Plugins/Experimental/Mover/Source/Mover/Private/Backends/MoverNetworkPhysicsLiaisonBase.cpp
--- a/Engine/Plugins/Experimental/Mover/Source/Mover/Private/Backends/MoverNetworkPhysicsLiaisonBase.cpp
+++ b/Engine/Plugins/Experimental/Mover/Source/Mover/Private/Backends/MoverNetworkPhysicsLiaisonBase.cpp
@@ -650,14 +650,15 @@ void UMoverNetworkPhysicsLiaisonComponentBase::ProcessInputs_Internal(int32 Phys
 				Chaos::FPBDRigidsEvolution* RigidsEvolution = GetPhysicsSolver()->GetEvolution();
 				const bool bIsSolverResim = RigidsEvolution->IsResimming();
 				const bool bIsFirstResimFrame = RigidsEvolution->IsResetting();
+				const bool bIsProxyRepResim = NetworkPhysicsComponent->GetNetworkPhysicsComponent_Internal()->GetPhysicsReplicationMode() == EPhysicsReplicationMode::Resimulation;
 
 				// Rollback mover state if on the first resimulation frame
-				if (bIsInputGenerator && bIsSolverResim && bIsFirstResimFrame)
+				if ((bIsInputGenerator || bIsProxyRepResim) && bIsSolverResim && bIsFirstResimFrame)
 				{
 					UMoverComponent& MoverComp = GetMoverComponent();
 					FMoverAuxStateContext UnusedInvalidAuxState;
 					FMoverAuxStateContext UnusedAuxState;
-					MoverComp.OnSimulationPreRollback(&LatestSyncState, &Input.SyncState, &UnusedInvalidAuxState, &UnusedAuxState);
+					MoverComp.OnSimulationPreRollback(&Input.SyncState, &LatestSyncState, &UnusedInvalidAuxState, &UnusedAuxState);
 					GetCurrentStateData(Input.SyncState);
 					MoverComp.OnSimulationRollback(&Input.SyncState, &UnusedAuxState);
 				}
diff --git a/Engine/Plugins/Experimental/Mover/Source/Mover/Private/PhysicsMover/MovementModifiers/PhysicsStanceModifier.cpp b/Engine/Plugins/Experimental/Mover/Source/Mover/Private/PhysicsMover/MovementModifiers/PhysicsStanceModifier.cpp
--- a/Engine/Plugins/Experimental/Mover/Source/Mover/Private/PhysicsMover/MovementModifiers/PhysicsStanceModifier.cpp
+++ b/Engine/Plugins/Experimental/Mover/Source/Mover/Private/PhysicsMover/MovementModifiers/PhysicsStanceModifier.cpp
@@ -20,6 +20,28 @@ void FPhysicsStanceModifier::AdjustCapsule(UMoverComponent* MoverComp, float Old
 	UpdateTargetHeight<UPhysicsDrivenFallingMode>(MoverComp, OldHalfHeight, NewHalfHeight);
 }
 
+void FPhysicsStanceModifier::OnStart(UMoverComponent* MoverComp, const FMoverTimeStep& TimeStep, const FMoverSyncState& SyncState, const FMoverAuxStateContext& AuxState)
+{
+	Super::OnStart(MoverComp, TimeStep, SyncState, AuxState);
+
+	// Ensures crouching, especially if triggered through state syncing (rollbacks, etc.)  
+	if (UCharacterMoverComponent* CharMoverComp = Cast<UCharacterMoverComponent>(MoverComp))
+	{
+		CharMoverComp->Crouch();
+	}
+}
+
+void FPhysicsStanceModifier::OnEnd(UMoverComponent* MoverComp, const FMoverTimeStep& TimeStep, const FMoverSyncState& SyncState, const FMoverAuxStateContext& AuxState)
+{
+	Super::OnEnd(MoverComp, TimeStep, SyncState, AuxState);
+
+	// Ensures uncrouching, especially if triggered through state syncing (rollbacks, etc.)  
+	if (UCharacterMoverComponent* CharMoverComp = Cast<UCharacterMoverComponent>(MoverComp))
+	{
+		CharMoverComp->UnCrouch();
+	}
+}
+
 bool FPhysicsStanceModifier::CanExpand_Internal(UMoverComponent* MoverComponent, USceneComponent* UpdatedComp, const FMoverSyncState& SyncState) const
 {
 	float StandingHalfHeight = 90;
diff --git a/Engine/Plugins/Experimental/Mover/Source/Mover/Public/PhysicsMover/MovementModifiers/PhysicsStanceModifier.h b/Engine/Plugins/Experimental/Mover/Source/Mover/Public/PhysicsMover/MovementModifiers/PhysicsStanceModifier.h
--- a/Engine/Plugins/Experimental/Mover/Source/Mover/Public/PhysicsMover/MovementModifiers/PhysicsStanceModifier.h
+++ b/Engine/Plugins/Experimental/Mover/Source/Mover/Public/PhysicsMover/MovementModifiers/PhysicsStanceModifier.h
@@ -31,6 +31,10 @@ struct MOVER_API FPhysicsStanceModifier : public FStanceModifier
 
 	static void OnPostSimulationTick(const FStanceModifier* Modifier, UMoverComponent* MoverComp, UPrimitiveComponent* UpdatedPrimitive,  bool bIsCrouching, bool& bPostProcessed, OUT bool& bStanceChanged);
 
+	virtual void OnStart(UMoverComponent* MoverComp, const FMoverTimeStep& TimeStep, const FMoverSyncState& SyncState, const FMoverAuxStateContext& AuxState) override;
+
+	virtual void OnEnd(UMoverComponent* MoverComp, const FMoverTimeStep& TimeStep, const FMoverSyncState& SyncState, const FMoverAuxStateContext& AuxState) override;
+
 protected:
 	// Modifies the Capsule Location by updating the target height of the movement modes
 	virtual void AdjustCapsule(UMoverComponent* MoverComp, float OldHalfHeight, float NewHalfHeight, float NewEyeHeight) override;
