[mover] Change name of scoped timer

#rnx

diff --git a/Engine/Plugins/Experimental/Mover/Source/Mover/Private/Backends/MoverStandaloneLiaison.cpp b/Engine/Plugins/Experimental/Mover/Source/Mover/Private/Backends/MoverStandaloneLiaison.cpp
--- a/Engine/Plugins/Experimental/Mover/Source/Mover/Private/Backends/MoverStandaloneLiaison.cpp
+++ b/Engine/Plugins/Experimental/Mover/Source/Mover/Private/Backends/MoverStandaloneLiaison.cpp
@@ -300,7 +300,7 @@ void UMoverStandaloneLiaisonComponent::TickMovementSimulation(float DeltaSeconds
 void UMoverStandaloneLiaisonComponent::TickApplySimulationState(float DeltaSeconds)
 {
 	{
-		SCOPED_NAMED_EVENT(SilverMoverBackendLiaison_Tick_FinalizeFrame, FColor::Green);
+		SCOPED_NAMED_EVENT(MoverBackendLiaison_Tick_FinalizeFrame, FColor::Green);
 		MoverComp->FinalizeFrame(&CachedLastSyncState, &CachedLastAuxState);
 	}
 }
