Mover: exposing turn generator interface for external API

diff --git a/Engine/Plugins/Experimental/Mover/Source/Mover/Public/MoveLibrary/ModularMovement.h b/Engine/Plugins/Experimental/Mover/Source/Mover/Public/MoveLibrary/ModularMovement.h
--- a/Engine/Plugins/Experimental/Mover/Source/Mover/Public/MoveLibrary/ModularMovement.h
+++ b/Engine/Plugins/Experimental/Mover/Source/Mover/Public/MoveLibrary/ModularMovement.h
@@ -12,8 +12,8 @@ class UMoverBlackboard;
 
 
 
-UINTERFACE(BlueprintType)
-class MOVER_API UTurnGeneratorInterface : public UInterface
+UINTERFACE(BlueprintType, MinimalAPI)
+class UTurnGeneratorInterface : public UInterface
 {
 	GENERATED_BODY()
 };
@@ -29,8 +29,8 @@ public:
 	// TODO: consider making GetTurn take a params struct instead of a fixed arg list
 
 	/** Returns an additive angular velocity (per-axis degrees/second) based on the starting state and timestep */
-	UFUNCTION(BlueprintNativeEvent)
-	FRotator GetTurn(FRotator TargetOrientation, const FMoverTickStartData& FullStartState, const FMoverDefaultSyncState& MoverState, const FMoverTimeStep& TimeStep, const FProposedMove& ProposedMove, UMoverBlackboard* SimBlackboard);
+	UFUNCTION(BlueprintNativeEvent, BlueprintCallable, Category=Mover)
+	MOVER_API FRotator GetTurn(FRotator TargetOrientation, const FMoverTickStartData& FullStartState, const FMoverDefaultSyncState& MoverState, const FMoverTimeStep& TimeStep, const FProposedMove& ProposedMove, UMoverBlackboard* SimBlackboard);
 };
 
 
