[Mover] Decay Inputs to make pawn movements less jarring when starting and stopping in resim mode and forward predicting

diff --git a/Engine/Plugins/Experimental/Mover/Source/Mover/Private/Backends/MoverNetworkPhysicsLiaisonBase.cpp b/Engine/Plugins/Experimental/Mover/Source/Mover/Private/Backends/MoverNetworkPhysicsLiaisonBase.cpp
--- a/Engine/Plugins/Experimental/Mover/Source/Mover/Private/Backends/MoverNetworkPhysicsLiaisonBase.cpp
+++ b/Engine/Plugins/Experimental/Mover/Source/Mover/Private/Backends/MoverNetworkPhysicsLiaisonBase.cpp
@@ -46,6 +46,11 @@ void FNetworkPhysicsMoverInputs::BuildData(const UActorComponent* NetworkCompone
 	}
 }
 
+void FNetworkPhysicsMoverInputs::DecayData(float DecayAmount)
+{
+	InputCmdContext.InputCollection.Decay(DecayAmount);
+}
+
 bool FNetworkPhysicsMoverInputs::NetSerialize(FArchive& Ar, class UPackageMap* PackageMap, bool& bOutSuccess)
 {
 	SerializeFrames(Ar);
diff --git a/Engine/Plugins/Experimental/Mover/Source/Mover/Private/MoverDataModelTypes.cpp b/Engine/Plugins/Experimental/Mover/Source/Mover/Private/MoverDataModelTypes.cpp
--- a/Engine/Plugins/Experimental/Mover/Source/Mover/Private/MoverDataModelTypes.cpp
+++ b/Engine/Plugins/Experimental/Mover/Source/Mover/Private/MoverDataModelTypes.cpp
@@ -140,6 +140,21 @@ void FCharacterDefaultInputs::Merge(const FMoverDataStructBase& From)
 	bIsJumpPressed |= TypedFrom.bIsJumpPressed;
 }
 
+static float CharacterDefaultInputsDecayAmountMultiplier = 0.25f;
+FAutoConsoleVariableRef CVarCharacterDefaultInputsDecayAmountMultiplier(
+	TEXT("Mover.Input.CharacterDefaultInputsDecayAmountMultiplier"),
+	CharacterDefaultInputsDecayAmountMultiplier,
+	TEXT("Multiplier to use when decaying CharacterDefaultInputs."));
+
+void FCharacterDefaultInputs::Decay(float DecayAmount)
+{
+	DecayAmount *= CharacterDefaultInputsDecayAmountMultiplier;
+
+	MoveInput *= (1.0f - DecayAmount);
+
+	// Single use inputs
+	bIsJumpJustPressed = FMath::IsNearlyZero(DecayAmount) ? bIsJumpJustPressed : false;
+}
 
 // FMoverDefaultSyncState //////////////////////////////////////////////////////////////
 
diff --git a/Engine/Plugins/Experimental/Mover/Source/Mover/Private/MoverTypes.cpp b/Engine/Plugins/Experimental/Mover/Source/Mover/Private/MoverTypes.cpp
--- a/Engine/Plugins/Experimental/Mover/Source/Mover/Private/MoverTypes.cpp
+++ b/Engine/Plugins/Experimental/Mover/Source/Mover/Private/MoverTypes.cpp
@@ -215,6 +215,13 @@ void FMoverDataCollection::Merge(const FMoverDataCollection& From)
 	}
 }
 
+void FMoverDataCollection::Decay(float DecayAmount)
+{
+	for (const TSharedPtr<FMoverDataStructBase>& Element : DataArray)
+	{
+		Element->Decay(DecayAmount);
+	}
+}
 
 void FMoverDataCollection::AddStructReferencedObjects(FReferenceCollector& Collector) const
 {
diff --git a/Engine/Plugins/Experimental/Mover/Source/Mover/Private/PhysicsMover/PhysicsMoverSimulationTypes.cpp b/Engine/Plugins/Experimental/Mover/Source/Mover/Private/PhysicsMover/PhysicsMoverSimulationTypes.cpp
--- a/Engine/Plugins/Experimental/Mover/Source/Mover/Private/PhysicsMover/PhysicsMoverSimulationTypes.cpp
+++ b/Engine/Plugins/Experimental/Mover/Source/Mover/Private/PhysicsMover/PhysicsMoverSimulationTypes.cpp
@@ -94,6 +94,20 @@ void FMovementSettingsInputs::Merge(const FMoverDataStructBase& From)
 {
 }
 
+static float MovementSettingsInputsDecayAmountMultiplier = 0.25f;
+FAutoConsoleVariableRef CVarMovementSettingsInputsDecayAmountMultiplier(
+	TEXT("Mover.Input.MovementSettingsInputsDecayAmountMultiplier"),
+	MovementSettingsInputsDecayAmountMultiplier,
+	TEXT("Multiplier to use when decaying MovementSettingsInputs."));
+
+void FMovementSettingsInputs::Decay(float DecayAmount)
+{
+	DecayAmount *= MovementSettingsInputsDecayAmountMultiplier;
+
+	MaxSpeed *= (1.0f - DecayAmount);
+	Acceleration *= (1.0f - DecayAmount);
+}
+
 
 //////////////////////////////////////////////////////////////////////////
 // FMoverAIInputs
diff --git a/Engine/Plugins/Experimental/Mover/Source/Mover/Public/Backends/MoverNetworkPhysicsLiaisonBase.h b/Engine/Plugins/Experimental/Mover/Source/Mover/Public/Backends/MoverNetworkPhysicsLiaisonBase.h
--- a/Engine/Plugins/Experimental/Mover/Source/Mover/Public/Backends/MoverNetworkPhysicsLiaisonBase.h
+++ b/Engine/Plugins/Experimental/Mover/Source/Mover/Public/Backends/MoverNetworkPhysicsLiaisonBase.h
@@ -26,6 +26,9 @@ struct MOVER_API FNetworkPhysicsMoverInputs : public FNetworkPhysicsData
 	/**  Build the data from the network physics component */
 	virtual void BuildData(const UActorComponent* NetworkComponent) override;
 
+	/** Decay input during resimulation and forward prediction */
+	virtual void DecayData(float DecayAmount) override;
+
 	/**  Serialize data function that will be used to transfer the struct across the network */
 	bool NetSerialize(FArchive& Ar, class UPackageMap* Map, bool& bOutSuccess);
 
diff --git a/Engine/Plugins/Experimental/Mover/Source/Mover/Public/MoverDataModelTypes.h b/Engine/Plugins/Experimental/Mover/Source/Mover/Public/MoverDataModelTypes.h
--- a/Engine/Plugins/Experimental/Mover/Source/Mover/Public/MoverDataModelTypes.h
+++ b/Engine/Plugins/Experimental/Mover/Source/Mover/Public/MoverDataModelTypes.h
@@ -119,6 +119,7 @@ public:
 	virtual bool ShouldReconcile(const FMoverDataStructBase& AuthorityState) const override;
 	virtual void Interpolate(const FMoverDataStructBase& From, const FMoverDataStructBase& To, float Pct) override;
 	virtual void Merge(const FMoverDataStructBase& From) override;
+	virtual void Decay(float DecayAmount) override;
 };
 
 template<>
diff --git a/Engine/Plugins/Experimental/Mover/Source/Mover/Public/MoverTypes.h b/Engine/Plugins/Experimental/Mover/Source/Mover/Public/MoverTypes.h
--- a/Engine/Plugins/Experimental/Mover/Source/Mover/Public/MoverTypes.h
+++ b/Engine/Plugins/Experimental/Mover/Source/Mover/Public/MoverTypes.h
@@ -162,6 +162,13 @@ struct MOVER_API FMoverDataStructBase
 	 * Override for: Types that compose INPUT data for physics-based movement
 	 */
 	virtual void Merge(const FMoverDataStructBase& From);
+	
+	/**
+	 * Decays contained data during resimulation if data is forward predicted. For more information, see the comments on FNetworkPhysicsData::DecayData.
+	 * @param DecayAmount = Total amount of decay as a multiplier. 10% decay = 0.1.
+	 * Override for: Types that compose INPUT data for physics-based movement
+	 */
+	virtual void Decay(float DecayAmount) {};
 };
 
 template<>
@@ -209,7 +216,10 @@ struct MOVER_API FMoverDataCollection
 
 	/** Merge a previous frame's collection with this collection */
 	void Merge(const FMoverDataCollection& From);
-	
+
+	/** Decay input based on DecayAmount for resimulation and forward prediction */
+	void Decay(float DecayAmount);
+
 	/** Exposes references to GC system */
 	void AddStructReferencedObjects(FReferenceCollector& Collector) const;
 
diff --git a/Engine/Plugins/Experimental/Mover/Source/Mover/Public/PhysicsMover/PhysicsMoverSimulationTypes.h b/Engine/Plugins/Experimental/Mover/Source/Mover/Public/PhysicsMover/PhysicsMoverSimulationTypes.h
--- a/Engine/Plugins/Experimental/Mover/Source/Mover/Public/PhysicsMover/PhysicsMoverSimulationTypes.h
+++ b/Engine/Plugins/Experimental/Mover/Source/Mover/Public/PhysicsMover/PhysicsMoverSimulationTypes.h
@@ -141,6 +141,7 @@ public:
 	virtual bool ShouldReconcile(const FMoverDataStructBase& AuthorityState) const override;
 	virtual void Interpolate(const FMoverDataStructBase& From, const FMoverDataStructBase& To, float Pct) override;
 	virtual void Merge(const FMoverDataStructBase& From) override;
+	virtual void Decay(float DecayAmount) override;
 };
 
 template<>
