Expose mover tick dependencies to blueprint

diff --git a/Engine/Plugins/Experimental/Mover/Source/Mover/Private/Backends/MoverStandaloneLiaison.cpp b/Engine/Plugins/Experimental/Mover/Source/Mover/Private/Backends/MoverStandaloneLiaison.cpp
--- a/Engine/Plugins/Experimental/Mover/Source/Mover/Private/Backends/MoverStandaloneLiaison.cpp
+++ b/Engine/Plugins/Experimental/Mover/Source/Mover/Private/Backends/MoverStandaloneLiaison.cpp
@@ -165,6 +165,29 @@ FTickFunction* UMoverStandaloneLiaisonComponent::FindTickFunction(EMoverTickPhas
 	return FoundFunction;
 }
 
+void UMoverStandaloneLiaisonComponent::AddTickDependency(UActorComponent* OtherComponent, EMoverTickDependencyOrder TickOrder, EMoverTickPhase TickPhase)
+{
+	ensure(OtherComponent);
+	FTickFunction* TickFunction = FindTickFunction(TickPhase);
+	if (OtherComponent && TickFunction)
+	{
+		switch (TickOrder)
+		{
+			case EMoverTickDependencyOrder::After:
+				TickFunction->RemovePrerequisite(OtherComponent, OtherComponent->PrimaryComponentTick);
+				OtherComponent->PrimaryComponentTick.AddPrerequisite(this, *TickFunction);
+				break;
+			case EMoverTickDependencyOrder::Before:
+				OtherComponent->PrimaryComponentTick.RemovePrerequisite(this, *TickFunction);
+				TickFunction->AddPrerequisite(OtherComponent, OtherComponent->PrimaryComponentTick);
+				break;
+			
+			default:
+				break;
+		}
+	}
+}
+
 void UMoverStandaloneLiaisonComponent::SetUseAsyncProduceInput(bool bUseAsyncInputProduction)
 {
 	bUseAsyncProduceInput = bUseAsyncInputProduction;
diff --git a/Engine/Plugins/Experimental/Mover/Source/Mover/Public/Backends/MoverStandaloneLiaison.h b/Engine/Plugins/Experimental/Mover/Source/Mover/Public/Backends/MoverStandaloneLiaison.h
--- a/Engine/Plugins/Experimental/Mover/Source/Mover/Public/Backends/MoverStandaloneLiaison.h
+++ b/Engine/Plugins/Experimental/Mover/Source/Mover/Public/Backends/MoverStandaloneLiaison.h
@@ -114,6 +114,14 @@ public:
 
 	FTickFunction* FindTickFunction(EMoverTickPhase MoverTickPhase);
 
+	/**
+	 * Adds a tick dependency between another component and one of mover's tick functions.
+	 * @param OtherComponent		The component to add a dependency with.
+	 * @param TickOrder				What order OtherComponent's should have relative to TickPhase i.e. OtherComponent Before SimulateMovement.
+	 * @param TickPhase				The mover tick phase we want to a dependency with.
+	 */
+	UFUNCTION(BlueprintCallable, Category = Mover, meta = (Keywords = "Tick Dependency"))
+	void AddTickDependency(UActorComponent* OtherComponent, EMoverTickDependencyOrder TickOrder, EMoverTickPhase TickPhase);
 	
 	/** Sets whether this instance's produce input can run on worker threads or not. See @bUseAsyncProduceInput */
 	UFUNCTION(BlueprintCallable, Category = "Async Movement")
diff --git a/Engine/Plugins/Experimental/Mover/Source/Mover/Public/MoverTypes.h b/Engine/Plugins/Experimental/Mover/Source/Mover/Public/MoverTypes.h
--- a/Engine/Plugins/Experimental/Mover/Source/Mover/Public/MoverTypes.h
+++ b/Engine/Plugins/Experimental/Mover/Source/Mover/Public/MoverTypes.h
@@ -36,6 +36,14 @@ enum class EMoverTickPhase : uint8
 	ApplyState = 3          UMETA(DisplayName = "ApplyState")
 };
 
+/** List of tick dependency order of execution relative to the Mover tick function. */
+UENUM(BlueprintType)
+enum class EMoverTickDependencyOrder : uint8
+{
+	Before = 0,
+	After = 1,
+};
+
 /** Options for how to handle smoothing frame data from the backend. Typically this is for advancing the simulation at a lower or fixed rate versus the game thread/render rate. */
 UENUM(BlueprintType)
 enum class EMoverSmoothingMode : uint8
