[Mover] [Bugfix] Entering swimming when crouched

diff --git a/Engine/Plugins/Experimental/Mover/Source/Mover/Private/DefaultMovementSet/Modes/SwimmingMode.cpp b/Engine/Plugins/Experimental/Mover/Source/Mover/Private/DefaultMovementSet/Modes/SwimmingMode.cpp
--- a/Engine/Plugins/Experimental/Mover/Source/Mover/Private/DefaultMovementSet/Modes/SwimmingMode.cpp
+++ b/Engine/Plugins/Experimental/Mover/Source/Mover/Private/DefaultMovementSet/Modes/SwimmingMode.cpp
@@ -47,12 +47,13 @@ void USwimmingMode::GenerateMove_Implementation(const FMoverTickStartData& Start
 		const FWaterFlowSplineData& WaterData = LastWaterResult.WaterSplineData;
 		
 		const float ReciprocalCapsuleHeight = 0.5f / CapsuleHalfHeight;
+		const float ReciprocalOriginalCapsuleHeight = 0.5f / OriginalCapsuleHalfHeight;
 		const FVector CapsuleLocation = StartingSyncState->GetLocation_WorldSpace();
 
 		const float ImmersionDepth = WaterData.ImmersionDepth + CapsuleHalfHeight;
 		const float ImmersionPercent = FMath::Clamp(ImmersionDepth * ReciprocalCapsuleHeight, UE_KINDA_SMALL_NUMBER, 1.f);
 		const float IdealDepth = CommonLegacySettings->SwimmingIdealImmersionDepth;
-		const float IdealImmersionPercent = FMath::Clamp(IdealDepth * ReciprocalCapsuleHeight, UE_KINDA_SMALL_NUMBER, 1.f);
+		const float IdealImmersionPercent = FMath::Clamp(IdealDepth * ReciprocalOriginalCapsuleHeight, UE_KINDA_SMALL_NUMBER, 1.f);
 
 		const float GravityForce = GetMoverComponent()->GetGravityAcceleration().Z;
 
@@ -170,7 +171,11 @@ void USwimmingMode::OnRegistered(const FName ModeName)
 {
 	Super::OnRegistered(ModeName);
 
-	CommonLegacySettings = GetMoverComponent()->FindSharedSettings<UCommonLegacyMovementSettings>();
+	if (UMoverComponent* MoverComp = GetMoverComponent())
+	{
+		CommonLegacySettings = MoverComp->FindSharedSettings<UCommonLegacyMovementSettings>();
+		OriginalCapsuleHalfHeight = MoverComp->GetOwner()->GetClass()->GetDefaultObject<AActor>()->GetSimpleCollisionHalfHeight();
+	}
 	ensureMsgf(CommonLegacySettings, TEXT("Failed to find instance of CommonLegacyMovementSettings on %s. Movement may not function properly."), *GetPathNameSafe(this));
 }
 
diff --git a/Engine/Plugins/Experimental/Mover/Source/Mover/Public/DefaultMovementSet/Modes/SwimmingMode.h b/Engine/Plugins/Experimental/Mover/Source/Mover/Public/DefaultMovementSet/Modes/SwimmingMode.h
--- a/Engine/Plugins/Experimental/Mover/Source/Mover/Public/DefaultMovementSet/Modes/SwimmingMode.h
+++ b/Engine/Plugins/Experimental/Mover/Source/Mover/Public/DefaultMovementSet/Modes/SwimmingMode.h
@@ -214,4 +214,6 @@ protected:
 	FRelativeBaseInfo UpdateFloorAndBaseInfo(const FFloorCheckResult& FloorResult) const;
 
 	TObjectPtr<const UCommonLegacyMovementSettings> CommonLegacySettings;
+
+	float OriginalCapsuleHalfHeight = 75.f;
 };
