Physics Mover - Fix to avoid losing horizontal velocity on landing

diff --git a/Engine/Plugins/Experimental/Mover/Source/Mover/Private/PhysicsMover/Modes/PhysicsDrivenFallingMode.cpp b/Engine/Plugins/Experimental/Mover/Source/Mover/Private/PhysicsMover/Modes/PhysicsDrivenFallingMode.cpp
--- a/Engine/Plugins/Experimental/Mover/Source/Mover/Private/PhysicsMover/Modes/PhysicsDrivenFallingMode.cpp
+++ b/Engine/Plugins/Experimental/Mover/Source/Mover/Private/PhysicsMover/Modes/PhysicsDrivenFallingMode.cpp
@@ -160,7 +160,7 @@ void UPhysicsDrivenFallingMode::SimulationTick_Implementation(const FSimulationT
 	if (FloorResult.IsWalkableFloor() && bIsFloorWithinReach && !bIsMovingUpRelativeToFloor)
 	{
 		OutputState.MovementEndState.NextModeName = CommonLegacySettings->GroundMovementModeName;
-		TargetPos += UpDir * (TargetHeight - FloorResult.FloorDist);
+		TargetPos += UpDir * (TargetHeight - FloorResult.FloorDist) + (TargetVel - TargetVel.Dot(UpDir) * UpDir) * DeltaSeconds;
 	}
 	else
 	{
