Mover: minor cleanup to consistently use GENERATED_BODY on ustruct types

#rnx

diff --git a/Engine/Plugins/Experimental/Mover/Source/Mover/Public/Backends/MoverNetworkPhysicsLiaisonBase.h b/Engine/Plugins/Experimental/Mover/Source/Mover/Public/Backends/MoverNetworkPhysicsLiaisonBase.h
--- a/Engine/Plugins/Experimental/Mover/Source/Mover/Public/Backends/MoverNetworkPhysicsLiaisonBase.h
+++ b/Engine/Plugins/Experimental/Mover/Source/Mover/Public/Backends/MoverNetworkPhysicsLiaisonBase.h
@@ -15,7 +15,7 @@
 USTRUCT()
 struct MOVER_API FNetworkPhysicsMoverInputs : public FNetworkPhysicsData
 {
-	GENERATED_USTRUCT_BODY()
+	GENERATED_BODY()
 
 	UPROPERTY()
 	FMoverInputCmdContext InputCmdContext;
@@ -59,7 +59,7 @@ struct TStructOpsTypeTraits<FNetworkPhysicsMoverInputs> : public TStructOpsTypeT
 USTRUCT()
 struct MOVER_API FNetworkPhysicsMoverState : public FNetworkPhysicsData
 {
-	GENERATED_USTRUCT_BODY()
+	GENERATED_BODY()
 
 	UPROPERTY()
 	FMoverSyncState SyncStateContext;
diff --git a/Engine/Plugins/Experimental/Mover/Source/Mover/Public/DefaultMovementSet/LayeredMoves/AnimRootMotionLayeredMove.h b/Engine/Plugins/Experimental/Mover/Source/Mover/Public/DefaultMovementSet/LayeredMoves/AnimRootMotionLayeredMove.h
--- a/Engine/Plugins/Experimental/Mover/Source/Mover/Public/DefaultMovementSet/LayeredMoves/AnimRootMotionLayeredMove.h
+++ b/Engine/Plugins/Experimental/Mover/Source/Mover/Public/DefaultMovementSet/LayeredMoves/AnimRootMotionLayeredMove.h
@@ -16,7 +16,7 @@ class UAnimMontage;
 USTRUCT(BlueprintType)
 struct MOVER_API FLayeredMove_AnimRootMotion : public FLayeredMoveBase
 {
-	GENERATED_USTRUCT_BODY()
+	GENERATED_BODY()
 
 	FLayeredMove_AnimRootMotion();
 	virtual ~FLayeredMove_AnimRootMotion() {}
diff --git a/Engine/Plugins/Experimental/Mover/Source/Mover/Public/DefaultMovementSet/LayeredMoves/BasicLayeredMoves.h b/Engine/Plugins/Experimental/Mover/Source/Mover/Public/DefaultMovementSet/LayeredMoves/BasicLayeredMoves.h
--- a/Engine/Plugins/Experimental/Mover/Source/Mover/Public/DefaultMovementSet/LayeredMoves/BasicLayeredMoves.h
+++ b/Engine/Plugins/Experimental/Mover/Source/Mover/Public/DefaultMovementSet/LayeredMoves/BasicLayeredMoves.h
@@ -25,7 +25,7 @@ enum class ELayeredMove_ConstantVelocitySettingsFlags : uint8
 USTRUCT(BlueprintType)
 struct MOVER_API FLayeredMove_LinearVelocity : public FLayeredMoveBase
 {
-	GENERATED_USTRUCT_BODY()
+	GENERATED_BODY()
 
 	FLayeredMove_LinearVelocity();
 
@@ -75,7 +75,7 @@ struct TStructOpsTypeTraits< FLayeredMove_LinearVelocity > : public TStructOpsTy
 USTRUCT(BlueprintType)
 struct MOVER_API FLayeredMove_JumpImpulseOverDuration : public FLayeredMoveBase
 {
-	GENERATED_USTRUCT_BODY()
+	GENERATED_BODY()
 
 	FLayeredMove_JumpImpulseOverDuration();
 
@@ -115,7 +115,7 @@ struct TStructOpsTypeTraits< FLayeredMove_JumpImpulseOverDuration > : public TSt
 USTRUCT(BlueprintType)
 struct MOVER_API FLayeredMove_JumpTo : public FLayeredMoveBase
 {
-	GENERATED_USTRUCT_BODY()
+	GENERATED_BODY()
 
 	FLayeredMove_JumpTo();
 	virtual ~FLayeredMove_JumpTo() {}
@@ -181,7 +181,7 @@ struct TStructOpsTypeTraits< FLayeredMove_JumpTo > : public TStructOpsTypeTraits
 USTRUCT(BlueprintType)
 struct MOVER_API FLayeredMove_MoveTo : public FLayeredMoveBase
 {
-	GENERATED_USTRUCT_BODY()
+	GENERATED_BODY()
 
 	FLayeredMove_MoveTo();
 	virtual ~FLayeredMove_MoveTo() {}
@@ -244,7 +244,7 @@ struct TStructOpsTypeTraits< FLayeredMove_MoveTo > : public TStructOpsTypeTraits
 USTRUCT(BlueprintType)
 struct MOVER_API FLayeredMove_MoveToDynamic : public FLayeredMove_MoveTo
 {
-	GENERATED_USTRUCT_BODY()
+	GENERATED_BODY()
 
 	FLayeredMove_MoveToDynamic();
 	virtual ~FLayeredMove_MoveToDynamic() {}
@@ -281,7 +281,7 @@ struct TStructOpsTypeTraits< FLayeredMove_MoveToDynamic > : public TStructOpsTyp
 USTRUCT(BlueprintType)
 struct MOVER_API FLayeredMove_RadialImpulse : public FLayeredMoveBase
 {
-	GENERATED_USTRUCT_BODY()
+	GENERATED_BODY()
 
 	FLayeredMove_RadialImpulse();
 	virtual ~FLayeredMove_RadialImpulse() {}
diff --git a/Engine/Plugins/Experimental/Mover/Source/Mover/Public/DefaultMovementSet/LayeredMoves/LaunchMove.h b/Engine/Plugins/Experimental/Mover/Source/Mover/Public/DefaultMovementSet/LayeredMoves/LaunchMove.h
--- a/Engine/Plugins/Experimental/Mover/Source/Mover/Public/DefaultMovementSet/LayeredMoves/LaunchMove.h
+++ b/Engine/Plugins/Experimental/Mover/Source/Mover/Public/DefaultMovementSet/LayeredMoves/LaunchMove.h
@@ -10,7 +10,7 @@
 USTRUCT(BlueprintType)
 struct MOVER_API FLayeredMove_Launch : public FLayeredMoveBase
 {
-	GENERATED_USTRUCT_BODY()
+	GENERATED_BODY()
 
 	FLayeredMove_Launch();
 	virtual ~FLayeredMove_Launch() {}
diff --git a/Engine/Plugins/Experimental/Mover/Source/Mover/Public/InstantMovementEffect.h b/Engine/Plugins/Experimental/Mover/Source/Mover/Public/InstantMovementEffect.h
--- a/Engine/Plugins/Experimental/Mover/Source/Mover/Public/InstantMovementEffect.h
+++ b/Engine/Plugins/Experimental/Mover/Source/Mover/Public/InstantMovementEffect.h
@@ -32,7 +32,7 @@ struct FApplyMovementEffectParams
 USTRUCT(BlueprintInternalUseOnly)
 struct MOVER_API FInstantMovementEffect
 {
-	GENERATED_USTRUCT_BODY()
+	GENERATED_BODY()
 
 	FInstantMovementEffect() { }
 
diff --git a/Engine/Plugins/Experimental/Mover/Source/Mover/Public/LayeredMove.h b/Engine/Plugins/Experimental/Mover/Source/Mover/Public/LayeredMove.h
--- a/Engine/Plugins/Experimental/Mover/Source/Mover/Public/LayeredMove.h
+++ b/Engine/Plugins/Experimental/Mover/Source/Mover/Public/LayeredMove.h
@@ -30,7 +30,7 @@ enum class ELayeredMoveFinishVelocityMode : uint8
 USTRUCT(BlueprintType)
 struct FLayeredMoveFinishVelocitySettings
 {
-	GENERATED_USTRUCT_BODY()
+	GENERATED_BODY()
 
 	FLayeredMoveFinishVelocitySettings()
 		: SetVelocity(FVector::ZeroVector)
@@ -69,7 +69,7 @@ struct FLayeredMoveFinishVelocitySettings
 USTRUCT(BlueprintInternalUseOnly)
 struct MOVER_API FLayeredMoveBase
 {
-	GENERATED_USTRUCT_BODY()
+	GENERATED_BODY()
 
 	FLayeredMoveBase();
 	virtual ~FLayeredMoveBase() {}
@@ -156,7 +156,7 @@ struct TStructOpsTypeTraits< FLayeredMoveBase > : public TStructOpsTypeTraitsBas
 USTRUCT(BlueprintType)
 struct MOVER_API FLayeredMoveGroup
 {
-	GENERATED_USTRUCT_BODY()
+	GENERATED_BODY()
 
 	FLayeredMoveGroup();
 	
diff --git a/Engine/Plugins/Experimental/Mover/Source/Mover/Public/MoveLibrary/BasedMovementUtils.h b/Engine/Plugins/Experimental/Mover/Source/Mover/Public/MoveLibrary/BasedMovementUtils.h
--- a/Engine/Plugins/Experimental/Mover/Source/Mover/Public/MoveLibrary/BasedMovementUtils.h
+++ b/Engine/Plugins/Experimental/Mover/Source/Mover/Public/MoveLibrary/BasedMovementUtils.h
@@ -14,7 +14,7 @@ struct FFloorCheckResult;
 USTRUCT(BlueprintType)
 struct MOVER_API FRelativeBaseInfo
 {
-	GENERATED_USTRUCT_BODY()
+	GENERATED_BODY()
 
 	/** Component we are moving relative to */
 	UPROPERTY(VisibleInstanceOnly, BlueprintReadOnly, Category="Mover")
@@ -137,7 +137,7 @@ public:
 USTRUCT()
 struct FMoverDynamicBasedMovementTickFunction : public FTickFunction
 {
-	GENERATED_USTRUCT_BODY()
+	GENERATED_BODY()
 
 	/** MoverComponent that is the target of this tick **/
 	UMoverComponent* TargetMoverComp;
diff --git a/Engine/Plugins/Experimental/Mover/Source/Mover/Public/MoveLibrary/ConstrainedMoveUtils.h b/Engine/Plugins/Experimental/Mover/Source/Mover/Public/MoveLibrary/ConstrainedMoveUtils.h
--- a/Engine/Plugins/Experimental/Mover/Source/Mover/Public/MoveLibrary/ConstrainedMoveUtils.h
+++ b/Engine/Plugins/Experimental/Mover/Source/Mover/Public/MoveLibrary/ConstrainedMoveUtils.h
@@ -12,7 +12,7 @@
 USTRUCT(BlueprintType)
 struct MOVER_API FPlanarConstraint
 {
-	GENERATED_USTRUCT_BODY()
+	GENERATED_BODY()
 
 	/** If true, movement will be constrained to a plane. */
 	UPROPERTY(EditAnywhere, BlueprintReadOnly, Category = PlanarMovement)
diff --git a/Engine/Plugins/Experimental/Mover/Source/Mover/Public/MoveLibrary/FloorQueryUtils.h b/Engine/Plugins/Experimental/Mover/Source/Mover/Public/MoveLibrary/FloorQueryUtils.h
--- a/Engine/Plugins/Experimental/Mover/Source/Mover/Public/MoveLibrary/FloorQueryUtils.h
+++ b/Engine/Plugins/Experimental/Mover/Source/Mover/Public/MoveLibrary/FloorQueryUtils.h
@@ -19,7 +19,7 @@ namespace UE::FloorQueryUtility
 USTRUCT(BlueprintType)
 struct MOVER_API FFloorCheckResult
 {
-	GENERATED_USTRUCT_BODY()
+	GENERATED_BODY()
 
 	/**
 	* True if there was a blocking hit in the floor test that was NOT in initial penetration.
diff --git a/Engine/Plugins/Experimental/Mover/Source/Mover/Public/MoveLibrary/MovementUtils.h b/Engine/Plugins/Experimental/Mover/Source/Mover/Public/MoveLibrary/MovementUtils.h
--- a/Engine/Plugins/Experimental/Mover/Source/Mover/Public/MoveLibrary/MovementUtils.h
+++ b/Engine/Plugins/Experimental/Mover/Source/Mover/Public/MoveLibrary/MovementUtils.h
@@ -25,7 +25,7 @@ namespace UE::MoverUtils
 USTRUCT(BlueprintType)
 struct MOVER_API FTrajectorySampleInfo
 {
-	GENERATED_USTRUCT_BODY()
+	GENERATED_BODY()
 
 	FTransform Transform;					// Position and orientation (world space)
 
diff --git a/Engine/Plugins/Experimental/Mover/Source/Mover/Public/MoveLibrary/WaterMovementUtils.h b/Engine/Plugins/Experimental/Mover/Source/Mover/Public/MoveLibrary/WaterMovementUtils.h
--- a/Engine/Plugins/Experimental/Mover/Source/Mover/Public/MoveLibrary/WaterMovementUtils.h
+++ b/Engine/Plugins/Experimental/Mover/Source/Mover/Public/MoveLibrary/WaterMovementUtils.h
@@ -67,7 +67,7 @@ struct MOVER_API FWaterMoveParams
 USTRUCT(BlueprintType)
 struct MOVER_API FWaterFlowSplineData
 {
-	GENERATED_USTRUCT_BODY()
+	GENERATED_BODY()
 	
 	FWaterFlowSplineData()
 	{
@@ -118,7 +118,7 @@ struct MOVER_API FWaterFlowSplineData
 USTRUCT(BlueprintType)
 struct MOVER_API FWaterCheckResult
 {
-	GENERATED_USTRUCT_BODY()
+	GENERATED_BODY()
 	
 	/** True if the hit found a valid swimmable volume. */
 	UPROPERTY(VisibleInstanceOnly, BlueprintReadOnly, Category = CharacterVolume)
diff --git a/Engine/Plugins/Experimental/Mover/Source/Mover/Public/MoverDataModelTypes.h b/Engine/Plugins/Experimental/Mover/Source/Mover/Public/MoverDataModelTypes.h
--- a/Engine/Plugins/Experimental/Mover/Source/Mover/Public/MoverDataModelTypes.h
+++ b/Engine/Plugins/Experimental/Mover/Source/Mover/Public/MoverDataModelTypes.h
@@ -28,7 +28,7 @@ enum class EMoveInputType : uint8
 USTRUCT(BlueprintType)
 struct MOVER_API FCharacterDefaultInputs : public FMoverDataStructBase
 {
-	GENERATED_USTRUCT_BODY()
+	GENERATED_BODY()
 
 	// Sets the directional move inputs for a simulation frame
 	void SetMoveInput(EMoveInputType InMoveInputType, const FVector& InMoveInput);
@@ -139,7 +139,7 @@ struct TStructOpsTypeTraits< FCharacterDefaultInputs > : public TStructOpsTypeTr
 USTRUCT(BlueprintType)
 struct MOVER_API FMoverDefaultSyncState : public FMoverDataStructBase
 {
-	GENERATED_USTRUCT_BODY()
+	GENERATED_BODY()
 protected:
 	// Position relative to MovementBase if set, world space otherwise
 	UPROPERTY(BlueprintReadWrite, EditAnywhere, Category = Mover)
diff --git a/Engine/Plugins/Experimental/Mover/Source/Mover/Public/MoverTypes.h b/Engine/Plugins/Experimental/Mover/Source/Mover/Public/MoverTypes.h
--- a/Engine/Plugins/Experimental/Mover/Source/Mover/Public/MoverTypes.h
+++ b/Engine/Plugins/Experimental/Mover/Source/Mover/Public/MoverTypes.h
@@ -127,7 +127,7 @@ struct MOVER_API FMoverTimeStep
 USTRUCT(BlueprintType)
 struct MOVER_API FMoverDataStructBase
 {
-	GENERATED_USTRUCT_BODY()
+	GENERATED_BODY()
 
 	FMoverDataStructBase();
 	virtual ~FMoverDataStructBase() {}
@@ -198,7 +198,7 @@ struct TStructOpsTypeTraits< FMoverDataStructBase > : public TStructOpsTypeTrait
 USTRUCT(BlueprintType)
 struct MOVER_API FMoverDataCollection
 {
-	GENERATED_USTRUCT_BODY()
+	GENERATED_BODY()
 
 	FMoverDataCollection();
 
diff --git a/Engine/Plugins/Experimental/Mover/Source/Mover/Public/PhysicsMover/PhysicsMoverSimulationTypes.h b/Engine/Plugins/Experimental/Mover/Source/Mover/Public/PhysicsMover/PhysicsMoverSimulationTypes.h
--- a/Engine/Plugins/Experimental/Mover/Source/Mover/Public/PhysicsMover/PhysicsMoverSimulationTypes.h
+++ b/Engine/Plugins/Experimental/Mover/Source/Mover/Public/PhysicsMover/PhysicsMoverSimulationTypes.h
@@ -114,7 +114,7 @@ namespace PhysicsMovementModeUtils
 USTRUCT(BlueprintType)
 struct MOVER_API FMovementSettingsInputs : public FMoverDataStructBase
 {
-	GENERATED_USTRUCT_BODY()
+	GENERATED_BODY()
 
 public:
 	// Maximum speed in cm/s
@@ -162,7 +162,7 @@ struct TStructOpsTypeTraits< FMovementSettingsInputs > : public TStructOpsTypeTr
 USTRUCT(BlueprintType)
 struct MOVER_API FMoverAIInputs : public FMoverDataStructBase
 {
-	GENERATED_USTRUCT_BODY()
+	GENERATED_BODY()
 
 public:
 	// ROV Velocity calculated on the Server
