[Mover] [BugFix] ColliderOffset calculation for PhysicsStanceModifier

diff --git a/Engine/Plugins/Experimental/Mover/Source/Mover/Public/PhysicsMover/Modes/PhysicsDrivenFallingMode.h b/Engine/Plugins/Experimental/Mover/Source/Mover/Public/PhysicsMover/Modes/PhysicsDrivenFallingMode.h
--- a/Engine/Plugins/Experimental/Mover/Source/Mover/Public/PhysicsMover/Modes/PhysicsDrivenFallingMode.h
+++ b/Engine/Plugins/Experimental/Mover/Source/Mover/Public/PhysicsMover/Modes/PhysicsDrivenFallingMode.h
@@ -28,9 +28,11 @@ public:
 	virtual EDataValidationResult IsDataValid(class FDataValidationContext& Context) const override;
 #endif // WITH_EDITOR
 
-	void SetTargetHeightOverride(float InTargetHeight);
-	void ClearTargetHeightOverride();
-	
+	virtual float GetTargetHeight() const override { return TargetHeight; }
+
+	virtual void SetTargetHeightOverride(float InTargetHeight) override;
+	virtual void ClearTargetHeightOverride() override;
+
 protected:
 	// Maximum torque the character can apply to rotate in air about the vertical axis
 	UPROPERTY(EditAnywhere, Category = "Physics Mover", meta = (ClampMin = "0", UIMin = "0", ForceUnits = "NewtonMeters"))
diff --git a/Engine/Plugins/Experimental/Mover/Source/Mover/Public/PhysicsMover/Modes/PhysicsDrivenSwimmingMode.h b/Engine/Plugins/Experimental/Mover/Source/Mover/Public/PhysicsMover/Modes/PhysicsDrivenSwimmingMode.h
--- a/Engine/Plugins/Experimental/Mover/Source/Mover/Public/PhysicsMover/Modes/PhysicsDrivenSwimmingMode.h
+++ b/Engine/Plugins/Experimental/Mover/Source/Mover/Public/PhysicsMover/Modes/PhysicsDrivenSwimmingMode.h
@@ -30,9 +30,11 @@ public:
 	virtual EDataValidationResult IsDataValid(class FDataValidationContext& Context) const override;
 #endif // WITH_EDITOR
 
-	void SetTargetHeightOverride(float InTargetHeight);
-	void ClearTargetHeightOverride();
-	
+	virtual float GetTargetHeight() const override { return TargetHeight; }
+
+	virtual void SetTargetHeightOverride(float InTargetHeight) override;
+	virtual void ClearTargetHeightOverride() override;
+
 protected:
 	/**
 	 * Optional override target height for the character (the desired distance from the center of the capsule to the floor).
diff --git a/Engine/Plugins/Experimental/Mover/Source/Mover/Public/PhysicsMover/Modes/PhysicsDrivenWalkingMode.h b/Engine/Plugins/Experimental/Mover/Source/Mover/Public/PhysicsMover/Modes/PhysicsDrivenWalkingMode.h
--- a/Engine/Plugins/Experimental/Mover/Source/Mover/Public/PhysicsMover/Modes/PhysicsDrivenWalkingMode.h
+++ b/Engine/Plugins/Experimental/Mover/Source/Mover/Public/PhysicsMover/Modes/PhysicsDrivenWalkingMode.h
@@ -31,8 +31,10 @@ public:
 	virtual EDataValidationResult IsDataValid(class FDataValidationContext& Context) const override;
 #endif // WITH_EDITOR
 
-	void SetTargetHeightOverride(float InTargetHeight);
-	void ClearTargetHeightOverride();
+	virtual float GetTargetHeight() const override { return TargetHeight; }
+
+	virtual void SetTargetHeightOverride(float InTargetHeight) override;
+	virtual void ClearTargetHeightOverride() override;
 	
 protected:
 	// Maximum force the character can apply to reach the motion target
diff --git a/Engine/Plugins/Experimental/Mover/Source/Mover/Public/PhysicsMover/MovementModifiers/PhysicsStanceModifier.h b/Engine/Plugins/Experimental/Mover/Source/Mover/Public/PhysicsMover/MovementModifiers/PhysicsStanceModifier.h
--- a/Engine/Plugins/Experimental/Mover/Source/Mover/Public/PhysicsMover/MovementModifiers/PhysicsStanceModifier.h
+++ b/Engine/Plugins/Experimental/Mover/Source/Mover/Public/PhysicsMover/MovementModifiers/PhysicsStanceModifier.h
@@ -5,6 +5,7 @@
 #include "MoverComponent.h"
 #include "DefaultMovementSet/MovementModifiers/StanceModifier.h"
 #include "Components/SceneComponent.h"
+#include "PhysicsMover/PhysicsMoverSimulationTypes.h"
 
 #include "PhysicsStanceModifier.generated.h"
 
@@ -44,8 +45,11 @@ protected:
 	{
 		if (T* MovementMode = MoverComp->FindMode_Mutable<T>())
 		{
-			const float ColliderOffset = -MoverComp->GetPrimaryVisualComponent()->GetRelativeLocation().Z - OldHalfHeight;
-			MovementMode->SetTargetHeightOverride(NewHalfHeight + ColliderOffset);
+			if (MovementMode->template Implements<UPhysicsCharacterMovementModeInterface>())
+			{
+				float TargetHeight = MovementMode->GetTargetHeight();
+				MovementMode->SetTargetHeightOverride(TargetHeight + (NewHalfHeight - OldHalfHeight));
+			}
 		}
 	}
 };
\ No newline at end of file
diff --git a/Engine/Plugins/Experimental/Mover/Source/Mover/Public/PhysicsMover/PhysicsMoverSimulationTypes.h b/Engine/Plugins/Experimental/Mover/Source/Mover/Public/PhysicsMover/PhysicsMoverSimulationTypes.h
--- a/Engine/Plugins/Experimental/Mover/Source/Mover/Public/PhysicsMover/PhysicsMoverSimulationTypes.h
+++ b/Engine/Plugins/Experimental/Mover/Source/Mover/Public/PhysicsMover/PhysicsMoverSimulationTypes.h
@@ -97,6 +97,13 @@ public:
 
 	// Optionally run contact modification on the physics thread
 	virtual void OnContactModification_Internal(const FPhysicsMoverSimulationContactModifierParams& Params, Chaos::FCollisionContactModifier& Modifier) const {};
+
+	// Optional TargetHeight getter
+	virtual float GetTargetHeight() const { return -1.f; }
+
+	// Optional TargetHeightOverride methods
+	virtual void SetTargetHeightOverride(float InTargetHeight) {}
+	virtual void ClearTargetHeightOverride() {}
 };
 
 class FDataValidationContext;
