Fix for a build error when WITH_CHAOS_VISUAL_DEBUGGER is not defined

diff --git a/Engine/Plugins/Experimental/Mover/Source/Mover/Public/ChaosVisualDebugger/MoverCVDRuntimeTrace.h b/Engine/Plugins/Experimental/Mover/Source/Mover/Public/ChaosVisualDebugger/MoverCVDRuntimeTrace.h
--- a/Engine/Plugins/Experimental/Mover/Source/Mover/Public/ChaosVisualDebugger/MoverCVDRuntimeTrace.h
+++ b/Engine/Plugins/Experimental/Mover/Source/Mover/Public/ChaosVisualDebugger/MoverCVDRuntimeTrace.h
@@ -2,8 +2,6 @@
 
 #pragma once
 
-#if WITH_CHAOS_VISUAL_DEBUGGER
-
 #include "Containers/Array.h"
 #include "HAL/Platform.h"
 
@@ -14,6 +12,8 @@ struct FMoverInputCmdContext;
 struct FMoverSyncState;
 class UMoverComponent;
 
+#if WITH_CHAOS_VISUAL_DEBUGGER
+
 /** Utility functions used to trace Mover data into the Chaos Visual Debugger */
 class FMoverCVDRuntimeTrace
 {
@@ -36,6 +36,8 @@ class FMoverCVDRuntimeTrace
 {
 public:
 	static void TraceMoverData(UMoverComponent* MoverComponent, const FMoverInputCmdContext* InputCmd, const FMoverSyncState* SyncState) {}
+	MOVER_API static void UnwrapSimData(const FMoverCVDSimDataWrapper& InSimDataWrapper, TSharedPtr<FMoverInputCmdContext>& OutInputCmd, TSharedPtr<FMoverSyncState>& OutSyncState) {}
+	static void WrapSimData(uint32 SolverID, uint32 ParticleID, const FMoverInputCmdContext& InInputCmd, const FMoverSyncState& InSyncState, FMoverCVDSimDataWrapper& OutSimDataWrapper) {}
 };
 
 #endif // WITH_CHAOS_VISUAL_DEBUGGER
