Physics Mover - Fix movement code to deal with offset centres of mass on the ground object

diff --git a/Engine/Plugins/Experimental/Mover/Source/Mover/Private/PhysicsMover/PhysicsMovementUtils.cpp b/Engine/Plugins/Experimental/Mover/Source/Mover/Private/PhysicsMover/PhysicsMovementUtils.cpp
--- a/Engine/Plugins/Experimental/Mover/Source/Mover/Private/PhysicsMover/PhysicsMovementUtils.cpp
+++ b/Engine/Plugins/Experimental/Mover/Source/Mover/Private/PhysicsMover/PhysicsMovementUtils.cpp
@@ -144,7 +144,8 @@ FVector UPhysicsMovementUtils::ComputeGroundVelocityFromHitResult(const FVector&
 	FVector GroundVelocity = FVector::ZeroVector;
 	if (const Chaos::FPBDRigidParticleHandle* Rigid = GetRigidParticleHandleFromHitResult(FloorHit))
 	{
-		FVector Offset = CharacterPosition - Rigid->GetX();
+		Chaos::FRigidTransform3 ComTransform = Rigid->GetTransformXRCom();
+		FVector Offset = CharacterPosition - ComTransform.GetLocation();
 		Offset -= Offset.ProjectOnToNormal(FloorHit.ImpactNormal);
 		GroundVelocity = Rigid->GetV() + Rigid->GetW().Cross(Offset);
 	}
@@ -156,7 +157,8 @@ FVector UPhysicsMovementUtils::ComputeIntegratedGroundVelocityFromHitResult(cons
 	FVector GroundVelocity = FVector::ZeroVector;
 	if (const Chaos::FPBDRigidParticleHandle* Rigid = GetRigidParticleHandleFromHitResult(FloorHit))
 	{
-		FVector Offset = CharacterPosition - Rigid->GetX();
+		Chaos::FRigidTransform3 ComTransform = Rigid->GetTransformXRCom();
+		FVector Offset = CharacterPosition - ComTransform.GetLocation();
 		Offset -= Offset.ProjectOnToNormal(FloorHit.ImpactNormal);
 		GroundVelocity = Rigid->GetV() + Rigid->GetW().Cross(Offset);
 
diff --git a/Engine/Plugins/Experimental/Mover/Source/Mover/Private/PhysicsMover/Transitions/PhysicsJumpCheck.cpp b/Engine/Plugins/Experimental/Mover/Source/Mover/Private/PhysicsMover/Transitions/PhysicsJumpCheck.cpp
--- a/Engine/Plugins/Experimental/Mover/Source/Mover/Private/PhysicsMover/Transitions/PhysicsJumpCheck.cpp
+++ b/Engine/Plugins/Experimental/Mover/Source/Mover/Private/PhysicsMover/Transitions/PhysicsJumpCheck.cpp
@@ -18,7 +18,8 @@ namespace Private
 {
 	static void ApplyImpulse(Chaos::FPBDRigidParticleHandle* Particle, const Chaos::FVec3& Impulse, const Chaos::FVec3& Location)
 	{
-		const Chaos::FVec3 Offset = Location - Particle->GetX();
+		Chaos::FRigidTransform3 ComTransform = Particle->GetTransformXRCom();
+		const Chaos::FVec3 Offset = Location - ComTransform.GetLocation();
 		Particle->SetW(Particle->GetW() + Particle->InvI() * Offset.Cross(Impulse));
 		Particle->SetV(Particle->GetV() + Particle->InvM() * Impulse);
 	}
