Fix for a non unity build compile error, missing include

diff --git a/Engine/Plugins/Experimental/Mover/Source/Mover/Private/ChaosVisualDebugger/MoverCVDRuntimeTrace.cpp b/Engine/Plugins/Experimental/Mover/Source/Mover/Private/ChaosVisualDebugger/MoverCVDRuntimeTrace.cpp
--- a/Engine/Plugins/Experimental/Mover/Source/Mover/Private/ChaosVisualDebugger/MoverCVDRuntimeTrace.cpp
+++ b/Engine/Plugins/Experimental/Mover/Source/Mover/Private/ChaosVisualDebugger/MoverCVDRuntimeTrace.cpp
@@ -17,6 +17,7 @@
 #include "PBDRigidsSolver.h"
 #include "Physics/Experimental/PhysScene_Chaos.h"
 #include "Interfaces/IPhysicsComponent.h"
+#include "Engine/World.h"
 
 CVD_DEFINE_OPTIONAL_DATA_CHANNEL(MoverData, EChaosVDDataChannelInitializationFlags::CanChangeEnabledState)
 
