Build fix, missing another include

diff --git a/Engine/Plugins/Experimental/Mover/Source/Mover/Public/ChaosVisualDebugger/MoverCVDRuntimeTrace.h b/Engine/Plugins/Experimental/Mover/Source/Mover/Public/ChaosVisualDebugger/MoverCVDRuntimeTrace.h
--- a/Engine/Plugins/Experimental/Mover/Source/Mover/Public/ChaosVisualDebugger/MoverCVDRuntimeTrace.h
+++ b/Engine/Plugins/Experimental/Mover/Source/Mover/Public/ChaosVisualDebugger/MoverCVDRuntimeTrace.h
@@ -4,6 +4,7 @@
 
 #include "Containers/Array.h"
 #include "HAL/Platform.h"
+#include "Templates/SharedPointer.h"
 
 #include "ChaosVisualDebugger/ChaosVDOptionalDataChannel.h"
 
