Added option to disable ProduceInput tick function altogether for standalone Mover backend

diff --git a/Engine/Plugins/Experimental/Mover/Source/Mover/Private/Backends/MoverStandaloneLiaison.cpp b/Engine/Plugins/Experimental/Mover/Source/Mover/Private/Backends/MoverStandaloneLiaison.cpp
--- a/Engine/Plugins/Experimental/Mover/Source/Mover/Private/Backends/MoverStandaloneLiaison.cpp
+++ b/Engine/Plugins/Experimental/Mover/Source/Mover/Private/Backends/MoverStandaloneLiaison.cpp
@@ -112,8 +112,7 @@ void UMoverStandaloneLiaisonComponent::BeginPlay()
 {
 	Super::BeginPlay();
 
-	CurrentSimTimeMs = GetWorld()->GetTimeSeconds() * 1000.0;
-	CurrentSimFrame = GFrameCounter;
+	UpdateSimulationTime();
 
 	if (const AActor* OwnerActor = GetOwner())
 	{
@@ -199,6 +198,30 @@ bool UMoverStandaloneLiaisonComponent::GetUseAsyncProduceInput() const
 	return bUseAsyncProduceInput;
 }
 
+void UMoverStandaloneLiaisonComponent::SetEnableProduceInput(bool bEnableInputProduction)
+{
+	ProduceInputTickFunction.SetTickFunctionEnable(bEnableInputProduction);
+	ProduceInputTickFunction.bStartWithTickEnabled = bEnableInputProduction;
+
+	if (HasBegunPlay())
+	{
+		if (!bEnableInputProduction && ProduceInputTickFunction.IsTickFunctionRegistered())
+		{
+			ProduceInputTickFunction.UnRegisterTickFunction();
+		}
+		else if (bEnableInputProduction && !ProduceInputTickFunction.IsTickFunctionRegistered())
+		{
+			const AActor* Owner = GetOwner();
+			ULevel* ComponentLevel = (Owner ? Owner->GetLevel() : ToRawPtr(GetWorld()->PersistentLevel));
+			ProduceInputTickFunction.RegisterTickFunction(ComponentLevel);
+		}
+	}
+}
+
+bool UMoverStandaloneLiaisonComponent::GetEnableProduceInput() const
+{
+	return ProduceInputTickFunction.IsTickFunctionEnabled();
+}
 
 void UMoverStandaloneLiaisonComponent::SetUseAsyncMovementSimulationTick(bool bUseAsyncMovementSim)
 {
@@ -282,12 +305,17 @@ void UMoverStandaloneLiaisonComponent::RegisterComponentTickFunctions(bool bRegi
 }
 
 
-void UMoverStandaloneLiaisonComponent::TickInputProduction(float DeltaSeconds)
+void UMoverStandaloneLiaisonComponent::UpdateSimulationTime()
 {
-	const int32 DeltaTimeMs = DeltaSeconds * 1000.f;
-
 	CurrentSimTimeMs = GetWorld()->GetTimeSeconds() * 1000.0;
 	CurrentSimFrame = GFrameCounter;
+}
+
+void UMoverStandaloneLiaisonComponent::TickInputProduction(float DeltaSeconds)
+{
+	UpdateSimulationTime();
+
+	const int32 DeltaTimeMs = DeltaSeconds * 1000.f;
 
 	FMoverInputCmdContext InputCmd;
 	{
@@ -298,6 +326,12 @@ void UMoverStandaloneLiaisonComponent::TickInputProduction(float DeltaSeconds)
 
 void UMoverStandaloneLiaisonComponent::TickMovementSimulation(float DeltaSeconds)
 {
+	if (!GetEnableProduceInput())
+	{
+		// If input production is turned off, we need to update current time ourselves
+		UpdateSimulationTime();
+	}
+
 	FMoverTimeStep TimeStep;
 	TimeStep.ServerFrame = CurrentSimFrame;
 	TimeStep.BaseSimTimeMs = CurrentSimTimeMs;
diff --git a/Engine/Plugins/Experimental/Mover/Source/Mover/Public/Backends/MoverStandaloneLiaison.h b/Engine/Plugins/Experimental/Mover/Source/Mover/Public/Backends/MoverStandaloneLiaison.h
--- a/Engine/Plugins/Experimental/Mover/Source/Mover/Public/Backends/MoverStandaloneLiaison.h
+++ b/Engine/Plugins/Experimental/Mover/Source/Mover/Public/Backends/MoverStandaloneLiaison.h
@@ -120,21 +120,29 @@ public:
 	 * @param TickOrder				What order OtherComponent's should have relative to TickPhase i.e. OtherComponent Before SimulateMovement.
 	 * @param TickPhase				The mover tick phase we want to a dependency with.
 	 */
-	UFUNCTION(BlueprintCallable, Category = Mover, meta = (Keywords = "Tick Dependency"))
+	UFUNCTION(BlueprintCallable, Category = "Ticking", meta = (Keywords = "Tick Dependency"))
 	void AddTickDependency(UActorComponent* OtherComponent, EMoverTickDependencyOrder TickOrder, EMoverTickPhase TickPhase);
 	
-	/** Sets whether this instance's produce input can run on worker threads or not. See @bUseAsyncProduceInput */
-	UFUNCTION(BlueprintCallable, Category = "Async Movement")
+	/** Sets whether this instance's produce input can run on worker threads or not. See @bUseAsyncProduceInput and @SetEnableProduceInput */
+	UFUNCTION(BlueprintCallable, Category = "Ticking")
 	void SetUseAsyncProduceInput(bool bUseAsyncInputProduction);
 
-	UFUNCTION(BlueprintCallable, Category = "Async Movement")
+	UFUNCTION(BlueprintCallable, Category = "Ticking")
 	bool GetUseAsyncProduceInput() const;
+
+	/** Sets whether this instance's produce-input tick will run at all. It may be useful to disable on actors that don't rely on Mover input to move.  */
+	UFUNCTION(BlueprintCallable, Category = "Ticking")
+	void SetEnableProduceInput(bool bEnableInputProduction);
+	
+	/** Whether this instance will have its produce-input tick called. */
+	UFUNCTION(BlueprintCallable, Category = "Ticking")
+	bool GetEnableProduceInput() const;
 	
 	/** Sets whether this instance's movement simulation tick can run on worker threads or not. See @bUseAsyncMovementSimulationTick */
-	UFUNCTION(BlueprintCallable, Category = "Async Movement")
+	UFUNCTION(BlueprintCallable, Category = "Ticking")
 	void SetUseAsyncMovementSimulationTick(bool bUseAsyncMovementSim);
 
-	UFUNCTION(BlueprintCallable, Category = "Async Movement")
+	UFUNCTION(BlueprintCallable, Category = "Ticking")
 	bool GetUseAsyncMovementSimulationTick() const;
 
 protected:
@@ -142,6 +150,8 @@ protected:
 	void TickMovementSimulation(float DeltaSeconds);
 	void TickApplySimulationState(float DeltaSeconds);
 
+	void UpdateSimulationTime();
+
 	// Called when controller changes, used to manage ticking dependencies
 	UFUNCTION()
 	virtual void OnControllerChanged(APawn* Pawn, AController* OldController, AController* NewController);
