Mover: Removing FMoverDefaultSyncState data collection add from FMoverSyncState constructor. Also adding a workaround for the NPP liaison to make sure start state has valid persistent sync states.

diff --git a/Engine/Plugins/Experimental/Mover/Source/Mover/Private/Backends/MoverNetworkPredictionLiaison.cpp b/Engine/Plugins/Experimental/Mover/Source/Mover/Private/Backends/MoverNetworkPredictionLiaison.cpp
--- a/Engine/Plugins/Experimental/Mover/Source/Mover/Private/Backends/MoverNetworkPredictionLiaison.cpp
+++ b/Engine/Plugins/Experimental/Mover/Source/Mover/Private/Backends/MoverNetworkPredictionLiaison.cpp
@@ -97,6 +97,12 @@ void UMoverNetworkPredictionLiaisonComponent::SimulationTick(const FNetSimTimeSt
 	StartData.SyncState = *SimInput.Sync;
 	StartData.AuxState  = *SimInput.Aux;
 
+	// Ensure persistent SyncStates are present in the start state for a SimTick.
+	for (const FMoverDataPersistence& PersistentSyncEntry : MoverComp->PersistentSyncStateDataTypes)
+	{
+		StartData.SyncState.SyncStateCollection.FindOrAddDataByType(PersistentSyncEntry.RequiredType);
+	}
+	
 	MoverComp->SimulationTick(FMoverTimeStep(TimeStep), StartData, OUT EndData);
 
 	*SimOutput.Sync = EndData.SyncState;
diff --git a/Engine/Plugins/Experimental/Mover/Source/Mover/Public/MoverSimulationTypes.h b/Engine/Plugins/Experimental/Mover/Source/Mover/Public/MoverSimulationTypes.h
--- a/Engine/Plugins/Experimental/Mover/Source/Mover/Public/MoverSimulationTypes.h
+++ b/Engine/Plugins/Experimental/Mover/Source/Mover/Public/MoverSimulationTypes.h
@@ -119,7 +119,6 @@ public:
 	FMoverSyncState()
 	{
 		MovementMode = NAME_None;
-		SyncStateCollection.FindOrAddDataByType<FMoverDefaultSyncState>();
 	}
 
 	UScriptStruct* GetStruct() const { return StaticStruct(); }
diff --git a/Engine/Plugins/Experimental/Mover/Source/Mover/Public/MoverTypes.h b/Engine/Plugins/Experimental/Mover/Source/Mover/Public/MoverTypes.h
--- a/Engine/Plugins/Experimental/Mover/Source/Mover/Public/MoverTypes.h
+++ b/Engine/Plugins/Experimental/Mover/Source/Mover/Public/MoverTypes.h
@@ -298,14 +298,19 @@ struct MOVER_API FMoverDataCollection
 		return DataArray;
 	}
 
-protected:
-	static TSharedPtr<FMoverDataStructBase> CreateDataByType(const UScriptStruct* DataStructType);
-	FMoverDataStructBase* AddDataByType(const UScriptStruct* DataStructType);
-
+	/** Find data of a specific type in the collection. */
 	FMoverDataStructBase* FindDataByType(const UScriptStruct* DataStructType) const;
+
+	/** Find data of a specific type in the collection. If not found, a new default instance will be added. */
 	FMoverDataStructBase* FindOrAddDataByType(const UScriptStruct* DataStructType);
+
+	/** Removes data of a specific type in the collection. Returns true if data was removed. */
 	bool RemoveDataByType(const UScriptStruct* DataStructType);
 
+protected:
+	FMoverDataStructBase* AddDataByType(const UScriptStruct* DataStructType);
+	static TSharedPtr<FMoverDataStructBase> CreateDataByType(const UScriptStruct* DataStructType);
+	
 	/** Helper function for serializing array of data */
 	static void NetSerializeDataArray(FArchive& Ar, TArray<TSharedPtr<FMoverDataStructBase>>& DataArray);
 
