Mover: reducing hidden copying in various simulation delegates that would occur even if not bound to anything

diff --git a/Engine/Plugins/Experimental/Mover/Source/Mover/Private/Backends/MoverNetworkPhysicsLiaisonBase.cpp b/Engine/Plugins/Experimental/Mover/Source/Mover/Private/Backends/MoverNetworkPhysicsLiaisonBase.cpp
--- a/Engine/Plugins/Experimental/Mover/Source/Mover/Private/Backends/MoverNetworkPhysicsLiaisonBase.cpp
+++ b/Engine/Plugins/Experimental/Mover/Source/Mover/Private/Backends/MoverNetworkPhysicsLiaisonBase.cpp
@@ -556,7 +556,10 @@ void UMoverNetworkPhysicsLiaisonComponentBase::ProduceInput_External(int32 Physi
 		MoverComp.CachedLastSimTickTimeStep = MoverTimeStep;
 		
 		//@todo DanH: The input info broadcast here will be empty/incorrect when not locally controlled. That seems like a possible landmine of confusion?
-		MoverComp.OnPreSimulationTick.Broadcast(MoverTimeStep, Input.InputCmd);
+		if (MoverComp.OnPreSimulationTick.IsBound())
+		{
+			MoverComp.OnPreSimulationTick.Broadcast(MoverTimeStep, Input.InputCmd);
+		}
 	}
 }
 
@@ -690,7 +693,11 @@ void UMoverNetworkPhysicsLiaisonComponentBase::PerformProcessInputs_Internal(int
 {
 	const FMoverTimeStep MoverTimeStep = bUsingAsyncPhysics ? GetCurrentAsyncMoverTimeStep_Internal() : GetCurrentMoverTimeStep(DeltaTime);
 	FMoverAuxStateContext UnusedAuxState;
-	GetMoverComponent().OnPreMovement.Broadcast(MoverTimeStep, Input.InputCmd, Input.SyncState, UnusedAuxState);
+
+	if (GetMoverComponent().OnPreMovement.IsBound())
+	{
+		GetMoverComponent().OnPreMovement.Broadcast(MoverTimeStep, Input.InputCmd, Input.SyncState, UnusedAuxState);
+	}
 
 	if (Chaos::FPBDRigidParticleHandle* ParticleHandle = GetControlledParticle_Internal())
 	{
diff --git a/Engine/Plugins/Experimental/Mover/Source/Mover/Private/MoverComponent.cpp b/Engine/Plugins/Experimental/Mover/Source/Mover/Private/MoverComponent.cpp
--- a/Engine/Plugins/Experimental/Mover/Source/Mover/Private/MoverComponent.cpp
+++ b/Engine/Plugins/Experimental/Mover/Source/Mover/Private/MoverComponent.cpp
@@ -293,7 +293,10 @@ void UMoverComponent::FinalizeFrame(const FMoverSyncState* SyncState, const FMov
 		bHasValidCachedState = true;
 	}
 
-	OnPostFinalize.Broadcast(*SyncState, *AuxState);
+	if (OnPostFinalize.IsBound())
+	{
+		OnPostFinalize.Broadcast(*SyncState, *AuxState);
+	}
 }
 
 void UMoverComponent::FinalizeSmoothingFrame(const FMoverSyncState* SyncState, const FMoverAuxStateContext* AuxState)
@@ -434,7 +437,10 @@ void UMoverComponent::SimulationTick(const FMoverTimeStep& InTimeStep, const FMo
 		ProcessFirstSimTickAfterRollback(InTimeStep);
 	}
 
-	OnPreSimulationTick.Broadcast(MoverTimeStep, SimInput.InputCmd);
+	if (OnPreSimulationTick.IsBound())
+	{
+		OnPreSimulationTick.Broadcast(MoverTimeStep, SimInput.InputCmd);
+	}
 
 	if (!ModeFSM)
 	{
@@ -485,7 +491,10 @@ void UMoverComponent::SimulationTick(const FMoverTimeStep& InTimeStep, const FMo
 		ModeFSM->QueueNextMode(Input->SuggestedMovementMode);
 	}
 
-	OnPreMovement.Broadcast(MoverTimeStep, SimInput.InputCmd, SimInput.SyncState, SimInput.AuxState);
+	if (OnPreMovement.IsBound())
+	{
+		OnPreMovement.Broadcast(MoverTimeStep, SimInput.InputCmd, SimInput.SyncState, SimInput.AuxState);
+	}
 
 	// Tick the actual simulation. This is where the proposed moves are queried and executed, affecting change to the moving actor's gameplay state and captured in the output sim state
 	if (IsInGameThread())
@@ -524,7 +533,10 @@ void UMoverComponent::SimulationTick(const FMoverTimeStep& InTimeStep, const FMo
 			*GetNameSafe(UpdatedComponent->GetOwner()), UpdatedComponent->GetOwnerRole(), MoverTimeStep.ServerFrame, *SimOutput.MoveRecord.ToString());
 	}
 
-	OnPostMovement.Broadcast(MoverTimeStep, SimOutput.SyncState, SimOutput.AuxState);
+	if (OnPostMovement.IsBound())
+	{
+		OnPostMovement.Broadcast(MoverTimeStep, SimOutput.SyncState, SimOutput.AuxState);
+	}
 
 	CachedLastUsedInputCmd = SimInput.InputCmd;
 	bHasValidCachedUsedInput = true;
