[Mover] Added last used input in FPhysicsMoverAsyncOutput and populating Mover's CachedLastUsedInputCmd

diff --git a/Engine/Plugins/Experimental/Mover/Source/Mover/Private/Backends/MoverNetworkPhysicsLiaison.cpp b/Engine/Plugins/Experimental/Mover/Source/Mover/Private/Backends/MoverNetworkPhysicsLiaison.cpp
--- a/Engine/Plugins/Experimental/Mover/Source/Mover/Private/Backends/MoverNetworkPhysicsLiaison.cpp
+++ b/Engine/Plugins/Experimental/Mover/Source/Mover/Private/Backends/MoverNetworkPhysicsLiaison.cpp
@@ -97,7 +97,7 @@ void UMoverNetworkPhysicsLiaisonComponent::HandleComponentPhysicsStateChanged(UP
 	{
 		SetupConstraint();
 
-		InitializeCachedSyncState();
+		InitializeSimOutputData();
 	}
 
 	Super::HandleComponentPhysicsStateChanged(ChangedComponent, StateChange);
diff --git a/Engine/Plugins/Experimental/Mover/Source/Mover/Private/Backends/MoverNetworkPhysicsLiaisonBase.cpp b/Engine/Plugins/Experimental/Mover/Source/Mover/Private/Backends/MoverNetworkPhysicsLiaisonBase.cpp
--- a/Engine/Plugins/Experimental/Mover/Source/Mover/Private/Backends/MoverNetworkPhysicsLiaisonBase.cpp
+++ b/Engine/Plugins/Experimental/Mover/Source/Mover/Private/Backends/MoverNetworkPhysicsLiaisonBase.cpp
@@ -280,7 +280,7 @@ void UMoverNetworkPhysicsLiaisonComponentBase::HandleComponentPhysicsStateChange
 			MoverComp.ModeFSM->SetModeImmediately(MoverComp.StartingMovementMode);
 		}
 
-		InitializeCachedSyncState();
+		InitializeSimOutputData();
 	}
 }
 
@@ -408,7 +408,7 @@ void UMoverNetworkPhysicsLiaisonComponentBase::BeginPlay()
 		MoverManager->RegisterPhysicsMoverComponent(this);
 	}
 
-	InitializeCachedSyncState();
+	InitializeSimOutputData();
 }
 
 void UMoverNetworkPhysicsLiaisonComponentBase::EndPlay(const EEndPlayReason::Type EndPlayReason)
@@ -491,9 +491,9 @@ FMoverTimeStep UMoverNetworkPhysicsLiaisonComponentBase::GetCurrentMoverTimeStep
 	return TimeStep;
 }
 
-void UMoverNetworkPhysicsLiaisonComponentBase::InitializeCachedSyncState()
+void UMoverNetworkPhysicsLiaisonComponentBase::InitializeSimOutputData()
 {
-	SyncStateRecord.Clear();
+	SimOutputRecord.Clear();
 
 	FMoverSyncState SyncState;
 	SyncState.MovementMode = GetMoverComponent().StartingMovementMode;
@@ -504,7 +504,9 @@ void UMoverNetworkPhysicsLiaisonComponentBase::InitializeCachedSyncState()
 		SimTime = Solver->GetPhysicsResultsTime_External();
 	}
 
-	SyncStateRecord.Add(SimTime, SyncState);
+	FMoverInputCmdContext InputCmd;
+	
+	SimOutputRecord.Add(SimTime, SyncState, InputCmd);
 	LatestSyncState = SyncState;
 }
 
@@ -576,7 +578,7 @@ void UMoverNetworkPhysicsLiaisonComponentBase::ConsumeOutput_External(const FPhy
 
 	if (Output.bIsValid)
 	{
-		SyncStateRecord.Add(OutputTimeInSeconds, Output.SyncState);
+		SimOutputRecord.Add(OutputTimeInSeconds, Output.SyncState, Output.InputCmd);
 	}
 }
 
@@ -594,14 +596,16 @@ void UMoverNetworkPhysicsLiaisonComponentBase::PostPhysicsUpdate_External()
 	const FMoverTimeStep MoverTimeStep = bUsingAsyncPhysics ? GetCurrentAsyncMoverTimeStep_External() : GetCurrentMoverTimeStep(DeltaSeconds);
 
 	FMoverSyncState InterpolatedSyncState;
+	FMoverInputCmdContext InterpolatedInputCmd;
 	if (bUsingAsyncPhysics)
 	{
 		const float ResultsTime = GetPhysicsSolver()->GetPhysicsResultsTime_External();
-		SyncStateRecord.GetInterpolated(ResultsTime, InterpolatedSyncState);
+		SimOutputRecord.GetInterpolated(ResultsTime, InterpolatedSyncState, InterpolatedInputCmd);
 	}
 	else
 	{
-		InterpolatedSyncState = SyncStateRecord.GetLatest();
+		InterpolatedSyncState = SimOutputRecord.GetLatestSyncState();
+		InterpolatedInputCmd = SimOutputRecord.GetLatestInputCmd();
 	}
 	
 	// Physics interactions in the last frame may have caused a change in position or velocity that's different from what a simple lerp would predict,
@@ -622,6 +626,9 @@ void UMoverNetworkPhysicsLiaisonComponentBase::PostPhysicsUpdate_External()
 	MoverComp.CachedLastSyncState = InterpolatedSyncState;
 	MoverComp.bHasValidCachedState = true;
 
+	MoverComp.CachedLastUsedInputCmd = InterpolatedInputCmd;
+	MoverComp.bHasValidCachedUsedInput = true;
+
 	// Note this may return something different now because its dependent on the CachedLastSyncState that we just changed above
 	const FName NextMode = MoverComp.GetMovementModeName();
 	if (CachedMovementMode != NextMode)
@@ -713,6 +720,8 @@ void UMoverNetworkPhysicsLiaisonComponentBase::PreSimulate_Internal(const FPhysi
 	// Sync state should carry over to the next sim frame by default unless something modifies it
 	Output.SyncState = Input.SyncState;
 
+	Output.InputCmd = Input.InputCmd;
+
 	if (CanSimulate_Internal(TickParams, Input))
 	{
 		PerformPreSimulate_Internal(TickParams, Input, Output);
diff --git a/Engine/Plugins/Experimental/Mover/Source/Mover/Public/Backends/MoverNetworkPhysicsLiaisonBase.h b/Engine/Plugins/Experimental/Mover/Source/Mover/Public/Backends/MoverNetworkPhysicsLiaisonBase.h
--- a/Engine/Plugins/Experimental/Mover/Source/Mover/Public/Backends/MoverNetworkPhysicsLiaisonBase.h
+++ b/Engine/Plugins/Experimental/Mover/Source/Mover/Public/Backends/MoverNetworkPhysicsLiaisonBase.h
@@ -98,30 +98,44 @@ struct MOVER_API FNetworkPhysicsMoverTraits
 };
 
 //////////////////////////////////////////////////////////////////////////
-// FSyncStateRecord
+// FSimulationOutputRecord
 
-class FSyncStateRecord
+class FSimulationOutputRecord
 {
 public:
 	struct FData
 	{
-		FMoverSyncState SyncState;
 		double SimTime;
+		FMoverSyncState SyncState;
+		FMoverInputCmdContext InputCmd;
+
+		void Clear()
+		{
+			SimTime = 0.f;
+			SyncState.Reset();
+			InputCmd.Reset();
+		}
 	};
 
-	void Add(double InSimTime, const FMoverSyncState& InSyncState)
+	void Add(double InSimTime, const FMoverSyncState& InSyncState, const FMoverInputCmdContext& InInputCmd)
 	{
 		CurrentIndex = (CurrentIndex + 1) % 2;
 		Data[CurrentIndex].SimTime = InSimTime;
 		Data[CurrentIndex].SyncState = InSyncState;
+		Data[CurrentIndex].InputCmd = InInputCmd;
 	}
 
-	const FMoverSyncState& GetLatest() const
+	const FMoverSyncState& GetLatestSyncState() const
 	{
 		return Data[CurrentIndex].SyncState;
 	}
 
-	void GetInterpolated(double SimTime, FMoverSyncState& OutSyncState) const
+	const FMoverInputCmdContext& GetLatestInputCmd() const
+	{
+		return Data[CurrentIndex].InputCmd;
+	}
+
+	void GetInterpolated(double SimTime, FMoverSyncState& OutSyncState, FMoverInputCmdContext& OutInputCmd) const
 	{
 		const uint8 PrevIndex = (CurrentIndex + 1) % 2;
 		const double PrevTime = Data[PrevIndex].SimTime;
@@ -130,25 +144,26 @@ public:
 		if (FMath::IsNearlyEqual(PrevTime, CurrTime) || (SimTime >= CurrTime))
 		{
 			OutSyncState = Data[CurrentIndex].SyncState;
+			OutInputCmd = Data[CurrentIndex].InputCmd;
 		}
 		else if (SimTime <= PrevTime)
 		{
 			OutSyncState = Data[PrevIndex].SyncState;
+			OutInputCmd = Data[PrevIndex].InputCmd;
 		}
 		else
 		{
 			const float Alpha = FMath::Clamp((SimTime - PrevTime) / (CurrTime - PrevTime), 0.0f, 1.0f);
 			OutSyncState.Interpolate(&Data[PrevIndex].SyncState, &Data[CurrentIndex].SyncState, Alpha);
+			OutInputCmd.Interpolate(&Data[PrevIndex].InputCmd, &Data[CurrentIndex].InputCmd, Alpha);
 		}
 	}
 
 	void Clear()
 	{
 		CurrentIndex = 1;
-		Data[0].SyncState.Reset();
-		Data[0].SyncState.SyncStateCollection.Empty();
-		Data[1].SyncState.Reset();
-		Data[1].SyncState.SyncStateCollection.Empty();
+		Data[0].Clear();
+		Data[1].Clear();
 	}
 
 	FData Data[2];
@@ -266,7 +281,7 @@ protected:
 	// Time step on the game thread when using non-async physics
 	FMoverTimeStep GetCurrentMoverTimeStep(float DeltaSeconds) const;
 	
-	void InitializeCachedSyncState();
+	void InitializeSimOutputData();
 
 	UMoverComponent& GetMoverComponent() const;
 
@@ -287,10 +302,10 @@ protected:
 	TObjectPtr<UNetworkPhysicsComponent> NetworkPhysicsComponent;
 
 	// This stores the previous two outputs from the physics thread
-	// which can be interpolated to get the sync state at a time
+	// which can be interpolated to get the new output at a time
 	// that matches the interpolated physics particle.
-	FSyncStateRecord SyncStateRecord;
-
+	FSimulationOutputRecord SimOutputRecord;
+	
 	bool bCachedInputIsValid = false;
 	bool bUsingAsyncPhysics = false;
 };
\ No newline at end of file
diff --git a/Engine/Plugins/Experimental/Mover/Source/Mover/Public/MoverSimulationTypes.h b/Engine/Plugins/Experimental/Mover/Source/Mover/Public/MoverSimulationTypes.h
--- a/Engine/Plugins/Experimental/Mover/Source/Mover/Public/MoverSimulationTypes.h
+++ b/Engine/Plugins/Experimental/Mover/Source/Mover/Public/MoverSimulationTypes.h
@@ -88,6 +88,16 @@ public:
 	{
 		InputCollection.ToString(Out);
 	}
+
+	void Interpolate(const FMoverInputCmdContext* From, const FMoverInputCmdContext* To, float Pct)
+	{
+		InputCollection.Interpolate(From->InputCollection, To->InputCollection, Pct);
+	}
+
+	void Reset()
+	{
+		InputCollection.Empty();
+	}
 };
 
 
diff --git a/Engine/Plugins/Experimental/Mover/Source/Mover/Public/PhysicsMover/PhysicsMoverSimulationTypes.h b/Engine/Plugins/Experimental/Mover/Source/Mover/Public/PhysicsMover/PhysicsMoverSimulationTypes.h
--- a/Engine/Plugins/Experimental/Mover/Source/Mover/Public/PhysicsMover/PhysicsMoverSimulationTypes.h
+++ b/Engine/Plugins/Experimental/Mover/Source/Mover/Public/PhysicsMover/PhysicsMoverSimulationTypes.h
@@ -63,6 +63,7 @@ struct FPhysicsMoverAsyncInput
 struct FPhysicsMoverAsyncOutput
 {
 	FMoverSyncState SyncState;
+	FMoverInputCmdContext InputCmd;
 	FFloorCheckResult FloorResult;
 	bool bIsValid = false;
 };
