Mover: Exposing BaseVisualComponentTransform with a getter and setter function so we can set the offset of visual components. We now also use these functions to properly adjust the visual component of the pawn when crouching.

diff --git a/Engine/Plugins/Experimental/Mover/Source/Mover/Private/DefaultMovementSet/MovementModifiers/StanceModifier.cpp b/Engine/Plugins/Experimental/Mover/Source/Mover/Private/DefaultMovementSet/MovementModifiers/StanceModifier.cpp
--- a/Engine/Plugins/Experimental/Mover/Source/Mover/Private/DefaultMovementSet/MovementModifiers/StanceModifier.cpp
+++ b/Engine/Plugins/Experimental/Mover/Source/Mover/Private/DefaultMovementSet/MovementModifiers/StanceModifier.cpp
@@ -12,50 +12,6 @@
 #include "GameFramework/Pawn.h"
 #include "MoveLibrary/MovementUtils.h"
 
-// Gets CDO component type - useful for getting original values
-template <class ComponentType>
-static const ComponentType* GetOriginalComponentType(const UMoverComponent* MoverComp)
-{
-	const ComponentType* OriginalComponent = nullptr;
-
-	if (const AActor* OwnerCDO = Cast<AActor>(MoverComp->GetOwner()->GetClass()->GetDefaultObject()))
-	{
-		// Check if native CDO has Capsule component
-		OriginalComponent = OwnerCDO->FindComponentByClass<ComponentType>();
-
-		// check if it comes from a BP
-		if (!OriginalComponent)
-		{
-			if (const UBlueprintGeneratedClass* OwnerClassAsBP = Cast<UBlueprintGeneratedClass>(OwnerCDO->GetClass()))
-			{
-				TArray<const UBlueprintGeneratedClass*> BlueprintClasses;
-				UBlueprintGeneratedClass::GetGeneratedClassesHierarchy(OwnerClassAsBP, BlueprintClasses);
-				for (const UBlueprintGeneratedClass* BlueprintClass : BlueprintClasses)
-				{
-					if (BlueprintClass->SimpleConstructionScript)
-					{
-						// Check Simple construction script
-						const TArray<USCS_Node*>& SCSNodes = BlueprintClass->SimpleConstructionScript->GetAllNodes();
-						for (USCS_Node* SCSNode : SCSNodes)
-						{
-							if (SCSNode)
-							{
-								if (const ComponentType* BPComponent = Cast<ComponentType>(SCSNode->ComponentTemplate))
-								{
-									OriginalComponent = BPComponent;
-									break;
-								}
-							}
-						}
-					}
-				}
-			}
-		}
-	}
-
-	return OriginalComponent;
-}
-
 FStanceModifier::FStanceModifier()
 	: ActiveStance(EStanceMode::Crouch)
 {
@@ -264,7 +220,9 @@ void FStanceModifier::AdjustCapsule(UMoverComponent* MoverComp, float OldHalfHei
 	}
 	
 	// Add offset to visual component as the base location has changed
-	MoverComp->GetPrimaryVisualComponent()->AddRelativeLocation(VisualOffset);
+	FTransform MoverVisualComponentOffset = MoverComp->GetBaseVisualComponentTransform();
+	MoverVisualComponentOffset.SetLocation(MoverVisualComponentOffset.GetLocation() + VisualOffset);
+	MoverComp->SetBaseVisualComponentTransform(MoverVisualComponentOffset);
 }
 
 void FStanceModifier::ApplyMovementSettings(UMoverComponent* MoverComp)
diff --git a/Engine/Plugins/Experimental/Mover/Source/Mover/Private/MoverComponent.cpp b/Engine/Plugins/Experimental/Mover/Source/Mover/Private/MoverComponent.cpp
--- a/Engine/Plugins/Experimental/Mover/Source/Mover/Private/MoverComponent.cpp
+++ b/Engine/Plugins/Experimental/Mover/Source/Mover/Private/MoverComponent.cpp
@@ -1841,6 +1841,16 @@ void UMoverComponent::SetPlanarConstraint(const FPlanarConstraint& InConstraint)
 	PlanarConstraint = InConstraint;
 }
 
+void UMoverComponent::SetBaseVisualComponentTransform(const FTransform& ComponentTransform)
+{
+	BaseVisualComponentTransform = ComponentTransform;
+}
+
+FTransform UMoverComponent::GetBaseVisualComponentTransform() const
+{
+	return BaseVisualComponentTransform;
+}
+
 TArray<FTrajectorySampleInfo> UMoverComponent::GetFutureTrajectory(float FutureSeconds, float SamplesPerSecond)
 {
 	FMoverPredictTrajectoryParams PredictionParams;
diff --git a/Engine/Plugins/Experimental/Mover/Source/Mover/Public/MoverComponent.h b/Engine/Plugins/Experimental/Mover/Source/Mover/Public/MoverComponent.h
--- a/Engine/Plugins/Experimental/Mover/Source/Mover/Public/MoverComponent.h
+++ b/Engine/Plugins/Experimental/Mover/Source/Mover/Public/MoverComponent.h
@@ -251,6 +251,14 @@ public:
 	UFUNCTION(BlueprintCallable, Category = Mover)
 	void SetPlanarConstraint(const FPlanarConstraint& InConstraint);
 
+	// Sets BaseVisualComponentTransform used for cases where we want to move the visual component away from the root component. See @BaseVisualComponentTransform
+	UFUNCTION(BlueprintCallable, Category = Mover)
+	void SetBaseVisualComponentTransform (const FTransform& ComponentTransform);
+
+	// Gets BaseVisualComponentTransform used for cases where we want to move the visual component away from the root component. See @BaseVisualComponentTransform
+	UFUNCTION(BlueprintCallable, Category = Mover)
+	FTransform GetBaseVisualComponentTransform() const;
+
 public:
 
 	/**
