Mover: optimizing async nav walking to eliminate the need for modifying primitive comp collision properties at sim time, also eliminating the need for binding to the delegate. Adds new versions of async-friendly collision queries that allow parameter changing without altering the scene component(s).

diff --git a/Engine/Plugins/Experimental/Mover/Source/Mover/Private/DefaultMovementSet/Modes/AsyncNavWalkingMode.cpp b/Engine/Plugins/Experimental/Mover/Source/Mover/Private/DefaultMovementSet/Modes/AsyncNavWalkingMode.cpp
--- a/Engine/Plugins/Experimental/Mover/Source/Mover/Private/DefaultMovementSet/Modes/AsyncNavWalkingMode.cpp
+++ b/Engine/Plugins/Experimental/Mover/Source/Mover/Private/DefaultMovementSet/Modes/AsyncNavWalkingMode.cpp
@@ -121,6 +121,8 @@ void UAsyncNavWalkingMode::GenerateMove_Implementation(const FMoverTickStartData
 
 void UAsyncNavWalkingMode::SimulationTick_Implementation(const FSimulationTickParams& Params, FMoverTickEndData& OutputState)
 {
+	SCOPED_NAMED_EVENT(AsyncNavWalkingMode_SimulationTick, FColor::Yellow);
+
 	const UMoverComponent* MoverComp = Cast<UMoverComponent>(GetMoverComponent());
 	if (!MoverComp || !NavMoverComponent.IsValid() || !CommonLegacySettings.IsValid())
 	{
@@ -280,7 +282,13 @@ void UAsyncNavWalkingMode::SimulationTick_Implementation(const FSimulationTickPa
 		{
 			FHitResult MoveHitResult;
 
-			if (UAsyncMovementUtils::TestDepenetratingMove(Params.MovingComps, StartLocation, StartLocation + AdjustedDelta, StartRotation, TargetRotation, bSweepWhileNavWalking, OUT MoveHitResult, IN OUT MoveRecord))
+			FMoverCollisionParams CollisionParams(Params.MovingComps.UpdatedComponent.Get());
+
+			// Ignore all world geometry while moving on nav mesh
+			CollisionParams.ResponseParams.CollisionResponse.SetResponse(ECC_WorldStatic, ECR_Ignore);
+			CollisionParams.ResponseParams.CollisionResponse.SetResponse(ECC_WorldDynamic, ECR_Ignore);
+
+			if (UAsyncMovementUtils::TestDepenetratingMove(Params.MovingComps, StartLocation, StartLocation + AdjustedDelta, StartRotation, TargetRotation, bSweepWhileNavWalking, CollisionParams, OUT MoveHitResult, IN OUT MoveRecord))
 			{
 				LocationInProgress = StartLocation + (AdjustedDelta * MoveHitResult.Time);
 				RotationInProgress = FQuat::Slerp(StartRotation, TargetRotation, MoveHitResult.Time);
@@ -333,46 +341,12 @@ void UAsyncNavWalkingMode::SetTurnGeneratorClass(TSubclassOf<UObject> TurnGenera
 	}
 }
 
-void UAsyncNavWalkingMode::SetCollisionForNavWalking(bool bEnable)
-{
-	if (const UMoverComponent* MoverComponent = GetMoverComponent())
-	{
-		if (UPrimitiveComponent* UpdatedCompAsPrimitive = Cast<UPrimitiveComponent>(MoverComponent->GetUpdatedComponent()))
-		{
-			if (bEnable)
-			{
-				CollideVsWorldStatic = UpdatedCompAsPrimitive->GetCollisionResponseToChannel(ECC_WorldStatic);
-				CollideVsWorldDynamic = UpdatedCompAsPrimitive->GetCollisionResponseToChannel(ECC_WorldDynamic);
-
-				// TODO NS: Eventually might want make the new collision response an option so they trigger certain things while still not colliding 
-				UpdatedCompAsPrimitive->SetCollisionResponseToChannel(ECC_WorldStatic, ECR_Ignore);
-				UpdatedCompAsPrimitive->SetCollisionResponseToChannel(ECC_WorldDynamic, ECR_Ignore);
-			}
-			else
-			{
-				// Grabbing the original shape settings and reverting to our old collision responses
-				if (const UShapeComponent* OriginalShapeComp = UMovementUtils::GetOriginalComponentType<UShapeComponent>(MoverComponent->GetOwner()))
-				{
-					CollideVsWorldStatic = OriginalShapeComp->GetCollisionResponseToChannel(ECC_WorldStatic);
-					CollideVsWorldDynamic = OriginalShapeComp->GetCollisionResponseToChannel(ECC_WorldDynamic);
-				}
-
-				UpdatedCompAsPrimitive->SetCollisionResponseToChannel(ECC_WorldStatic, CollideVsWorldStatic);
-				UpdatedCompAsPrimitive->SetCollisionResponseToChannel(ECC_WorldDynamic, CollideVsWorldDynamic);
-			}
-		}
-	}
-}
-
 void UAsyncNavWalkingMode::Activate()
 {
 	Super::Activate();
 	
 	if (UMoverComponent* MoverComp = GetMoverComponent())
 	{
-		// TEMP solution for our inability to change collision settings outside of the game thread
-		MoverComp->OnPostFinalize.AddDynamic(this, &UAsyncNavWalkingMode::PostFinalizeSetCollisionForNavWalking);
-
 		if (UMoverBlackboard* SimBlackboard = MoverComp->GetSimBlackboard_Mutable())
 		{
 			SimBlackboard->Invalidate(NavWalkingBlackBoard::ProjectedNavMeshHitResult);
@@ -382,15 +356,6 @@ void UAsyncNavWalkingMode::Activate()
 	NavMeshProjectionTimer = (NavMeshProjectionInterval > 0.f) ? FMath::FRandRange(-NavMeshProjectionInterval, 0.f) : 0.f;
 }
 
-void UAsyncNavWalkingMode::Deactivate()
-{
-	if (UMoverComponent* MoverComp = GetMoverComponent())
-	{
-		MoverComp->OnPostFinalize.AddDynamic(this, &UAsyncNavWalkingMode::PostFinalizeRestoreOriginalCollision);
-	}
-	
-	Super::Deactivate();
-}
 
 const INavigationDataInterface* UAsyncNavWalkingMode::GetNavData() const
 {
@@ -614,22 +579,3 @@ void UAsyncNavWalkingMode::CaptureFinalState(const FVector FinalLocation, const
 		nullptr);	// no movement base
 }
 
-void UAsyncNavWalkingMode::PostFinalizeSetCollisionForNavWalking(const FMoverSyncState& SyncState, const FMoverAuxStateContext& AuxState)
-{
-	SetCollisionForNavWalking(true);
-
-	if (UMoverComponent* MoverComp = GetMoverComponent())
-	{
-		MoverComp->OnPostFinalize.RemoveDynamic(this, &UAsyncNavWalkingMode::PostFinalizeSetCollisionForNavWalking);
-	}
-}
-
-void UAsyncNavWalkingMode::PostFinalizeRestoreOriginalCollision(const FMoverSyncState& SyncState, const FMoverAuxStateContext& AuxState)
-{
-	SetCollisionForNavWalking(false);
-
-	if (UMoverComponent* MoverComp = GetMoverComponent())
-	{
-		MoverComp->OnPostFinalize.RemoveDynamic(this, &UAsyncNavWalkingMode::PostFinalizeRestoreOriginalCollision);
-	}
-}
diff --git a/Engine/Plugins/Experimental/Mover/Source/Mover/Private/MoveLibrary/AsyncMovementUtils.cpp b/Engine/Plugins/Experimental/Mover/Source/Mover/Private/MoveLibrary/AsyncMovementUtils.cpp
--- a/Engine/Plugins/Experimental/Mover/Source/Mover/Private/MoveLibrary/AsyncMovementUtils.cpp
+++ b/Engine/Plugins/Experimental/Mover/Source/Mover/Private/MoveLibrary/AsyncMovementUtils.cpp
@@ -19,7 +19,15 @@ static float PenetrationOverlapCheckInflationCVar = 0.1f;
 static const FName TestSafeMoveSubstepName = "TestSafeMove";
 static const FName TestSafeDepenetrationSubstepName = "TestSafeResolvePenetration";
 
+/* static */
 bool UAsyncMovementUtils::TestDepenetratingMove(const FMovingComponentSet& MovingComps, const FVector& StartLocation, const FVector& TargetLocation, const FQuat& StartRotation, const FQuat& TargetRotation, bool bShouldSweep, FHitResult& OutHit, FMovementRecord& InOutMoveRecord)
+{
+	FMoverCollisionParams CollisionParams(MovingComps.UpdatedComponent.Get());
+	return TestDepenetratingMove(MovingComps, StartLocation, TargetLocation, StartRotation, TargetRotation, bShouldSweep, CollisionParams, OutHit, InOutMoveRecord);
+}
+
+/* static */
+bool UAsyncMovementUtils::TestDepenetratingMove(const FMovingComponentSet& MovingComps, const FVector& StartLocation, const FVector& TargetLocation, const FQuat& StartRotation, const FQuat& TargetRotation, bool bShouldSweep, FMoverCollisionParams& CollisionParams, FHitResult& OutHit, FMovementRecord& InOutMoveRecord)
 {
 	const USceneComponent* UpdatedComponent = MovingComps.UpdatedComponent.Get();
 
@@ -30,7 +38,6 @@ bool UAsyncMovementUtils::TestDepenetratingMove(const FMovingComponentSet& Movin
 	}
 
 	bool bDidMove = false;
-	EMoveComponentFlags MoveCompFlags = MOVECOMP_NoFlags;
 
 	FVector ResolvedStartLocation = StartLocation;
 	FVector CurrentLocation = StartLocation;
@@ -39,8 +46,9 @@ bool UAsyncMovementUtils::TestDepenetratingMove(const FMovingComponentSet& Movin
 	{
 		const EMoveComponentFlags IncludeBlockingOverlapsWithoutEvents = (MOVECOMP_NeverIgnoreBlockingOverlaps | MOVECOMP_DisableBlockingOverlapDispatch);
 
-		MoveCompFlags = (MoveCompFlags | IncludeBlockingOverlapsWithoutEvents);
-		bDidMove = TestMoveComponent_Internal(MovingComps, ResolvedStartLocation, TargetLocation, StartRotation, TargetRotation, bShouldSweep, MoveCompFlags, OutHit);
+		TGuardValue<EMoveComponentFlags> ScopedFlagRestore(CollisionParams.MoveComponentFlags, EMoveComponentFlags(CollisionParams.MoveComponentFlags | IncludeBlockingOverlapsWithoutEvents));
+
+		bDidMove = TestMoveComponent_Internal(MovingComps, ResolvedStartLocation, TargetLocation, StartRotation, TargetRotation, bShouldSweep, CollisionParams, OutHit);
 
 		CurrentLocation = ResolvedStartLocation + ((TargetLocation-ResolvedStartLocation) * OutHit.Time);
 
@@ -52,7 +60,7 @@ bool UAsyncMovementUtils::TestDepenetratingMove(const FMovingComponentSet& Movin
 	if (OutHit.bStartPenetrating)
 	{
 		FVector AdjustmentToResolve = FVector::ZeroVector;
-		if (FindMoveToResolveInitialPenetration_Internal(MovingComps, StartLocation, StartRotation, OutHit, MoveCompFlags, AdjustmentToResolve))
+		if (FindMoveToResolveInitialPenetration_Internal(MovingComps, StartLocation, StartRotation, OutHit, CollisionParams, AdjustmentToResolve))
 		{
 			ResolvedStartLocation = StartLocation + AdjustmentToResolve;
 			CurrentLocation = ResolvedStartLocation;	// keeping track of where we've moved so far
@@ -60,7 +68,7 @@ bool UAsyncMovementUtils::TestDepenetratingMove(const FMovingComponentSet& Movin
 			InOutMoveRecord.Append(FMovementSubstep(TestSafeDepenetrationSubstepName, AdjustmentToResolve, false));
 
 			// Retry original move
-			bDidMove |= TestMoveComponent_Internal(MovingComps, ResolvedStartLocation, TargetLocation, StartRotation, TargetRotation, bShouldSweep, MoveCompFlags, OutHit);
+			bDidMove |= TestMoveComponent_Internal(MovingComps, ResolvedStartLocation, TargetLocation, StartRotation, TargetRotation, bShouldSweep, CollisionParams, OutHit);
 
 			UE_LOG(LogMover, VeryVerbose, TEXT("TestDepenetratingMove retry: %s (role %i) Delta=%s DidMove=%i"),
 				*GetNameSafe(UpdatedComponent->GetOwner()), UpdatedComponent->GetOwnerRole(), *(TargetLocation - ResolvedStartLocation).ToCompactString(), bDidMove);
@@ -80,9 +88,16 @@ bool UAsyncMovementUtils::TestDepenetratingMove(const FMovingComponentSet& Movin
 	return bDidMove;
 }
 
+/* static */
+float UAsyncMovementUtils::TestSlidingMoveAlongHitSurface(const FMovingComponentSet& MovingComps, const FVector& OriginalMoveDelta, const FVector& LocationAtHit, const FQuat& TargetRotation, FHitResult& InOutHit, FMovementRecord& InOutMoveRecord)
+{
+	FMoverCollisionParams CollisionParams(MovingComps.UpdatedComponent.Get());
+	return TestSlidingMoveAlongHitSurface(MovingComps, OriginalMoveDelta, LocationAtHit, TargetRotation, CollisionParams, InOutHit, InOutMoveRecord);
+
+}
 
 /* static */ 
-float UAsyncMovementUtils::TestSlidingMoveAlongHitSurface(const FMovingComponentSet& MovingComps, const FVector& OriginalMoveDelta, const FVector& LocationAtHit, const FQuat& TargetRotation, FHitResult& InOutHit, FMovementRecord& InOutMoveRecord)
+float UAsyncMovementUtils::TestSlidingMoveAlongHitSurface(const FMovingComponentSet& MovingComps, const FVector& OriginalMoveDelta, const FVector& LocationAtHit, const FQuat& TargetRotation, FMoverCollisionParams& CollisionParams, FHitResult& InOutHit, FMovementRecord& InOutMoveRecord)
 {
 	if (!InOutHit.bBlockingHit)
 	{
@@ -97,7 +112,7 @@ float UAsyncMovementUtils::TestSlidingMoveAlongHitSurface(const FMovingComponent
 
 	if ((SlideDelta | OriginalMoveDelta) > 0.f)
 	{
-		TestDepenetratingMove(MovingComps, LocationAtHit, LocationAtHit+SlideDelta, TargetRotation, TargetRotation, /*bShouldSweep*/ true, InOutHit, InOutMoveRecord);
+		TestDepenetratingMove(MovingComps, LocationAtHit, LocationAtHit+SlideDelta, TargetRotation, TargetRotation, /*bShouldSweep*/ true, CollisionParams, InOutHit, InOutMoveRecord);
 
 		PctOfTimeUsed = InOutHit.Time;
 
@@ -114,7 +129,7 @@ float UAsyncMovementUtils::TestSlidingMoveAlongHitSurface(const FMovingComponent
 				const FVector LocationAfter1stSlide = InOutHit.TraceStart + ((InOutHit.TraceEnd-InOutHit.TraceStart) * InOutHit.Time);
 
 				// Perform second move
-				TestDepenetratingMove(MovingComps, LocationAfter1stSlide, LocationAfter1stSlide + SlideDelta, TargetRotation, TargetRotation, /*bShouldSweep*/ true, InOutHit, InOutMoveRecord);
+				TestDepenetratingMove(MovingComps, LocationAfter1stSlide, LocationAfter1stSlide + SlideDelta, TargetRotation, TargetRotation, /*bShouldSweep*/ true, CollisionParams, InOutHit, InOutMoveRecord);
 				PctOfTimeUsed += (InOutHit.Time * (1.f - PctOfTimeUsed));
 
 				// TODO: maybe record second impact here
@@ -130,7 +145,7 @@ float UAsyncMovementUtils::TestSlidingMoveAlongHitSurface(const FMovingComponent
 
 
 /* static */
-bool UAsyncMovementUtils::FindMoveToResolveInitialPenetration_Internal(const FMovingComponentSet& MovingComps, const FVector& StartLocation, const FQuat& StartRotation, const FHitResult& PenetratingHit, EMoveComponentFlags MoveCompFlags, FVector& OutAdjustmentDelta)
+bool UAsyncMovementUtils::FindMoveToResolveInitialPenetration_Internal(const FMovingComponentSet& MovingComps, const FVector& StartLocation, const FQuat& StartRotation, const FHitResult& PenetratingHit, FMoverCollisionParams& CollisionParams, FVector& OutAdjustmentDelta)
 {
 	if (ensureMsgf(PenetratingHit.bStartPenetrating, TEXT("Expected a hit that started penetrating. Will not attempt adjustment.")))
 	{
@@ -176,13 +191,13 @@ bool UAsyncMovementUtils::FindMoveToResolveInitialPenetration_Internal(const FMo
 			}
 
 			// Disable MOVECOMP_NeverIgnoreBlockingOverlaps if it is enabled, otherwise we wouldn't be able to sweep out of the object to fix the penetration.
-			TGuardValue<EMoveComponentFlags> ScopedFlagRestore(MoveCompFlags, EMoveComponentFlags(MoveCompFlags & (~MOVECOMP_NeverIgnoreBlockingOverlaps)));
+			TGuardValue<EMoveComponentFlags> ScopedFlagRestore(CollisionParams.MoveComponentFlags, EMoveComponentFlags(CollisionParams.MoveComponentFlags & (~MOVECOMP_NeverIgnoreBlockingOverlaps)));
 
 			// Try sweeping out as far as possible...
 
 			FHitResult SweepOutHit(1.f);
 
-			bool bMoved = TestMoveComponent_Internal(MovingComps, StartLocation, StartLocation + ProposedAdjustment, StartRotation, StartRotation, /* bShouldSweep */ true, MoveCompFlags, SweepOutHit);
+			bool bMoved = TestMoveComponent_Internal(MovingComps, StartLocation, StartLocation + ProposedAdjustment, StartRotation, StartRotation, /* bShouldSweep */ true, CollisionParams, SweepOutHit);
 
 			UE_LOG(LogMover, VeryVerbose, TEXT("FindMoveToResolveInitialPenetration_Internal: %s (role %i) Adjustment=%s DidMove=%i"),
 				*GetNameSafe(UpdatedComponent->GetOwner()), UpdatedComponent->GetOwnerRole(), *ProposedAdjustment.ToCompactString(), bMoved);
@@ -196,7 +211,7 @@ bool UAsyncMovementUtils::FindMoveToResolveInitialPenetration_Internal(const FMo
 
 				if (SecondMTD != ProposedAdjustment && !CombinedMTD.IsZero())
 				{
-					bMoved = TestMoveComponent_Internal(MovingComps, StartLocation, StartLocation+CombinedMTD, StartRotation, StartRotation, /*bShouldSweep*/ true, MoveCompFlags, SweepOutHit);
+					bMoved = TestMoveComponent_Internal(MovingComps, StartLocation, StartLocation+CombinedMTD, StartRotation, StartRotation, /*bShouldSweep*/ true, CollisionParams, SweepOutHit);
 
 					UE_LOG(LogMover, VeryVerbose, TEXT("FindMoveToResolveInitialPenetration_Internal combined: %s (role %i) CombinedAdjustment=%s DidMove=%i"),
 						*GetNameSafe(UpdatedComponent->GetOwner()), UpdatedComponent->GetOwnerRole(), *CombinedMTD.ToCompactString(), bMoved);
@@ -214,7 +229,7 @@ bool UAsyncMovementUtils::FindMoveToResolveInitialPenetration_Internal(const FMo
 				if (!FullMoveDelta.IsZero())
 				{
 					const FVector TargetLocationWithAdjustment = StartLocation + ProposedAdjustment + FullMoveDelta;
-					bMoved = TestMoveComponent_Internal(MovingComps, StartLocation, TargetLocationWithAdjustment, StartRotation, StartRotation, /*bShouldSweep*/ true, MoveCompFlags, SweepOutHit);
+					bMoved = TestMoveComponent_Internal(MovingComps, StartLocation, TargetLocationWithAdjustment, StartRotation, StartRotation, /*bShouldSweep*/ true, CollisionParams, SweepOutHit);
 
 					UE_LOG(LogMover, VeryVerbose, TEXT("FindMoveToResolveInitialPenetration_Internal multiple: %s (role %i) TargetLocationWithAdjustment=%s DidMove=%i"),
 						*GetNameSafe(UpdatedComponent->GetOwner()), UpdatedComponent->GetOwnerRole(), *TargetLocationWithAdjustment.ToCompactString(), bMoved);
@@ -275,9 +290,10 @@ static bool ShouldIgnoreHitResult(const UWorld* InWorld, FHitResult const& TestH
 	return false;
 }
 
-static void PullBackHit(FHitResult& Hit, const FVector& Start, const FVector& End, const float Dist)
+// Adjusts the time of the hit result to pull it back slightly from the actual hit
+static void PullBackHit(FHitResult& Hit, const float OrigMoveDist)
 {
-	const float DesiredTimeBack = FMath::Clamp(0.1f, 0.1f / Dist, 1.f / Dist) + 0.001f;
+	const float DesiredTimeBack = FMath::Clamp(0.1f, 0.1f / OrigMoveDist, 1.f / OrigMoveDist) + 0.001f;
 	Hit.Time = FMath::Clamp(Hit.Time - DesiredTimeBack, 0.f, 1.f);
 }
 
@@ -285,7 +301,7 @@ static void PullBackHit(FHitResult& Hit, const FVector& Start, const FVector& En
 static const FName TestSweepTraceTagName = "SweepTestMoverComponent";
 
 // Returns true if any of the sweeping movement could occur
-bool UAsyncMovementUtils::TestMoveComponent_Internal(const FMovingComponentSet& MovingComps, const FVector& StartLocation, const FVector& TargetLocation, const FQuat& StartRotation, const FQuat& TargetRotation, bool bShouldSweep, EMoveComponentFlags MoveCompFlags, FHitResult& OutHit)
+bool UAsyncMovementUtils::TestMoveComponent_Internal(const FMovingComponentSet& MovingComps, const FVector& StartLocation, const FVector& TargetLocation, const FQuat& StartRotation, const FQuat& TargetRotation, bool bShouldSweep, FMoverCollisionParams& CollisionParams, FHitResult& OutHit)
 {
 	UPrimitiveComponent* UpdatedPrimitive = MovingComps.UpdatedPrimitive.Get();
 	if (!UpdatedPrimitive)
@@ -329,33 +345,27 @@ bool UAsyncMovementUtils::TestMoveComponent_Internal(const FMovingComponentSet&
 		DeltaSizeSq = 0.f;
 	}
 
-
-	TArray<FHitResult> AllHits;
-	FVector NewLocation = TraceStart;
-
 	// Perform movement collision checking if needed for this actor.
 	const bool bCollisionEnabled = UpdatedPrimitive->IsQueryCollisionEnabled();
 	UWorld* const MyWorld = UpdatedPrimitive->GetWorld();
 	if (MyWorld && bCollisionEnabled && (DeltaSizeSq > 0.f))
 	{
-		FComponentQueryParams Params(SCENE_QUERY_STAT(SweepTestMoverComponent), OwningActor);
-		FCollisionResponseParams ResponseParam;
-		UpdatedPrimitive->InitSweepCollisionParams(Params, ResponseParam);
-		Params.bIgnoreTouches = true;
-		Params.TraceTag = TestSweepTraceTagName;
+		TArray<FHitResult> AllHits;
 
-		const bool bHadBlockingHit = MyWorld->ComponentSweepMulti(AllHits, UpdatedPrimitive, TraceStart, TraceEnd, InitialRotationQuat, Params);
+		TGuardValue<bool> GuardIgnoreTouches(CollisionParams.QueryParams.bIgnoreTouches, true);
+
+		const bool bHadBlockingHit = MyWorld->SweepMultiByChannel(AllHits, TraceStart, TraceEnd, InitialRotationQuat, 
+			CollisionParams.Channel, CollisionParams.Shape, CollisionParams.QueryParams, CollisionParams.ResponseParams);
 
 		if (AllHits.Num() > 0)
 		{
 			const float DeltaSize = FMath::Sqrt(DeltaSizeSq);
 			for (int32 HitIdx = 0; HitIdx < AllHits.Num(); HitIdx++)
 			{
-				PullBackHit(AllHits[HitIdx], TraceStart, TraceEnd, DeltaSize);
+				PullBackHit(AllHits[HitIdx], DeltaSize);
 			}
 		}
 
-
 		// Find the best blocking hit from AllHits
 		if (bHadBlockingHit)
 		{
@@ -367,7 +377,8 @@ bool UAsyncMovementUtils::TestMoveComponent_Internal(const FMovingComponentSet&
 				const FHitResult& TestHit = AllHits[HitIdx];
 				if (TestHit.bBlockingHit)
 				{
-					if (!ShouldIgnoreHitResult(MyWorld, TestHit, ConstrainedDelta, OwningActor, MoveCompFlags) && !UpdatedPrimitive->ShouldComponentIgnoreHitResult(TestHit, MoveCompFlags))
+					if (!ShouldIgnoreHitResult(MyWorld, TestHit, ConstrainedDelta, OwningActor, CollisionParams.MoveComponentFlags) && 
+					    !UpdatedPrimitive->ShouldComponentIgnoreHitResult(TestHit, CollisionParams.MoveComponentFlags))
 					{
 						if (TestHit.bStartPenetrating)
 						{
diff --git a/Engine/Plugins/Experimental/Mover/Source/Mover/Private/MoveLibrary/GroundMovementUtils.cpp b/Engine/Plugins/Experimental/Mover/Source/Mover/Private/MoveLibrary/GroundMovementUtils.cpp
--- a/Engine/Plugins/Experimental/Mover/Source/Mover/Private/MoveLibrary/GroundMovementUtils.cpp
+++ b/Engine/Plugins/Experimental/Mover/Source/Mover/Private/MoveLibrary/GroundMovementUtils.cpp
@@ -792,12 +792,13 @@ bool UGroundMovementUtils::TestMoveToStepOver(const FMovingComponentSet& MovingC
 		return false;
 	}
 
+	FMoverCollisionParams CollisionParams(MovingComps.UpdatedComponent.Get());
 	// step up - treat as vertical wall
 	FHitResult SweepUpHit(1.f);
 
 	const FVector UpAdjustment = -GravDir * StepTravelUpHeight;
 
-	const bool bDidStepUp = UAsyncMovementUtils::TestMoveComponent_Internal(MovingComps, LocationInProgress, LocationInProgress + UpAdjustment, Rotation, Rotation, /*bShouldSweep=*/ true, MOVECOMP_NoFlags, SweepUpHit);
+	const bool bDidStepUp = UAsyncMovementUtils::TestMoveComponent_Internal(MovingComps, LocationInProgress, LocationInProgress + UpAdjustment, Rotation, Rotation, /*bShouldSweep=*/ true, CollisionParams, SweepUpHit);
 
 	UE_LOG(LogMover, VeryVerbose, TEXT("TestMoveToStepOver Up: %s (role %i) UpAdjustment=%s DidMove=%i"),
 		*GetNameSafe(CapsuleComponent->GetOwner()), CapsuleComponent->GetOwnerRole(), *UpAdjustment.ToCompactString(), bDidStepUp);
@@ -817,7 +818,7 @@ bool UGroundMovementUtils::TestMoveToStepOver(const FMovingComponentSet& MovingC
 
 	// step fwd
 	FHitResult StepFwdHit(1.f);
-	const bool bDidStepFwd = UAsyncMovementUtils::TestMoveComponent_Internal(MovingComps, LocationInProgress, LocationInProgress + MoveDelta, Rotation, Rotation, /*bShouldSweep=*/ true, MOVECOMP_NoFlags, StepFwdHit);
+	const bool bDidStepFwd = UAsyncMovementUtils::TestMoveComponent_Internal(MovingComps, LocationInProgress, LocationInProgress + MoveDelta, Rotation, Rotation, /*bShouldSweep=*/ true, CollisionParams, StepFwdHit);
 
 	FVector FwdStepDelta = (StepFwdHit.TraceStart + ((StepFwdHit.TraceEnd - StepFwdHit.TraceStart) * StepFwdHit.Time)) - LocationInProgress;
 
@@ -909,7 +910,7 @@ bool UGroundMovementUtils::TestMoveToStepOver(const FMovingComponentSet& MovingC
 	// Step down
 	const FVector StepDownAdjustment = GravDir * StepTravelDownHeight;
 
-	const bool bDidStepDown = UAsyncMovementUtils::TestMoveComponent_Internal(MovingComps, LocationInProgress, LocationInProgress + StepDownAdjustment, Rotation, Rotation, /*bShouldSweep=*/ true, MOVECOMP_NoFlags, StepFwdHit);
+	const bool bDidStepDown = UAsyncMovementUtils::TestMoveComponent_Internal(MovingComps, LocationInProgress, LocationInProgress + StepDownAdjustment, Rotation, Rotation, /*bShouldSweep=*/ true, CollisionParams, StepFwdHit);
 
 
 	UE_LOG(LogMover, VeryVerbose, TEXT("TestMoveToStepOver Down: %s (role %i) StepDownAdjustment=%s DidMove=%i"),
diff --git a/Engine/Plugins/Experimental/Mover/Source/Mover/Private/MoveLibrary/MovementUtilsTypes.cpp b/Engine/Plugins/Experimental/Mover/Source/Mover/Private/MoveLibrary/MovementUtilsTypes.cpp
--- a/Engine/Plugins/Experimental/Mover/Source/Mover/Private/MoveLibrary/MovementUtilsTypes.cpp
+++ b/Engine/Plugins/Experimental/Mover/Source/Mover/Private/MoveLibrary/MovementUtilsTypes.cpp
@@ -29,3 +29,34 @@ FMovingComponentSet::FMovingComponentSet(UMoverComponent* InMoverComponent)
 		UpdatedPrimitive = Cast<UPrimitiveComponent>(UpdatedComponent);
 	}
 }
+
+
+static const FName DefaultCollisionTraceTag = "SweepTestMoverComponent";
+
+FMoverCollisionParams::FMoverCollisionParams(const USceneComponent* SceneComp)
+{
+	if (const UPrimitiveComponent* AsPrimitive = Cast<const UPrimitiveComponent>(SceneComp))
+	{
+		SetFromPrimitiveComponent(AsPrimitive);
+	}
+	else
+	{
+		// TODO: set up a line trace if SceneComp is not a primitive component
+		ensureMsgf(0, TEXT("Support for non-primitive components is not yet implemented"));
+	}
+}
+
+void FMoverCollisionParams::SetFromPrimitiveComponent(const UPrimitiveComponent* PrimitiveComp)
+{
+	Channel = PrimitiveComp->GetCollisionObjectType();
+
+	Shape = PrimitiveComp->GetCollisionShape();
+
+	PrimitiveComp->InitSweepCollisionParams(QueryParams, ResponseParams);
+
+	const AActor* OwningActor = PrimitiveComp->GetOwner();
+	
+	QueryParams.TraceTag = DefaultCollisionTraceTag;
+	QueryParams.OwnerTag = OwningActor->GetFName();
+	QueryParams.AddIgnoredActor(OwningActor);
+}
diff --git a/Engine/Plugins/Experimental/Mover/Source/Mover/Public/DefaultMovementSet/Modes/AsyncNavWalkingMode.h b/Engine/Plugins/Experimental/Mover/Source/Mover/Public/DefaultMovementSet/Modes/AsyncNavWalkingMode.h
--- a/Engine/Plugins/Experimental/Mover/Source/Mover/Public/DefaultMovementSet/Modes/AsyncNavWalkingMode.h
+++ b/Engine/Plugins/Experimental/Mover/Source/Mover/Public/DefaultMovementSet/Modes/AsyncNavWalkingMode.h
@@ -28,7 +28,6 @@ public:
 	UAsyncNavWalkingMode();
 
 	virtual void Activate() override;
-	virtual void Deactivate() override;
 	virtual void SimulationTick_Implementation(const FSimulationTickParams& Params, FMoverTickEndData& OutputState) override;
 	virtual void GenerateMove_Implementation(const FMoverTickStartData& StartState, const FMoverTimeStep& TimeStep, FProposedMove& OutProposedMove) const override;
 
@@ -39,7 +38,7 @@ public:
 	bool bSweepWhileNavWalking;
 
 	/** Whether to raycast to underlying geometry to better conform navmesh-walking actors */
-	UPROPERTY(Category = "NavMesh Movement", EditAnywhere, BlueprintReadOnly)
+	UPROPERTY(Category = "NavMesh Movement", EditAnywhere, BlueprintReadWrite)
 	bool bProjectNavMeshWalking;
 
 	/**
@@ -103,9 +102,6 @@ protected:
 	/** Optional modular object for generating rotation towards desired orientation. If not specified, linear interpolation will be used. */
 	UPROPERTY(EditAnywhere, Instanced, Category = Mover, meta = (MustImplement = "/Script/Mover.TurnGeneratorInterface"))
 	TObjectPtr<UObject> TurnGenerator;
-
-	/** Switch collision settings for NavWalking mode (ignore world collisions) */
-	virtual void SetCollisionForNavWalking(bool bEnable);
 	
 	/** Get Navigation data for the actor. Returns null if there is no associated nav data. */
 	const INavigationDataInterface* GetNavData() const;
@@ -128,15 +124,4 @@ protected:
 
 	TWeakObjectPtr<const UCommonLegacyMovementSettings> CommonLegacySettings;
  
- 	// Collision responses saved at the activation of this mode, to be restored as upon deactivation if the CDO can't be used
- 	ECollisionResponse CollideVsWorldStatic = ECollisionResponse::ECR_Block;
- 	ECollisionResponse CollideVsWorldDynamic = ECollisionResponse::ECR_Block;
-
-private:
-	// Game thread event callbacks used for manipulating collision on the actor when entering/exiting this mode
-	UFUNCTION()
-	void PostFinalizeSetCollisionForNavWalking(const FMoverSyncState& SyncState, const FMoverAuxStateContext& AuxState);
-	UFUNCTION()
-	void PostFinalizeRestoreOriginalCollision(const FMoverSyncState& SyncState, const FMoverAuxStateContext& AuxState);
-	
 };
diff --git a/Engine/Plugins/Experimental/Mover/Source/Mover/Public/MoveLibrary/AsyncMovementUtils.h b/Engine/Plugins/Experimental/Mover/Source/Mover/Public/MoveLibrary/AsyncMovementUtils.h
--- a/Engine/Plugins/Experimental/Mover/Source/Mover/Public/MoveLibrary/AsyncMovementUtils.h
+++ b/Engine/Plugins/Experimental/Mover/Source/Mover/Public/MoveLibrary/AsyncMovementUtils.h
@@ -5,6 +5,7 @@
 #include "Kismet/BlueprintFunctionLibrary.h"
 #include "Components/SceneComponent.h"
 #include "MoverDataModelTypes.h"
+#include "MovementUtilsTypes.h"
 
 #include "AsyncMovementUtils.generated.h"
 
@@ -31,6 +32,13 @@ public:
 	UFUNCTION(BlueprintCallable, Category = Mover)
 	static bool TestDepenetratingMove(const FMovingComponentSet& MovingComps, const FVector& StartLocation, const FVector& TargetLocation, const FQuat& StartRotation, const FQuat& TargetRotation, bool bShouldSweep, FHitResult& OutHit, FMovementRecord& InOutMoveRecord);
 
+	/** Tests potential movement of a component without actually moving it, taking penetration resolution issues into account first.  Relies on CollisionParams to describe the query.
+	 * Returns true if any movement was possible
+	 * Modifies OutHit with final movement hit data
+	 * Appends to InOutMoveRecord with any movement substeps
+	 */
+	static bool TestDepenetratingMove(const FMovingComponentSet& MovingComps, const FVector& StartLocation, const FVector& TargetLocation, const FQuat& StartRotation, const FQuat& TargetRotation, bool bShouldSweep, FMoverCollisionParams& CollisionParams, FHitResult& OutHit, FMovementRecord& InOutMoveRecord);
+
 	/** Tests potential movement of a component sliding along a surface, without actually moving it. 
 	 * Returns the percent of time applied, with 0.0 meaning no movement would occur. 
 	 * Modifies InOutHit with final movement hit data
@@ -38,17 +46,27 @@ public:
 	 */
 	UFUNCTION(BlueprintCallable, Category = Mover)
 	static float TestSlidingMoveAlongHitSurface(const FMovingComponentSet& MovingComps, const FVector& OriginalMoveDelta, const FVector& LocationAtHit, const FQuat& TargetRotation, FHitResult& InOutHit, FMovementRecord& InOutMoveRecord);
+	
+
+	/** Tests potential movement of a component sliding along a surface, without actually moving it. Relies on CollisionParams to describe the query.
+	 * Returns the percent of time applied, with 0.0 meaning no movement would occur. 
+	 * Modifies InOutHit with final movement hit data
+	 * Appends to InOutMoveRecord with any movement substeps
+	 */
+	static float TestSlidingMoveAlongHitSurface(const FMovingComponentSet& MovingComps, const FVector& OriginalMoveDelta, const FVector& LocationAtHit, const FQuat& TargetRotation, FMoverCollisionParams& CollisionParams, FHitResult& InOutHit, FMovementRecord& InOutMoveRecord);
 
 
 	/** Attempts to find a move that would resolve an initially penetrating blockage.
 	 * Returns true if an adjustment was found. The adjustment will be in OutAdjustmentDelta.
 	 */
-	static bool FindMoveToResolveInitialPenetration_Internal(const FMovingComponentSet& MovingComps, const FVector& StartLocation, const FQuat& TargetRotation, const FHitResult& PenetratingHit, EMoveComponentFlags MoveComponentFlags, FVector& OutAdjustmentDelta);
+	static bool FindMoveToResolveInitialPenetration_Internal(const FMovingComponentSet& MovingComps, const FVector& StartLocation, const FQuat& StartRotation, const FHitResult& PenetratingHit, FMoverCollisionParams& CollisionParams, FVector& OutAdjustmentDelta);
 
 
 	/** Tests a move of a component without actually moving it. 
 	 * Returns true if any motion could occur. Detailed blocking hit results are written to OutHit.  
 	 * If not sweeping, full movement will always be allowed even if new blocking overlaps would occur.
 	 */
-	static bool TestMoveComponent_Internal(const FMovingComponentSet& MovingComps, const FVector& StartLocation, const FVector& NewLocation, const FQuat& StartRotation, const FQuat& NewRotation, bool bShouldSweep, EMoveComponentFlags MoveCompFlags, FHitResult& OutHit);
+	static bool TestMoveComponent_Internal(const FMovingComponentSet& MovingComps, const FVector& StartLocation, const FVector& TargetLocation, const FQuat& StartRotation, const FQuat& TargetRotation, bool bShouldSweep, FMoverCollisionParams& CollisionParams, FHitResult& OutHit);
+
+	
 };
diff --git a/Engine/Plugins/Experimental/Mover/Source/Mover/Public/MoveLibrary/MovementUtilsTypes.h b/Engine/Plugins/Experimental/Mover/Source/Mover/Public/MoveLibrary/MovementUtilsTypes.h
--- a/Engine/Plugins/Experimental/Mover/Source/Mover/Public/MoveLibrary/MovementUtilsTypes.h
+++ b/Engine/Plugins/Experimental/Mover/Source/Mover/Public/MoveLibrary/MovementUtilsTypes.h
@@ -2,9 +2,11 @@
 
 #pragma once
 
+#include "CollisionShape.h"
+#include "CollisionQueryParams.h"
+#include "Components/SceneComponent.h"
 #include "MovementUtilsTypes.generated.h"
 
-class USceneComponent;
 class UPrimitiveComponent;
 class UMoverComponent;
 
@@ -82,4 +84,27 @@ struct MOVER_API FMovingComponentSet
 
 	UPROPERTY(EditAnywhere, BlueprintReadWrite, Category = Mover)
 	TWeakObjectPtr<UMoverComponent> MoverComponent;
-};
\ No newline at end of file
+};
+
+
+/**
+ * Encapsulates a collision context involved in movement. Useful for altering queries without changing the component.
+ */
+struct MOVER_API FMoverCollisionParams
+{
+public: 
+	FMoverCollisionParams() {}
+	FMoverCollisionParams(const USceneComponent* SceneComp);
+
+	void SetFromPrimitiveComponent(const UPrimitiveComponent* PrimitiveComp);
+
+	FCollisionShape Shape;
+
+	ECollisionChannel Channel = ECollisionChannel::ECC_Pawn;
+
+	FCollisionQueryParams QueryParams;
+
+	FCollisionResponseParams ResponseParams;
+
+	EMoveComponentFlags MoveComponentFlags = MOVECOMP_NoFlags;
+};
