[Mover] [BugFix] Make offset calculations in PhysicsStanceModifiers resim friendly

diff --git a/Engine/Plugins/Experimental/Mover/Source/Mover/Private/PhysicsMover/MovementModifiers/PhysicsStanceModifier.cpp b/Engine/Plugins/Experimental/Mover/Source/Mover/Private/PhysicsMover/MovementModifiers/PhysicsStanceModifier.cpp
--- a/Engine/Plugins/Experimental/Mover/Source/Mover/Private/PhysicsMover/MovementModifiers/PhysicsStanceModifier.cpp
+++ b/Engine/Plugins/Experimental/Mover/Source/Mover/Private/PhysicsMover/MovementModifiers/PhysicsStanceModifier.cpp
@@ -10,19 +10,26 @@
 #include "PhysicsMover/Modes/PhysicsDrivenFallingMode.h"
 
 
-void FPhysicsStanceModifier::AdjustCapsule(UMoverComponent* MoverComp, float OldHalfHeight, float NewHalfHeight, float NewEyeHeight)
-{
-	const float HalfHeightDifference = FMath::Abs(NewHalfHeight - OldHalfHeight);
-	const bool bExpanding = OldHalfHeight < NewHalfHeight;
-
-	// Adjust the TargetHeight of all the relevant MovementModes
-	UpdateTargetHeight<UPhysicsDrivenWalkingMode>(MoverComp, OldHalfHeight, NewHalfHeight);
-	UpdateTargetHeight<UPhysicsDrivenFallingMode>(MoverComp, OldHalfHeight, NewHalfHeight);
-}
-
 void FPhysicsStanceModifier::OnStart(UMoverComponent* MoverComp, const FMoverTimeStep& TimeStep, const FMoverSyncState& SyncState, const FMoverAuxStateContext& AuxState)
 {
-	Super::OnStart(MoverComp, TimeStep, SyncState, AuxState);
+	if (UStanceSettings* StanceSettings = MoverComp->FindSharedSettings_Mutable<UStanceSettings>())
+	{
+		if (const UCapsuleComponent* CapsuleComponent = Cast<UCapsuleComponent>(MoverComp->GetUpdatedComponent()))
+		{
+			if (const UCapsuleComponent* OriginalCapsule = UMovementUtils::GetOriginalComponentType<UCapsuleComponent>(MoverComp->GetOwner()))
+			{
+				const float OriginalHalfHeight = OriginalCapsule->GetScaledCapsuleHalfHeight();
+				const float CrouchedHalfHeight = StanceSettings->CrouchHalfHeight;
+				const float CrouchedEyeHeight = StanceSettings->CrouchedEyeHeight;
+				
+				// Adjust the TargetHeight of all the relevant MovementModes
+				UpdateTargetHeight<UPhysicsDrivenWalkingMode>(MoverComp, OriginalHalfHeight, CrouchedHalfHeight, true);
+				UpdateTargetHeight<UPhysicsDrivenFallingMode>(MoverComp, OriginalHalfHeight, CrouchedHalfHeight, true);
+
+				ApplyMovementSettings(MoverComp);
+			}
+		}
+	}
 
 	// Ensures crouching, especially if triggered through state syncing (rollbacks, etc.)  
 	if (UCharacterMoverComponent* CharMoverComp = Cast<UCharacterMoverComponent>(MoverComp))
@@ -33,7 +40,23 @@ void FPhysicsStanceModifier::OnStart(UMoverComponent* MoverComp, const FMoverTim
 
 void FPhysicsStanceModifier::OnEnd(UMoverComponent* MoverComp, const FMoverTimeStep& TimeStep, const FMoverSyncState& SyncState, const FMoverAuxStateContext& AuxState)
 {
-	Super::OnEnd(MoverComp, TimeStep, SyncState, AuxState);
+	if (UStanceSettings* StanceSettings = MoverComp->FindSharedSettings_Mutable<UStanceSettings>())
+	{
+		if (UCapsuleComponent* CapsuleComponent = Cast<UCapsuleComponent>(MoverComp->GetUpdatedComponent()))
+		{
+			if (const UCapsuleComponent* OriginalCapsule = UMovementUtils::GetOriginalComponentType<UCapsuleComponent>(MoverComp->GetOwner()))
+			{
+				const float OriginalHalfHeight = OriginalCapsule->GetScaledCapsuleHalfHeight();
+				const float CrouchedHalfHeight = StanceSettings->CrouchHalfHeight;
+
+				// Adjust the TargetHeight of all the relevant MovementModes
+				UpdateTargetHeight<UPhysicsDrivenWalkingMode>(MoverComp, OriginalHalfHeight, CrouchedHalfHeight, false);
+				UpdateTargetHeight<UPhysicsDrivenFallingMode>(MoverComp, OriginalHalfHeight, CrouchedHalfHeight, false);
+
+				RevertMovementSettings(MoverComp);
+			}
+		}
+	}
 
 	// Ensures uncrouching, especially if triggered through state syncing (rollbacks, etc.)  
 	if (UCharacterMoverComponent* CharMoverComp = Cast<UCharacterMoverComponent>(MoverComp))
diff --git a/Engine/Plugins/Experimental/Mover/Source/Mover/Public/PhysicsMover/MovementModifiers/PhysicsStanceModifier.h b/Engine/Plugins/Experimental/Mover/Source/Mover/Public/PhysicsMover/MovementModifiers/PhysicsStanceModifier.h
--- a/Engine/Plugins/Experimental/Mover/Source/Mover/Public/PhysicsMover/MovementModifiers/PhysicsStanceModifier.h
+++ b/Engine/Plugins/Experimental/Mover/Source/Mover/Public/PhysicsMover/MovementModifiers/PhysicsStanceModifier.h
@@ -37,18 +37,18 @@ struct MOVER_API FPhysicsStanceModifier : public FStanceModifier
 	virtual void OnEnd(UMoverComponent* MoverComp, const FMoverTimeStep& TimeStep, const FMoverSyncState& SyncState, const FMoverAuxStateContext& AuxState) override;
 
 protected:
-	// Modifies the Capsule Location by updating the target height of the movement modes
-	virtual void AdjustCapsule(UMoverComponent* MoverComp, float OldHalfHeight, float NewHalfHeight, float NewEyeHeight) override;
 
 	template <class T>
-	void UpdateTargetHeight(UMoverComponent* MoverComp, float OldHalfHeight, float NewHalfHeight)
+	void UpdateTargetHeight(UMoverComponent* MoverComp, float OriginalHalfHeight, float CrouchedHalfHeight, bool bCrouching)
 	{
 		if (T* MovementMode = MoverComp->FindMode_Mutable<T>())
 		{
 			if (MovementMode->template Implements<UPhysicsCharacterMovementModeInterface>())
 			{
-				float TargetHeight = MovementMode->GetTargetHeight();
-				MovementMode->SetTargetHeightOverride(TargetHeight + (NewHalfHeight - OldHalfHeight));
+				const float OriginalTargetHeight = GetDefault<T>(MovementMode->GetClass())->GetTargetHeight();
+				const float Offset = OriginalTargetHeight - OriginalHalfHeight;
+				
+				MovementMode->SetTargetHeightOverride(Offset + (bCrouching ? CrouchedHalfHeight : OriginalHalfHeight));
 			}
 		}
 	}
