Disable automatic batching on Mover apply ticks as they cannot be batched efficiently due to dependencies

diff --git a/Engine/Plugins/Experimental/Mover/Source/Mover/Private/Backends/MoverStandaloneLiaison.cpp b/Engine/Plugins/Experimental/Mover/Source/Mover/Private/Backends/MoverStandaloneLiaison.cpp
--- a/Engine/Plugins/Experimental/Mover/Source/Mover/Private/Backends/MoverStandaloneLiaison.cpp
+++ b/Engine/Plugins/Experimental/Mover/Source/Mover/Private/Backends/MoverStandaloneLiaison.cpp
@@ -74,7 +74,6 @@ UMoverStandaloneLiaisonComponent::UMoverStandaloneLiaisonComponent(const FObject
 	ApplyStateTickFunction.bStartWithTickEnabled = true;
 	ApplyStateTickFunction.SetTickFunctionEnable(true);
 	ApplyStateTickFunction.TickGroup = TG_PrePhysics;
-	ApplyStateTickFunction.bAllowTickBatching = true;
 
 	bWantsInitializeComponent = true;
 	bAutoActivate = true;
