[Mover] Add internal safe functions for Physics Mover Component and other thread safety changes

diff --git a/Engine/Plugins/Experimental/Mover/Source/Mover/Private/PhysicsMover/PhysicsCharacterMoverComponent.cpp b/Engine/Plugins/Experimental/Mover/Source/Mover/Private/PhysicsMover/PhysicsCharacterMoverComponent.cpp
--- a/Engine/Plugins/Experimental/Mover/Source/Mover/Private/PhysicsMover/PhysicsCharacterMoverComponent.cpp
+++ b/Engine/Plugins/Experimental/Mover/Source/Mover/Private/PhysicsMover/PhysicsCharacterMoverComponent.cpp
@@ -33,13 +33,138 @@ void UPhysicsCharacterMoverComponent::BeginPlay()
 	}
 }
 
+const FMovementModifierBase* UPhysicsCharacterMoverComponent::FindMovementModifier_Internal(const FMoverSyncState& SyncState, const FMovementModifierHandle& ModifierHandle) const
+{
+	Chaos::EnsureIsInPhysicsThreadContext();
+	
+	// Check active modifiers for modifier handle
+	for (auto ActiveModifierFromSyncStateIt = SyncState.MovementModifiers.GetActiveModifiersIterator(); ActiveModifierFromSyncStateIt; ++ActiveModifierFromSyncStateIt)
+	{
+		const TSharedPtr<FMovementModifierBase> ActiveModifierFromSyncState = *ActiveModifierFromSyncStateIt;
+
+		if (ModifierHandle == ActiveModifierFromSyncState->GetHandle())
+		{
+			return ActiveModifierFromSyncState.Get();
+		}
+	}
+
+	// Check queued modifiers for modifier handle
+	for (auto QueuedModifierFromSyncStateIt = SyncState.MovementModifiers.GetQueuedModifiersIterator(); QueuedModifierFromSyncStateIt; ++QueuedModifierFromSyncStateIt)
+	{
+		const TSharedPtr<FMovementModifierBase> QueuedModifierFromSyncState = *QueuedModifierFromSyncStateIt;
+
+		if (ModifierHandle == QueuedModifierFromSyncState->GetHandle())
+		{
+			return QueuedModifierFromSyncState.Get();
+		}
+	}
+
+	return nullptr;
+}
+
+const FMovementModifierBase* UPhysicsCharacterMoverComponent::FindMovementModifierByType_Internal(const FMoverSyncState& SyncState, const UScriptStruct* DataStructType) const
+{
+	Chaos::EnsureIsInPhysicsThreadContext();
+
+	// Check active modifiers for modifier handle
+	for (auto ActiveModifierFromSyncStateIt = SyncState.MovementModifiers.GetActiveModifiersIterator(); ActiveModifierFromSyncStateIt; ++ActiveModifierFromSyncStateIt)
+	{
+		const TSharedPtr<FMovementModifierBase> ActiveModifierFromSyncState = *ActiveModifierFromSyncStateIt;
+
+		if (DataStructType == ActiveModifierFromSyncState->GetScriptStruct())
+		{
+			return ActiveModifierFromSyncState.Get();
+		}
+	}
+
+	// Check queued modifiers for modifier handle
+	for (auto QueuedModifierFromSyncStateIt = SyncState.MovementModifiers.GetQueuedModifiersIterator(); QueuedModifierFromSyncStateIt; ++QueuedModifierFromSyncStateIt)
+	{
+		const TSharedPtr<FMovementModifierBase> QueuedModifierFromSyncState = *QueuedModifierFromSyncStateIt;
+
+		if (DataStructType == QueuedModifierFromSyncState->GetScriptStruct())
+		{
+			return QueuedModifierFromSyncState.Get();
+		}
+	}
+
+	return nullptr;
+}
+
+bool UPhysicsCharacterMoverComponent::HasGameplayTag_Internal(const FMoverSyncState& SyncState, FGameplayTag TagToFind, bool bExactMatch) const
+{
+	Chaos::EnsureIsInPhysicsThreadContext();
+
+	if (bExactMatch)
+	{
+		if (ExternalGameplayTags.HasTagExact(TagToFind))
+		{
+			return true;
+		}
+	}
+	else
+	{
+		if (ExternalGameplayTags.HasTag(TagToFind))
+		{
+			return true;
+		}
+	}
+
+	// Search Movement Modes
+	if (const UBaseMovementMode* ActiveMovementMode = GetMovementMode())
+	{
+		if (ActiveMovementMode->HasGameplayTag(TagToFind, bExactMatch))
+		{
+			return true;
+		}
+	}
+
+	// Search Movement Modifiers
+	for (auto ModifierFromSyncStateIt = SyncState.MovementModifiers.GetActiveModifiersIterator(); ModifierFromSyncStateIt; ++ModifierFromSyncStateIt)
+	{
+		const TSharedPtr<FMovementModifierBase> ModifierFromSyncState = *ModifierFromSyncStateIt;
+		if (ModifierFromSyncState->HasGameplayTag(TagToFind, bExactMatch))
+		{
+			return true;
+		}
+	}
+
+	// Search Layered Moves
+	for (const TSharedPtr<FLayeredMoveBase>& LayeredMove : SyncState.LayeredMoves.GetActiveMoves())
+	{
+		if (LayeredMove->HasGameplayTag(TagToFind, bExactMatch))
+		{
+			return true;
+		}
+	}
+
+	return false;
+}
+
+void UPhysicsCharacterMoverComponent::Crouch_Internal(const FMoverSyncState& SyncState)
+{
+	Chaos::EnsureIsInPhysicsThreadContext();
+
+	if (CanCrouch_Internal(SyncState))
+	{
+		bWantsToCrouch = true;
+	}
+}
+
+void UPhysicsCharacterMoverComponent::UnCrouch_Internal(const FMoverSyncState& SyncState)
+{
+	Chaos::EnsureIsInPhysicsThreadContext();
+
+	bWantsToCrouch = false;
+}
+
 void UPhysicsCharacterMoverComponent::OnMoverPreMovement(const FMoverTimeStep& TimeStep, const FMoverInputCmdContext& InputCmd, const FMoverSyncState& SyncState, const FMoverAuxStateContext& AuxState)
 {
-	const FPhysicsStanceModifier* StanceModifier = static_cast<const FPhysicsStanceModifier*>(FindMovementModifier(StanceModifierHandle));
+	const FPhysicsStanceModifier* StanceModifier = static_cast<const FPhysicsStanceModifier*>(FindMovementModifier_Internal(SyncState, StanceModifierHandle));
 	// This is a fail safe in case our handle was bad - try finding the modifier by type if we can
 	if (!StanceModifier)
 	{
-		StanceModifier = FindMovementModifierByType<FPhysicsStanceModifier>();
+		StanceModifier = FindMovementModifierByType_Internal<FPhysicsStanceModifier>(SyncState);
 	}
 
 	EStanceMode OldActiveStance = EStanceMode::Invalid;
@@ -48,8 +173,8 @@ void UPhysicsCharacterMoverComponent::OnMoverPreMovement(const FMoverTimeStep& T
 		OldActiveStance = StanceModifier->ActiveStance;
 	}
 
-	const bool bIsCrouching = HasGameplayTag(Mover_IsCrouching, true);
-	if (bIsCrouching && (!bWantsToCrouch || !CanCrouch()))
+	const bool bIsCrouching = HasGameplayTag_Internal(SyncState, Mover_IsCrouching, true);
+	if (bIsCrouching && (!bWantsToCrouch || !CanCrouch_Internal(SyncState)))
 	{
 		if (StanceModifier && StanceModifier->CanExpand_Internal(this, UpdatedComponent, SyncState))
 		{
@@ -59,7 +184,7 @@ void UPhysicsCharacterMoverComponent::OnMoverPreMovement(const FMoverTimeStep& T
 			StanceModifier = nullptr;
 		}
 	}
-	else if (!bIsCrouching && bWantsToCrouch && CanCrouch())
+	else if (!bIsCrouching && bWantsToCrouch && CanCrouch_Internal(SyncState))
 	{
 		TSharedPtr<FPhysicsStanceModifier> NewStanceModifier = MakeShared<FPhysicsStanceModifier>();
 		StanceModifierHandle = QueueMovementModifier(NewStanceModifier);
diff --git a/Engine/Plugins/Experimental/Mover/Source/Mover/Public/PhysicsMover/PhysicsCharacterMoverComponent.h b/Engine/Plugins/Experimental/Mover/Source/Mover/Public/PhysicsMover/PhysicsCharacterMoverComponent.h
--- a/Engine/Plugins/Experimental/Mover/Source/Mover/Public/PhysicsMover/PhysicsCharacterMoverComponent.h
+++ b/Engine/Plugins/Experimental/Mover/Source/Mover/Public/PhysicsMover/PhysicsCharacterMoverComponent.h
@@ -17,6 +17,30 @@ public:
 	
 	virtual void BeginPlay() override;
 
+	// Internal Queries on Sync State ////////////
+
+	// Find movement modifier by it's handle. Returns nullptr if the modifier couldn't be found
+	const FMovementModifierBase* FindMovementModifier_Internal(const FMoverSyncState& SyncState, const FMovementModifierHandle& ModifierHandle) const;
+
+	// Find movement modifier by type (returns the first modifier it finds). Returns nullptr if the modifier couldn't be found
+	const FMovementModifierBase* FindMovementModifierByType_Internal(const FMoverSyncState& SyncState, const UScriptStruct* DataStructType) const;
+
+	// Find a movement modifier of a specific type in this components movement modifiers. If not found, null will be returned.
+	template <typename ModifierT = FMovementModifierBase UE_REQUIRES(std::is_base_of_v<FMovementModifierBase, ModifierT>)>
+	const ModifierT* FindMovementModifierByType_Internal(const FMoverSyncState& SyncState) const { return static_cast<const ModifierT*>(SyncState, FindMovementModifierByType(ModifierT::StaticStruct())); }
+
+	// Check Mover systems for a gameplay tag.
+	bool HasGameplayTag_Internal(const FMoverSyncState& SyncState, FGameplayTag TagToFind, bool bExactMatch) const;
+
+	// Check if crouching is allowed
+	virtual bool CanCrouch_Internal(const FMoverSyncState& SyncState) { return true; }
+
+	// Perform crouch on actor
+	void Crouch_Internal(const FMoverSyncState& SyncState);
+	
+	// Perform uncrouch on actor
+	void UnCrouch_Internal(const FMoverSyncState& SyncState);
+
 protected:
 
 	UFUNCTION()
