[Mover] Horde fix

diff --git a/Engine/Plugins/Experimental/Mover/Source/Mover/Public/PhysicsMover/PhysicsCharacterMoverComponent.h b/Engine/Plugins/Experimental/Mover/Source/Mover/Public/PhysicsMover/PhysicsCharacterMoverComponent.h
--- a/Engine/Plugins/Experimental/Mover/Source/Mover/Public/PhysicsMover/PhysicsCharacterMoverComponent.h
+++ b/Engine/Plugins/Experimental/Mover/Source/Mover/Public/PhysicsMover/PhysicsCharacterMoverComponent.h
@@ -27,7 +27,7 @@ public:
 
 	// Find a movement modifier of a specific type in this components movement modifiers. If not found, null will be returned.
 	template <typename ModifierT = FMovementModifierBase UE_REQUIRES(std::is_base_of_v<FMovementModifierBase, ModifierT>)>
-	const ModifierT* FindMovementModifierByType_Internal(const FMoverSyncState& SyncState) const { return static_cast<const ModifierT*>(SyncState, FindMovementModifierByType(ModifierT::StaticStruct())); }
+	const ModifierT* FindMovementModifierByType_Internal(const FMoverSyncState& SyncState) const { return static_cast<const ModifierT*>(FindMovementModifierByType_Internal(SyncState, ModifierT::StaticStruct())); }
 
 	// Check Mover systems for a gameplay tag.
 	bool HasGameplayTag_Internal(const FMoverSyncState& SyncState, FGameplayTag TagToFind, bool bExactMatch) const;
