Calculating the floor distance when landing based on the dot product between the height and the normal of the floor, instead of using the ground normal.
#rnx

diff --git a/Engine/Plugins/Experimental/Mover/Source/Mover/Private/PhysicsMover/Modes/PhysicsDrivenFallingMode.cpp b/Engine/Plugins/Experimental/Mover/Source/Mover/Private/PhysicsMover/Modes/PhysicsDrivenFallingMode.cpp
--- a/Engine/Plugins/Experimental/Mover/Source/Mover/Private/PhysicsMover/Modes/PhysicsDrivenFallingMode.cpp
+++ b/Engine/Plugins/Experimental/Mover/Source/Mover/Private/PhysicsMover/Modes/PhysicsDrivenFallingMode.cpp
@@ -150,9 +150,10 @@ void UPhysicsDrivenFallingMode::SimulationTick_Implementation(const FSimulationT
 	}
 
 	constexpr float FloorDistanceTolerance = 2.0f;
+	const float FloorDistanceWithFloorNormal = FloorResult.HitResult.ImpactNormal.Dot(StartingSyncState->GetLocation_WorldSpace() - FloorResult.HitResult.ImpactPoint);
 	const FVector ProjectedGroundVelocity = UPhysicsMovementUtils::ComputeIntegratedGroundVelocityFromHitResult(StartingSyncState->GetLocation_WorldSpace(), FloorResult.HitResult, DeltaSeconds);
 	const float ProjectedRelativeVerticalVelocity = FloorResult.HitResult.ImpactNormal.Dot(ProposedMove.LinearVelocity - ProjectedGroundVelocity);
-	const float ProjectedFloorDistance = FloorResult.FloorDist + UpDir.Dot(ProposedMove.LinearVelocity - ProjectedGroundVelocity) * DeltaSeconds;
+	const float ProjectedFloorDistance = FloorDistanceWithFloorNormal + ProjectedRelativeVerticalVelocity * DeltaSeconds;
 	const bool bIsFloorWithinReach = ProjectedFloorDistance < TargetHeight + FloorDistanceTolerance;
 	const bool bIsMovingUpRelativeToFloor = ProjectedRelativeVerticalVelocity > UE_KINDA_SMALL_NUMBER;
 
