Mover: Editing NavMoverComponent to remove Tick and adding some uint64s to help keep track of when input is requested and consumed to help with caching and clearing cached nav input values.

diff --git a/Engine/Plugins/Experimental/Mover/Source/Mover/Private/DefaultMovementSet/NavMoverComponent.cpp b/Engine/Plugins/Experimental/Mover/Source/Mover/Private/DefaultMovementSet/NavMoverComponent.cpp
--- a/Engine/Plugins/Experimental/Mover/Source/Mover/Private/DefaultMovementSet/NavMoverComponent.cpp
+++ b/Engine/Plugins/Experimental/Mover/Source/Mover/Private/DefaultMovementSet/NavMoverComponent.cpp
@@ -13,8 +13,6 @@ UNavMoverComponent::UNavMoverComponent()
 {
 	bWantsInitializeComponent = true;
 	bAutoActivate = true;
-	PrimaryComponentTick.bCanEverTick = true;
-	PrimaryComponentTick.bStartWithTickEnabled = true;
 }
 
 void UNavMoverComponent::InitializeComponent()
@@ -46,25 +44,6 @@ void UNavMoverComponent::BeginPlay()
 	}
 }
 
-void UNavMoverComponent::TickComponent(float DeltaTime, ELevelTick TickType, FActorComponentTickFunction* ThisTickFunction)
-{
-	Super::TickComponent(DeltaTime, TickType, ThisTickFunction);
-
-	if (bShouldClearNavMoveInputs)
-	{
-		CachedNavMoveInputIntent = FVector::ZeroVector;
-		CachedNavMoveInputVelocity = FVector::ZeroVector;
-		
-		bRequestedNavMovement = false;
-		bShouldClearNavMoveInputs = false;
-	}
-}
-
-bool UNavMoverComponent::HasRequestedNavMovement() const
-{
-	return bRequestedNavMovement;
-}
-
 float UNavMoverComponent::GetMaxSpeedForNavMovement() const
 {
 	float MaxSpeed = 0.0f;
@@ -94,19 +73,28 @@ void UNavMoverComponent::StopMovementImmediately()
 
 bool UNavMoverComponent::ConsumeNavMovementData(FVector& OutMoveInputIntent, FVector& OutMoveInputVelocity)
 {
-	if (bRequestedNavMovement)
+	const bool bHasFrameAdvanced = GFrameCounter > GameFrameNavMovementConsumed;
+	const bool bNoNewRequests = GameFrameNavMovementConsumed > GameFrameNavMovementRequested;
+	bool bHasNavMovement = false;
+	
+	if (bHasFrameAdvanced && bNoNewRequests)
+	{
+		CachedNavMoveInputIntent = FVector::ZeroVector;
+		CachedNavMoveInputVelocity = FVector::ZeroVector;
+	}
+	else
 	{
-		UE_LOG(LogMover, VeryVerbose, TEXT("Applying %s as NavMoveInputIntent."), *CachedNavMoveInputIntent.ToString());
-		UE_LOG(LogMover, VeryVerbose, TEXT("Applying %s as NavMoveInputVelocity."), *CachedNavMoveInputVelocity.ToString());
 		OutMoveInputIntent = CachedNavMoveInputIntent;
 		OutMoveInputVelocity = CachedNavMoveInputVelocity;
+		bHasNavMovement = true;
 		
-		bShouldClearNavMoveInputs = true;
-		
-		return true;
+		UE_LOG(LogMover, VeryVerbose, TEXT("Applying %s as NavMoveInputIntent."), *CachedNavMoveInputIntent.ToString());
+		UE_LOG(LogMover, VeryVerbose, TEXT("Applying %s as NavMoveInputVelocity."), *CachedNavMoveInputVelocity.ToString());
 	}
 
-	return false;
+	GameFrameNavMovementConsumed = GFrameCounter;
+
+	return bHasNavMovement;
 }
 
 FVector UNavMoverComponent::GetLocation() const
@@ -203,8 +191,7 @@ void UNavMoverComponent::RequestDirectMove(const FVector& MoveVelocity, bool bFo
 		return;
 	}
 
-	bRequestedNavMovement = true;
-	bShouldClearNavMoveInputs = false;
+	GameFrameNavMovementRequested = GFrameCounter;
 	
 	if (IsFalling())
 	{
@@ -236,8 +223,7 @@ void UNavMoverComponent::RequestPathMove(const FVector& MoveInput)
 		AdjustedMoveInput = MoveInput.GetSafeNormal2D() * Mag;
 	}
 
-	bRequestedNavMovement = true;
-	bShouldClearNavMoveInputs = false;
+	GameFrameNavMovementRequested = GFrameCounter;
 	CachedNavMoveInputIntent = AdjustedMoveInput.GetSafeNormal();
 }
 
diff --git a/Engine/Plugins/Experimental/Mover/Source/Mover/Public/DefaultMovementSet/NavMoverComponent.h b/Engine/Plugins/Experimental/Mover/Source/Mover/Public/DefaultMovementSet/NavMoverComponent.h
--- a/Engine/Plugins/Experimental/Mover/Source/Mover/Public/DefaultMovementSet/NavMoverComponent.h
+++ b/Engine/Plugins/Experimental/Mover/Source/Mover/Public/DefaultMovementSet/NavMoverComponent.h
@@ -37,9 +37,11 @@ protected:
 	UPROPERTY(EditAnywhere, BlueprintReadWrite, Category="Nav Movement")
 	FNavMovementProperties NavMovementProperties;
 
-	/** bool for keeping track of requested nav movement - is set to true when movement is requested, is false once nav movement was consumed */
-	UPROPERTY(BlueprintReadOnly, Category="Nav Movement")
-	bool bRequestedNavMovement = false;
+	/** Keeps track of the last game frame we consumed nav movement. */
+	uint64 GameFrameNavMovementConsumed = 0;
+
+	/** Keeps track of the last game frame we requested nav movement. */
+	uint64 GameFrameNavMovementRequested = 0;
 	
 	// Used to store movement input intent from requested moves
 	FVector CachedNavMoveInputIntent = FVector::ZeroVector;
@@ -55,17 +57,10 @@ private:
 	/** Associated Movement component that will actually move the actor */ 
 	UPROPERTY()
 	TWeakObjectPtr<UMoverComponent> MoverComponent;
-
-	UPROPERTY(Transient)
-	bool bShouldClearNavMoveInputs = false;
 	
 public:
 	virtual void InitializeComponent() override;
 	virtual void BeginPlay() override;
-
-	virtual void TickComponent(float DeltaTime, ELevelTick TickType, FActorComponentTickFunction* ThisTickFunction) override;
-
-	bool HasRequestedNavMovement() const;
 	
 	/** Get the owner of the object consuming nav movement */
 	virtual UObject* GetOwnerAsObject() const override { return GetOwner(); }
