[Mover] HasGameplayTag crash

diff --git a/Engine/Plugins/Experimental/Mover/Source/Mover/Private/Backends/MoverNetworkPhysicsLiaison.cpp b/Engine/Plugins/Experimental/Mover/Source/Mover/Private/Backends/MoverNetworkPhysicsLiaison.cpp
--- a/Engine/Plugins/Experimental/Mover/Source/Mover/Private/Backends/MoverNetworkPhysicsLiaison.cpp
+++ b/Engine/Plugins/Experimental/Mover/Source/Mover/Private/Backends/MoverNetworkPhysicsLiaison.cpp
@@ -174,7 +174,7 @@ void UMoverNetworkPhysicsLiaisonComponent::PostPhysicsUpdate_External()
 {
 	const UMoverComponent& MoverComp = GetMoverComponent();
 	const UBaseMovementMode* PrevMode = MoverComp.GetMovementMode();
-	bool bWasFalling = MoverComp.HasGameplayTag(Mover_IsFalling, true);
+	bool bWasFalling = PrevMode ? PrevMode->HasGameplayTag(Mover_IsFalling, true) : false;
 
 	Super::PostPhysicsUpdate_External();
 
diff --git a/Engine/Plugins/Experimental/Mover/Source/Mover/Private/MoverComponent.cpp b/Engine/Plugins/Experimental/Mover/Source/Mover/Private/MoverComponent.cpp
--- a/Engine/Plugins/Experimental/Mover/Source/Mover/Private/MoverComponent.cpp
+++ b/Engine/Plugins/Experimental/Mover/Source/Mover/Private/MoverComponent.cpp
@@ -720,10 +720,12 @@ bool UMoverComponent::HasGameplayTag(FGameplayTag TagToFind, bool bExactMatch) c
 		// Search Movement Modifiers
 		for (auto ModifierFromSyncStateIt = CachedLastSyncState.MovementModifiers.GetActiveModifiersIterator(); ModifierFromSyncStateIt; ++ModifierFromSyncStateIt)
 		{
-			const TSharedPtr<FMovementModifierBase> ModifierFromSyncState = *ModifierFromSyncStateIt;
-			if (ModifierFromSyncState->HasGameplayTag(TagToFind, bExactMatch))
+			if (const TSharedPtr<FMovementModifierBase> ModifierFromSyncState = *ModifierFromSyncStateIt)
 			{
-				return true;
+				if (ModifierFromSyncState.IsValid() && ModifierFromSyncState->HasGameplayTag(TagToFind, bExactMatch))
+				{
+					return true;
+				}
 			}
 		}
 		
