[Pathed Chaos Mover] Set the default max force on the platform constraint to be much stronger. That keeps it on the path while also allowing it to move a bit when a 150kg character jumps onto it

diff --git a/Engine/Plugins/Experimental/Mover/Source/Mover/Private/PhysicsMover/PathedMovement/PathedMovementMode.cpp b/Engine/Plugins/Experimental/Mover/Source/Mover/Private/PhysicsMover/PathedMovement/PathedMovementMode.cpp
--- a/Engine/Plugins/Experimental/Mover/Source/Mover/Private/PhysicsMover/PathedMovement/PathedMovementMode.cpp
+++ b/Engine/Plugins/Experimental/Mover/Source/Mover/Private/PhysicsMover/PathedMovement/PathedMovementMode.cpp
@@ -18,8 +18,7 @@ FAutoConsoleVariable CVarEnablePathedPhysicsMovementDebug(
 
 UPathedPhysicsMovementMode::UPathedPhysicsMovementMode()
 {
-	//@todo DanH: These settings result in a joint that isn't nearly rigid/string/stiff enough to be the default.
-	//	Need to better understand the knobs as cranking up values ad hoc doesn't seem to have much effect.
+	//@todo DanH: A number of these should be able to auto-configure based on mass and gravity 
 	
 	JointConstraintProperties.bLinearBreakable = false;
 	JointConstraintProperties.bAngularBreakable = false;
@@ -32,16 +31,18 @@ UPathedPhysicsMovementMode::UPathedPhysicsMovementMode()
 	JointConstraintProperties.LinearLimit.ZMotion = LCM_Free;
 	JointConstraintProperties.LinearLimit.Limit = DefaultMaxLinearDistance;
 
+	static constexpr float DefaultDriveMaxForce = 5000.f;
+	
 	// Linear drive config
 	static constexpr float LinearDriveStiffness = 750.f;
 	static const float LinearDriveCriticalDamping = 2 * FMath::Sqrt(LinearDriveStiffness);
-
+	
 	FConstraintDrive DefaultLinearDriveConstraint;
 	DefaultLinearDriveConstraint.bEnablePositionDrive = true;
 	DefaultLinearDriveConstraint.bEnableVelocityDrive = true;
 	DefaultLinearDriveConstraint.Stiffness = LinearDriveStiffness;
 	DefaultLinearDriveConstraint.Damping = LinearDriveCriticalDamping;
-	DefaultLinearDriveConstraint.MaxForce = 100.f;
+	DefaultLinearDriveConstraint.MaxForce = DefaultDriveMaxForce;
 
 	JointConstraintProperties.LinearDrive.XDrive = DefaultLinearDriveConstraint;
 	JointConstraintProperties.LinearDrive.YDrive = DefaultLinearDriveConstraint;
@@ -65,7 +66,7 @@ UPathedPhysicsMovementMode::UPathedPhysicsMovementMode()
 	JointConstraintProperties.AngularDrive.SlerpDrive.bEnableVelocityDrive = true;
 	JointConstraintProperties.AngularDrive.SlerpDrive.Stiffness = AngularDriveStiffness;
 	JointConstraintProperties.AngularDrive.SlerpDrive.Damping = AngularDriveCriticalDamping;
-	JointConstraintProperties.AngularDrive.SlerpDrive.MaxForce = 100.f;
+	JointConstraintProperties.AngularDrive.SlerpDrive.MaxForce = DefaultDriveMaxForce;
 }
 
 #if WITH_EDITOR
