CVD: Fix to the mover info displaying character ground constraints when those are selected (automatically or not). This bug would cause the cahracter ground constraint to be displayed on the mover info tab when directly selecting a constraint, or when stepping the recording with automatically select ground constraint it selected.

diff --git a/Engine/Plugins/ChaosVD/Source/ChaosVD/Public/ChaosVDObjectDetailsTab.h b/Engine/Plugins/ChaosVD/Source/ChaosVD/Public/ChaosVDObjectDetailsTab.h
--- a/Engine/Plugins/ChaosVD/Source/ChaosVD/Public/ChaosVDObjectDetailsTab.h
+++ b/Engine/Plugins/ChaosVD/Source/ChaosVD/Public/ChaosVDObjectDetailsTab.h
@@ -74,7 +74,7 @@ protected:
 
 	void HandleActorsSelection(TArrayView<AActor*> SelectedActors);
 
-	void HandleSolverDataSelectionChange(const TSharedPtr<FChaosVDSolverDataSelectionHandle>& SelectionHandle);
+	virtual void HandleSolverDataSelectionChange(const TSharedPtr<FChaosVDSolverDataSelectionHandle>& SelectionHandle);
 
 	FDelegateHandle SelectionDelegateHandle;
 	TWeakObjectPtr<UObject> CurrentSelectedObject = nullptr;
diff --git a/Engine/Plugins/Experimental/Mover/Source/MoverCVDEditor/Private/MoverCVDTab.cpp b/Engine/Plugins/Experimental/Mover/Source/MoverCVDEditor/Private/MoverCVDTab.cpp
--- a/Engine/Plugins/Experimental/Mover/Source/MoverCVDEditor/Private/MoverCVDTab.cpp
+++ b/Engine/Plugins/Experimental/Mover/Source/MoverCVDEditor/Private/MoverCVDTab.cpp
@@ -56,6 +56,13 @@ TSharedRef<SDockTab> FMoverCVDTab::HandleTabSpawnRequest(const FSpawnTabArgs& Ar
 	return Tab;
 }
 
+void FMoverCVDTab::HandleSolverDataSelectionChange(const TSharedPtr<FChaosVDSolverDataSelectionHandle>& SelectionHandle)
+{
+	// We override the default behavior, otherwise we will display any struct that gets selected and can be displayed
+	// while we only want to display mover info
+	return;
+}
+
 void FMoverCVDTab::RetrieveAllSolversMoverDataComponents()
 {
 	SolverToSimDataComponentMap.Empty();
diff --git a/Engine/Plugins/Experimental/Mover/Source/MoverCVDEditor/Private/MoverCVDTab.h b/Engine/Plugins/Experimental/Mover/Source/MoverCVDEditor/Private/MoverCVDTab.h
--- a/Engine/Plugins/Experimental/Mover/Source/MoverCVDEditor/Private/MoverCVDTab.h
+++ b/Engine/Plugins/Experimental/Mover/Source/MoverCVDEditor/Private/MoverCVDTab.h
@@ -29,6 +29,8 @@ public:
 	virtual TSharedRef<SDockTab> HandleTabSpawnRequest(const FSpawnTabArgs& Args) override;
 	virtual void HandlePostSelectionChange(const UTypedElementSelectionSet* ChangedSelectionSet) override;
 
+	virtual void HandleSolverDataSelectionChange(const TSharedPtr<FChaosVDSolverDataSelectionHandle>& SelectionHandle) override;
+
 	// Scene callbacks
 	void HandleSceneUpdated();
 
