Mover: Pass on Mover with the goal of exposing functions, parameters and structs/properties to allow for a Movement Mode to be created entirely in BP.

diff --git a/Engine/Plugins/Experimental/Mover/Source/Mover/Private/MoveLibrary/FloorQueryUtils.cpp b/Engine/Plugins/Experimental/Mover/Source/Mover/Private/MoveLibrary/FloorQueryUtils.cpp
--- a/Engine/Plugins/Experimental/Mover/Source/Mover/Private/MoveLibrary/FloorQueryUtils.cpp
+++ b/Engine/Plugins/Experimental/Mover/Source/Mover/Private/MoveLibrary/FloorQueryUtils.cpp
@@ -153,7 +153,7 @@ bool UFloorQueryUtils::FloorSweepTest(const FMovingComponentSet& MovingComps, FH
 	return bBlockingHit;
 }
 
-bool UFloorQueryUtils::IsHitSurfaceWalkable(const FHitResult& Hit, FVector UpDirection, float MaxWalkSlopeCosine)
+bool UFloorQueryUtils::IsHitSurfaceWalkable(const FHitResult& Hit, const FVector& UpDirection, float MaxWalkSlopeCosine)
 {
 	if (!Hit.IsValidBlockingHit())
 	{
diff --git a/Engine/Plugins/Experimental/Mover/Source/Mover/Private/MoveLibrary/MovementRecordUtils.cpp b/Engine/Plugins/Experimental/Mover/Source/Mover/Private/MoveLibrary/MovementRecordUtils.cpp
new file mode 100644
--- /dev/null
+++ b/Engine/Plugins/Experimental/Mover/Source/Mover/Private/MoveLibrary/MovementRecordUtils.cpp
@@ -0,0 +1,24 @@
+﻿// Copyright Epic Games, Inc. All Rights Reserved.
+
+#include "MoveLibrary/MovementRecordUtils.h"
+#include "MoveLibrary/MovementRecord.h"
+
+void UMovementRecordUtils::K2_SetDeltaSeconds(FMovementRecord& OutMovementRecord, float DeltaSeconds)
+{
+	OutMovementRecord.SetDeltaSeconds(DeltaSeconds);
+}
+
+const FVector& UMovementRecordUtils::K2_GetTotalMoveDelta(const FMovementRecord& MovementRecord)
+{
+	return MovementRecord.GetTotalMoveDelta();
+}
+
+const FVector& UMovementRecordUtils::K2_GetRelevantMoveDelta(const FMovementRecord& MovementRecord)
+{
+	return MovementRecord.GetRelevantMoveDelta();
+}
+
+FVector UMovementRecordUtils::K2_GetRelevantVelocity(const FMovementRecord& MovementRecord)
+{
+	return MovementRecord.GetRelevantVelocity();
+}	
\ No newline at end of file
diff --git a/Engine/Plugins/Experimental/Mover/Source/Mover/Private/MoveLibrary/MovementUtils.cpp b/Engine/Plugins/Experimental/Mover/Source/Mover/Private/MoveLibrary/MovementUtils.cpp
--- a/Engine/Plugins/Experimental/Mover/Source/Mover/Private/MoveLibrary/MovementUtils.cpp
+++ b/Engine/Plugins/Experimental/Mover/Source/Mover/Private/MoveLibrary/MovementUtils.cpp
@@ -318,6 +318,44 @@ float UMovementUtils::TryMoveToSlideAlongSurface(const FMovingComponentSet& Movi
 	return 0.f;
 }
 
+float UMovementUtils::TryMoveToSlideAlongSurfaceNoMovementRecord(const FMovingComponentSet& MovingComps, const FVector& Delta, float PctOfDeltaToMove, const FQuat Rotation, const FVector& Normal, FHitResult& Hit, bool bHandleImpact)
+{
+	FMovementRecord TempMoveRecord;
+	return TryMoveToSlideAlongSurface(MovingComps, Delta, PctOfDeltaToMove, Rotation, Normal, Hit, bHandleImpact, TempMoveRecord);
+}
+
+float UMovementUtils::TrySafeMoveAndSlideUpdatedComponent(const FMovingComponentSet& MovingComps, const FVector& Delta, const FQuat& NewRotation, bool bSweep, FHitResult& OutHit, ETeleportType Teleport, FMovementRecord& MoveRecord, bool bSlideAlongSurface)
+{
+	TrySafeMoveUpdatedComponent(MovingComps, Delta, NewRotation, bSweep, OutHit, Teleport, MoveRecord);
+
+	float PercentMovementApplied = OutHit.Time;
+	
+	if (OutHit.IsValidBlockingHit())
+	{
+		if (bSlideAlongSurface)
+		{
+			UMoverComponent* MoverComponent = MovingComps.MoverComponent.Get();
+			FMoverOnImpactParams ImpactParams(DefaultModeNames::Flying, OutHit, Delta);
+			MoverComponent->HandleImpact(ImpactParams);
+			// Try to slide the remaining distance along the surface.
+			TryMoveToSlideAlongSurface(FMovingComponentSet(MoverComponent), Delta, 1.f - OutHit.Time, NewRotation, OutHit.Normal, OutHit, true, MoveRecord);
+			PercentMovementApplied = OutHit.Time;
+		}
+	}
+	else
+	{
+		PercentMovementApplied = 1.0f;
+	}
+
+	return PercentMovementApplied;
+}
+
+float UMovementUtils::TrySafeMoveAndSlideUpdatedComponentNoMovementRecord(const FMovingComponentSet& MovingComps, const FVector& Delta, const FQuat& NewRotation, bool bSweep, FHitResult& OutHit, ETeleportType Teleport, bool bSlideAlongSurface)
+{
+	FMovementRecord TempRecord;
+	return TrySafeMoveAndSlideUpdatedComponent(MovingComps, Delta, NewRotation, bSweep, OutHit, Teleport, TempRecord);
+}
+
 
 static const FName SafeMoveSubstepName = "SafeMove";
 
@@ -376,6 +414,12 @@ bool UMovementUtils::TrySafeMoveUpdatedComponent(const FMovingComponentSet& Movi
 	return bMoveResult;
 }
 
+bool UMovementUtils::TrySafeMoveUpdatedComponentNoMovementRecord(const FMovingComponentSet& MovingComps, const FVector& Delta, const FQuat& NewRotation, bool bSweep, FHitResult& OutHit, ETeleportType Teleport)
+{
+	FMovementRecord TempMovementRecord;
+	return TrySafeMoveUpdatedComponent(MovingComps, Delta, NewRotation, bSweep, OutHit, Teleport, TempMovementRecord);
+}
+
 FVector UMovementUtils::ComputePenetrationAdjustment(const FHitResult& Hit)
 {
 	if (!Hit.bStartPenetrating)
diff --git a/Engine/Plugins/Experimental/Mover/Source/Mover/Private/MoveLibrary/WaterMovementUtils.cpp b/Engine/Plugins/Experimental/Mover/Source/Mover/Private/MoveLibrary/WaterMovementUtils.cpp
--- a/Engine/Plugins/Experimental/Mover/Source/Mover/Private/MoveLibrary/WaterMovementUtils.cpp
+++ b/Engine/Plugins/Experimental/Mover/Source/Mover/Private/MoveLibrary/WaterMovementUtils.cpp
@@ -30,11 +30,11 @@ FProposedMove UWaterMovementUtils::ComputeControlledWaterMove(const FWaterMovePa
 	return OutMove;
 }
 
-void UWaterMovementUtils::UpdateWaterSplineData(const FUpdateWaterSplineDataParams& UpdateWaterSplineDataParams, FWaterCheckResult& OutWaterResult)
+void UWaterMovementUtils::UpdateWaterSplineData(const FUpdateWaterSplineDataParams& UpdateWaterSplineDataParams, FWaterCheckResult& WaterCheckResult)
 {
-	if(OutWaterResult.HitResult.Component.IsValid())
+	if(WaterCheckResult.HitResult.Component.IsValid())
 	{
-		if (AWaterBody* WaterBody = Cast<AWaterBody>(OutWaterResult.HitResult.Component->GetOwner()))
+		if (AWaterBody* WaterBody = Cast<AWaterBody>(WaterCheckResult.HitResult.Component->GetOwner()))
 		{
 			if (WaterBody->GetWaterBodyComponent())
 			{
@@ -45,7 +45,7 @@ void UWaterMovementUtils::UpdateWaterSplineData(const FUpdateWaterSplineDataPara
 						   | EWaterBodyQueryFlags::ComputeVelocity
 						   | EWaterBodyQueryFlags::ComputeImmersionDepth);
 		
-				FWaterFlowSplineData& WaterSplineData = OutWaterResult.WaterSplineData;
+				FWaterFlowSplineData& WaterSplineData = WaterCheckResult.WaterSplineData;
 		
 				// Immersion depth
 				WaterSplineData.ImmersionDepth = QueryResult.GetImmersionDepth();
diff --git a/Engine/Plugins/Experimental/Mover/Source/Mover/Public/MoveLibrary/AirMovementUtils.h b/Engine/Plugins/Experimental/Mover/Source/Mover/Public/MoveLibrary/AirMovementUtils.h
--- a/Engine/Plugins/Experimental/Mover/Source/Mover/Public/MoveLibrary/AirMovementUtils.h
+++ b/Engine/Plugins/Experimental/Mover/Source/Mover/Public/MoveLibrary/AirMovementUtils.h
@@ -15,43 +15,43 @@ struct MOVER_API FFreeMoveParams
 {
 	GENERATED_BODY()
 
-	UPROPERTY(BlueprintReadOnly, EditAnywhere, Category = Mover)
+	UPROPERTY(BlueprintReadWrite, EditAnywhere, Category = Mover)
 	EMoveInputType MoveInputType = EMoveInputType::DirectionalIntent;
 
-	UPROPERTY(BlueprintReadOnly, EditAnywhere, Category = Mover)
+	UPROPERTY(BlueprintReadWrite, EditAnywhere, Category = Mover)
 	FVector MoveInput = FVector::ZeroVector;
 
-	UPROPERTY(BlueprintReadOnly, EditAnywhere, Category = Mover)
+	UPROPERTY(BlueprintReadWrite, EditAnywhere, Category = Mover)
 	FRotator OrientationIntent = FRotator::ZeroRotator;
 
-	UPROPERTY(BlueprintReadOnly, EditAnywhere, Category = Mover)
+	UPROPERTY(BlueprintReadWrite, EditAnywhere, Category = Mover)
 	FVector PriorVelocity = FVector::ZeroVector;
 
-	UPROPERTY(BlueprintReadOnly, EditAnywhere, Category = Mover)
+	UPROPERTY(BlueprintReadWrite, EditAnywhere, Category = Mover)
 	FRotator PriorOrientation = FRotator::ZeroRotator;
 
-	UPROPERTY(BlueprintReadOnly, EditAnywhere, Category = Mover)
+	UPROPERTY(BlueprintReadWrite, EditAnywhere, Category = Mover)
 	float MaxSpeed = 800.f;
 
-	UPROPERTY(BlueprintReadOnly, EditAnywhere, Category = Mover)
+	UPROPERTY(BlueprintReadWrite, EditAnywhere, Category = Mover)
 	float Acceleration = 4000.f;
 
-	UPROPERTY(BlueprintReadOnly, EditAnywhere, Category = Mover)
+	UPROPERTY(BlueprintReadWrite, EditAnywhere, Category = Mover)
 	float Deceleration = 8000.f;
 
-	UPROPERTY(BlueprintReadOnly, EditAnywhere, Category = Mover)
+	UPROPERTY(BlueprintReadWrite, EditAnywhere, Category = Mover)
 	float TurningBoost = 8.f;
 
-	UPROPERTY(BlueprintReadOnly, EditAnywhere, Category = Mover)
+	UPROPERTY(BlueprintReadWrite, EditAnywhere, Category = Mover)
 	float TurningRate = 500.f;
 
-	UPROPERTY(BlueprintReadOnly, EditAnywhere, Category = Mover)
+	UPROPERTY(BlueprintReadWrite, EditAnywhere, Category = Mover)
 	float DeltaSeconds = 0.f;
 
-	UPROPERTY(BlueprintReadOnly, EditAnywhere, Category = Mover)
+	UPROPERTY(BlueprintReadWrite, EditAnywhere, Category = Mover)
 	FQuat WorldToGravityQuat = FQuat::Identity;
 
-	UPROPERTY(BlueprintReadOnly, EditAnywhere, Category = Mover)
+	UPROPERTY(BlueprintReadWrite, EditAnywhere, Category = Mover)
 	bool bUseAccelerationForVelocityMove = true;
 };
 
@@ -70,11 +70,11 @@ public:
 	
     // Checks if a hit result represents a walkable location that an actor can land on
     UFUNCTION(BlueprintCallable, Category=Mover)
-    static bool IsValidLandingSpot(const FMovingComponentSet& MovingComps, const FVector& Location, const FHitResult& Hit, float FloorSweepDistance, float MaxWalkSlopeCosine, FFloorCheckResult& OutFloorResult);
+    static bool IsValidLandingSpot(const FMovingComponentSet& MovingComps, const FVector& Location, const FHitResult& Hit, float FloorSweepDistance, float MaxWalkSlopeCosine, UPARAM(Ref) FFloorCheckResult& OutFloorResult);
     
     /** Attempts to move a component along a surface, while checking for landing on a walkable surface. Intended for use while falling. Returns the percent of time applied, with 0.0 meaning no movement occurred. */
     UFUNCTION(BlueprintCallable, Category=Mover)
-    static float TryMoveToFallAlongSurface(const FMovingComponentSet& MovingComps, const FVector& Delta, float PctOfDeltaToMove, const FQuat Rotation, const FVector& Normal, FHitResult& Hit, bool bHandleImpact, float FloorSweepDistance, float MaxWalkSlopeCosine, FFloorCheckResult& OutFloorResult, FMovementRecord& MoveRecord);
+    static float TryMoveToFallAlongSurface(const FMovingComponentSet& MovingComps, const FVector& Delta, float PctOfDeltaToMove, const FQuat Rotation, const FVector& Normal, UPARAM(Ref) FHitResult& Hit, bool bHandleImpact, float FloorSweepDistance, float MaxWalkSlopeCosine, UPARAM(Ref) FFloorCheckResult& OutFloorResult, UPARAM(Ref) FMovementRecord& MoveRecord);
 
 public:	
 	// Threadsafe movement queries
@@ -86,7 +86,7 @@ public:
 	 * Returns the percent of time applied, with 0.0 meaning no movement would occur. 
 	 */
 	UFUNCTION(BlueprintCallable, Category = Mover)
-	static float TestFallingMoveAlongHitSurface(const FMovingComponentSet& MovingComps, const FVector& OriginalMoveDelta, const FVector& LocationAtHit, const FQuat& TargetRotation, bool bHandleImpact, float FloorSweepDistance, float MaxWalkSlopeCosine, FHitResult& InOutHit, FFloorCheckResult& OutFloorResult, FMovementRecord& InOutMoveRecord);
+	static float TestFallingMoveAlongHitSurface(const FMovingComponentSet& MovingComps, const FVector& OriginalMoveDelta, const FVector& LocationAtHit, const FQuat& TargetRotation, bool bHandleImpact, float FloorSweepDistance, float MaxWalkSlopeCosine, UPARAM(Ref) FHitResult& InOutHit, UPARAM(Ref) FFloorCheckResult& OutFloorResult, UPARAM(Ref) FMovementRecord& InOutMoveRecord);
 
 
 };
diff --git a/Engine/Plugins/Experimental/Mover/Source/Mover/Public/MoveLibrary/AsyncMovementUtils.h b/Engine/Plugins/Experimental/Mover/Source/Mover/Public/MoveLibrary/AsyncMovementUtils.h
--- a/Engine/Plugins/Experimental/Mover/Source/Mover/Public/MoveLibrary/AsyncMovementUtils.h
+++ b/Engine/Plugins/Experimental/Mover/Source/Mover/Public/MoveLibrary/AsyncMovementUtils.h
@@ -30,7 +30,7 @@ public:
 	 * Appends to InOutMoveRecord with any movement substeps
 	 */
 	UFUNCTION(BlueprintCallable, Category = Mover)
-	static bool TestDepenetratingMove(const FMovingComponentSet& MovingComps, const FVector& StartLocation, const FVector& TargetLocation, const FQuat& StartRotation, const FQuat& TargetRotation, bool bShouldSweep, FHitResult& OutHit, FMovementRecord& InOutMoveRecord);
+	static bool TestDepenetratingMove(const FMovingComponentSet& MovingComps, const FVector& StartLocation, const FVector& TargetLocation, const FQuat& StartRotation, const FQuat& TargetRotation, bool bShouldSweep, FHitResult& OutHit, UPARAM(ref) FMovementRecord& InOutMoveRecord);
 
 	/** Tests potential movement of a component without actually moving it, taking penetration resolution issues into account first.  Relies on CollisionParams to describe the query.
 	 * Returns true if any movement was possible
@@ -45,7 +45,7 @@ public:
 	 * Appends to InOutMoveRecord with any movement substeps
 	 */
 	UFUNCTION(BlueprintCallable, Category = Mover)
-	static float TestSlidingMoveAlongHitSurface(const FMovingComponentSet& MovingComps, const FVector& OriginalMoveDelta, const FVector& LocationAtHit, const FQuat& TargetRotation, FHitResult& InOutHit, FMovementRecord& InOutMoveRecord);
+	static float TestSlidingMoveAlongHitSurface(const FMovingComponentSet& MovingComps, const FVector& OriginalMoveDelta, const FVector& LocationAtHit, const FQuat& TargetRotation, UPARAM(ref) FHitResult& InOutHit, UPARAM(ref) FMovementRecord& InOutMoveRecord);
 	
 
 	/** Tests potential movement of a component sliding along a surface, without actually moving it. Relies on CollisionParams to describe the query.
diff --git a/Engine/Plugins/Experimental/Mover/Source/Mover/Public/MoveLibrary/BasedMovementUtils.h b/Engine/Plugins/Experimental/Mover/Source/Mover/Public/MoveLibrary/BasedMovementUtils.h
--- a/Engine/Plugins/Experimental/Mover/Source/Mover/Public/MoveLibrary/BasedMovementUtils.h
+++ b/Engine/Plugins/Experimental/Mover/Source/Mover/Public/MoveLibrary/BasedMovementUtils.h
@@ -17,23 +17,23 @@ struct MOVER_API FRelativeBaseInfo
 	GENERATED_BODY()
 
 	/** Component we are moving relative to */
-	UPROPERTY(VisibleInstanceOnly, BlueprintReadOnly, Category="Mover")
+	UPROPERTY(VisibleInstanceOnly, BlueprintReadWrite, Category="Mover")
 	TWeakObjectPtr<UPrimitiveComponent> MovementBase = nullptr;
 
 	/** Bone name on component, for skeletal meshes. NAME_None if not a skeletal mesh or if bone is invalid. */
-	UPROPERTY(VisibleInstanceOnly, BlueprintReadOnly, Category = "Mover")
+	UPROPERTY(VisibleInstanceOnly, BlueprintReadWrite, Category = "Mover")
 	FName BoneName = NAME_None;
 
 	/** Last captured worldspace location of MovementBase / Bone */
-	UPROPERTY(VisibleInstanceOnly, BlueprintReadOnly, Category = "Mover")
+	UPROPERTY(VisibleInstanceOnly, BlueprintReadWrite, Category = "Mover")
 	FVector Location = FVector::ZeroVector;
 
 	/** Last captured worldspace orientation of MovementBase / Bone */
-	UPROPERTY(VisibleInstanceOnly, BlueprintReadOnly, Category = "Mover")
+	UPROPERTY(VisibleInstanceOnly, BlueprintReadWrite, Category = "Mover")
 	FQuat Rotation = FQuat::Identity;
 
 	/** Last captured location of the tethering point where the Mover actor is "attached", relative to the base. */
-	UPROPERTY(VisibleInstanceOnly, BlueprintReadOnly, Category = "Mover")
+	UPROPERTY(VisibleInstanceOnly, BlueprintReadWrite, Category = "Mover")
 	FVector ContactLocalPosition = FVector::ZeroVector;
 
 public:
diff --git a/Engine/Plugins/Experimental/Mover/Source/Mover/Public/MoveLibrary/ConstrainedMoveUtils.h b/Engine/Plugins/Experimental/Mover/Source/Mover/Public/MoveLibrary/ConstrainedMoveUtils.h
--- a/Engine/Plugins/Experimental/Mover/Source/Mover/Public/MoveLibrary/ConstrainedMoveUtils.h
+++ b/Engine/Plugins/Experimental/Mover/Source/Mover/Public/MoveLibrary/ConstrainedMoveUtils.h
@@ -15,7 +15,7 @@ struct MOVER_API FPlanarConstraint
 	GENERATED_BODY()
 
 	/** If true, movement will be constrained to a plane. */
-	UPROPERTY(EditAnywhere, BlueprintReadOnly, Category = PlanarMovement)
+	UPROPERTY(EditAnywhere, BlueprintReadWrite, Category = PlanarMovement)
 	bool bConstrainToPlane = false;
 
 	/**
@@ -23,11 +23,11 @@ struct MOVER_API FPlanarConstraint
 	 * If for example you wanted to constrain movement to the X-Z plane (so that Y cannot change), the normal would be set to X=0 Y=1 Z=0.
 	 * It is normalized once the component is registered with the game world.
 	 */
-	UPROPERTY(EditAnywhere, BlueprintReadOnly, Category = PlanarMovement, meta = (editcondition = bConstrainToPlane))
+	UPROPERTY(EditAnywhere, BlueprintReadWrite, Category = PlanarMovement, meta = (editcondition = bConstrainToPlane))
 	FVector PlaneConstraintNormal = FVector::ForwardVector;
 
 	/** The origin of the plane that constrains movement, if plane constraint is enabled. */
-	UPROPERTY(EditAnywhere, BlueprintReadOnly, Category = PlanarMovement, meta = (editcondition = bConstrainToPlane))
+	UPROPERTY(EditAnywhere, BlueprintReadWrite, Category = PlanarMovement, meta = (editcondition = bConstrainToPlane))
 	FVector PlaneConstraintOrigin = FVector::ZeroVector;
 };
 
diff --git a/Engine/Plugins/Experimental/Mover/Source/Mover/Public/MoveLibrary/FloorQueryUtils.h b/Engine/Plugins/Experimental/Mover/Source/Mover/Public/MoveLibrary/FloorQueryUtils.h
--- a/Engine/Plugins/Experimental/Mover/Source/Mover/Public/MoveLibrary/FloorQueryUtils.h
+++ b/Engine/Plugins/Experimental/Mover/Source/Mover/Public/MoveLibrary/FloorQueryUtils.h
@@ -25,27 +25,27 @@ struct MOVER_API FFloorCheckResult
 	* True if there was a blocking hit in the floor test that was NOT in initial penetration.
 	* The HitResult can give more info about other circumstances.
 	*/
-	UPROPERTY(VisibleInstanceOnly, BlueprintReadOnly, Category = CharacterFloor)
+	UPROPERTY(VisibleInstanceOnly, BlueprintReadWrite, Category = CharacterFloor)
 	uint32 bBlockingHit : 1;
 
 	/** True if the hit found a valid walkable floor. */
-	UPROPERTY(VisibleInstanceOnly, BlueprintReadOnly, Category = CharacterFloor)
+	UPROPERTY(VisibleInstanceOnly, BlueprintReadWrite, Category = CharacterFloor)
 	uint32 bWalkableFloor : 1;
 
 	/** True if the hit found a valid walkable floor using a line trace (rather than a sweep test, which happens when the sweep test fails to yield a walkable surface). */
-	UPROPERTY(VisibleInstanceOnly, BlueprintReadOnly, Category = CharacterFloor)
+	UPROPERTY(VisibleInstanceOnly, BlueprintReadWrite, Category = CharacterFloor)
 	uint32 bLineTrace : 1;
 
 	/** The distance to the floor, computed from the trace. Only valid if bLineTrace is true. */
-	UPROPERTY(VisibleInstanceOnly, BlueprintReadOnly, Category = CharacterFloor)
+	UPROPERTY(VisibleInstanceOnly, BlueprintReadWrite, Category = CharacterFloor)
     float LineDist;
 
 	/** The distance to the floor, computed from the swept capsule trace. */
-	UPROPERTY(VisibleInstanceOnly, BlueprintReadOnly, Category = CharacterFloor)
+	UPROPERTY(VisibleInstanceOnly, BlueprintReadWrite, Category = CharacterFloor)
 	float FloorDist;
 
 	/** Hit result of the test that found a floor. Includes more specific data about the point of impact and surface normal at that point. */
-	UPROPERTY(VisibleInstanceOnly, BlueprintReadOnly, Category = CharacterFloor)
+	UPROPERTY(VisibleInstanceOnly, BlueprintReadWrite, Category = CharacterFloor)
 	FHitResult HitResult;
 
 public:
@@ -116,7 +116,7 @@ public:
 	static bool FloorSweepTest(const FMovingComponentSet& MovingComps, FHitResult& OutHit, const FVector& Start, const FVector& End, ECollisionChannel TraceChannel, const struct FCollisionShape& CollisionShape, const struct FCollisionQueryParams& Params, const struct FCollisionResponseParams& ResponseParam);
 
 	UFUNCTION(BlueprintCallable, Category=Mover)
-	static bool IsHitSurfaceWalkable(const FHitResult& Hit, FVector UpDirection, float MaxWalkSlopeCosine);
+	static bool IsHitSurfaceWalkable(const FHitResult& Hit, const FVector& UpDirection, float MaxWalkSlopeCosine);
 
 	/**
 	 * Return true if the 2D distance to the impact point is inside the edge tolerance (CapsuleRadius minus a small rejection threshold).
diff --git a/Engine/Plugins/Experimental/Mover/Source/Mover/Public/MoveLibrary/GroundMovementUtils.h b/Engine/Plugins/Experimental/Mover/Source/Mover/Public/MoveLibrary/GroundMovementUtils.h
--- a/Engine/Plugins/Experimental/Mover/Source/Mover/Public/MoveLibrary/GroundMovementUtils.h
+++ b/Engine/Plugins/Experimental/Mover/Source/Mover/Public/MoveLibrary/GroundMovementUtils.h
@@ -20,52 +20,52 @@ struct MOVER_API FGroundMoveParams
 {
 	GENERATED_BODY()
 
-	UPROPERTY(BlueprintReadOnly, EditAnywhere, Category = Mover)
+	UPROPERTY(BlueprintReadWrite, EditAnywhere, Category = Mover)
 	EMoveInputType MoveInputType = EMoveInputType::DirectionalIntent;
 
-	UPROPERTY(BlueprintReadOnly, EditAnywhere, Category = Mover)
+	UPROPERTY(BlueprintReadWrite, EditAnywhere, Category = Mover)
 	FVector MoveInput = FVector::ZeroVector;
 
-	UPROPERTY(BlueprintReadOnly, EditAnywhere, Category = Mover)
+	UPROPERTY(BlueprintReadWrite, EditAnywhere, Category = Mover)
 	FRotator OrientationIntent = FRotator::ZeroRotator;
 
-	UPROPERTY(BlueprintReadOnly, EditAnywhere, Category = Mover)
+	UPROPERTY(BlueprintReadWrite, EditAnywhere, Category = Mover)
 	FVector PriorVelocity = FVector::ZeroVector;
 
-	UPROPERTY(BlueprintReadOnly, EditAnywhere, Category = Mover)
+	UPROPERTY(BlueprintReadWrite, EditAnywhere, Category = Mover)
 	FRotator PriorOrientation = FRotator::ZeroRotator;
 
-	UPROPERTY(BlueprintReadOnly, EditAnywhere, Category = Mover)
+	UPROPERTY(BlueprintReadWrite, EditAnywhere, Category = Mover)
 	float MaxSpeed = 800.f;
 
-	UPROPERTY(BlueprintReadOnly, EditAnywhere, Category = Mover)
+	UPROPERTY(BlueprintReadWrite, EditAnywhere, Category = Mover)
 	float Acceleration = 4000.f;
 
-	UPROPERTY(BlueprintReadOnly, EditAnywhere, Category = Mover)
+	UPROPERTY(BlueprintReadWrite, EditAnywhere, Category = Mover)
 	float Deceleration = 8000.f;
 
-	UPROPERTY(BlueprintReadOnly, EditAnywhere, Category = Mover)
+	UPROPERTY(BlueprintReadWrite, EditAnywhere, Category = Mover)
 	float Friction = 0.f;
 
-	UPROPERTY(BlueprintReadOnly, EditAnywhere, Category = Mover)
+	UPROPERTY(BlueprintReadWrite, EditAnywhere, Category = Mover)
 	float TurningRate = 500.f;
 
-	UPROPERTY(BlueprintReadOnly, EditAnywhere, Category = Mover)
+	UPROPERTY(BlueprintReadWrite, EditAnywhere, Category = Mover)
 	float TurningBoost = 8.f;
 
-	UPROPERTY(BlueprintReadOnly, EditAnywhere, Category = Mover)
+	UPROPERTY(BlueprintReadWrite, EditAnywhere, Category = Mover)
 	FVector GroundNormal = FVector::UpVector;
 
-	UPROPERTY(BlueprintReadOnly, EditAnywhere, Category = Mover)
+	UPROPERTY(BlueprintReadWrite, EditAnywhere, Category = Mover)
 	float DeltaSeconds = 0.f;
 
-	UPROPERTY(BlueprintReadOnly, EditAnywhere, Category = Mover)
+	UPROPERTY(BlueprintReadWrite, EditAnywhere, Category = Mover)
 	FVector UpDirection = FVector::UpVector;
 	
-	UPROPERTY(BlueprintReadOnly, EditAnywhere, Category = Mover)
+	UPROPERTY(BlueprintReadWrite, EditAnywhere, Category = Mover)
 	FQuat WorldToGravityQuat = FQuat::Identity;
 
-	UPROPERTY(BlueprintReadOnly, EditAnywhere, Category = Mover)
+	UPROPERTY(BlueprintReadWrite, EditAnywhere, Category = Mover)
 	bool bUseAccelerationForVelocityMove = true;
 };
 
@@ -93,7 +93,7 @@ public:
      *  Note: This modifies the normal and calls UMovementUtils::TryMoveToSlideAlongSurface
      */
     UFUNCTION(BlueprintCallable, Category=Mover)
-    static float TryWalkToSlideAlongSurface(const FMovingComponentSet& MovingComps, const FVector& Delta, float PctOfDeltaToMove, const FQuat Rotation, const FVector& Normal, FHitResult& Hit, bool bHandleImpact, FMovementRecord& MoveRecord, float MaxWalkSlopeCosine, float MaxStepHeight);
+    static float TryWalkToSlideAlongSurface(const FMovingComponentSet& MovingComps, const FVector& Delta, float PctOfDeltaToMove, const FQuat Rotation, const FVector& Normal, UPARAM(ref) FHitResult& Hit, bool bHandleImpact, UPARAM(ref) FMovementRecord& MoveRecord, float MaxWalkSlopeCosine, float MaxStepHeight);
 
 	/** Used to change a movement to be along a ramp's surface, typically to prevent slow movement when running up/down a ramp */
 	UFUNCTION(BlueprintCallable, Category = Mover)
@@ -111,7 +111,7 @@ public:
 	 * Returns the percent of time applied, with 0.0 meaning no movement would occur. 
 	 */
 	UFUNCTION(BlueprintCallable, Category = Mover)
-	static float TestGroundedMoveAlongHitSurface(const FMovingComponentSet& MovingComps, const FVector& OriginalMoveDelta, const FVector& LocationAtHit, const FQuat& TargetRotation, bool bHandleImpact, float MaxStepHeight, float MaxWalkSlopeCosine, FHitResult& InOutHit, FMovementRecord& InOutMoveRecord);
+	static float TestGroundedMoveAlongHitSurface(const FMovingComponentSet& MovingComps, const FVector& OriginalMoveDelta, const FVector& LocationAtHit, const FQuat& TargetRotation, bool bHandleImpact, float MaxStepHeight, float MaxWalkSlopeCosine, UPARAM(ref) FHitResult& InOutHit, UPARAM(ref) FMovementRecord& InOutMoveRecord);
 
 	/** Tests potential movement of a component up/down to adjust to a walkable floor. Intended for use while performing ground movement. 
 	 * Modifies InOutCurrentFloor to account for adjustments, if any was made
@@ -119,7 +119,7 @@ public:
 	 * Returns new location, including any adjustment if it was made
 	 */
 	UFUNCTION(BlueprintCallable, Category = Mover)
-	static FVector TestMoveToAdjustToFloor(const FMovingComponentSet& MovingComps, const FVector& Location, const FQuat& Rotation, float MaxWalkSlopeCosine, FFloorCheckResult& InOutCurrentFloor, FMovementRecord& InOutMoveRecord);
+	static FVector TestMoveToAdjustToFloor(const FMovingComponentSet& MovingComps, const FVector& Location, const FQuat& Rotation, float MaxWalkSlopeCosine, UPARAM(ref) FFloorCheckResult& InOutCurrentFloor, UPARAM(ref) FMovementRecord& InOutMoveRecord);
 
 	/** Tests potential up-and-over movement of a component to overcome a blocking obstacle hit, such as moving up steps or slope.
 	 * Sets OutFloorTestResult with final floor info (optional)
diff --git a/Engine/Plugins/Experimental/Mover/Source/Mover/Public/MoveLibrary/MovementRecord.h b/Engine/Plugins/Experimental/Mover/Source/Mover/Public/MoveLibrary/MovementRecord.h
--- a/Engine/Plugins/Experimental/Mover/Source/Mover/Public/MoveLibrary/MovementRecord.h
+++ b/Engine/Plugins/Experimental/Mover/Source/Mover/Public/MoveLibrary/MovementRecord.h
@@ -41,8 +41,8 @@ struct MOVER_API FMovementRecord
 	void SetDeltaSeconds(float DeltaSeconds) { TotalDeltaSeconds = DeltaSeconds; }
 
 	const TArray<FMovementSubstep>& GetSubsteps() const { return Substeps; }
-	FVector GetTotalMoveDelta() const { return TotalMoveDelta; }
-	FVector GetRelevantMoveDelta() const { return RelevantMoveDelta; }
+	const FVector& GetTotalMoveDelta() const { return TotalMoveDelta; }
+	const FVector& GetRelevantMoveDelta() const { return RelevantMoveDelta; }
 	FVector GetRelevantVelocity() const { return ((TotalDeltaSeconds > 0.f) ? (RelevantMoveDelta / TotalDeltaSeconds) : FVector::ZeroVector); }
 
 	FString ToString() const;
diff --git a/Engine/Plugins/Experimental/Mover/Source/Mover/Public/MoveLibrary/MovementRecordUtils.h b/Engine/Plugins/Experimental/Mover/Source/Mover/Public/MoveLibrary/MovementRecordUtils.h
new file mode 100644
--- /dev/null
+++ b/Engine/Plugins/Experimental/Mover/Source/Mover/Public/MoveLibrary/MovementRecordUtils.h
@@ -0,0 +1,34 @@
+﻿// Copyright Epic Games, Inc. All Rights Reserved.
+
+#pragma once
+
+#include "Kismet/BlueprintFunctionLibrary.h"
+#include "MovementRecordUtils.generated.h"
+
+struct FMovementRecord;
+
+/**
+ * MovementRecordUtils: a collection of stateless static BP-accessible functions for movement record related operations
+ */
+UCLASS()
+class MOVER_API UMovementRecordUtils : public UBlueprintFunctionLibrary
+{
+	GENERATED_BODY()
+
+public:
+	/** Sets the delta time these moves took place over */
+	UFUNCTION(BlueprintCallable, Category = Mover, DisplayName="Set Delta Seconds")
+	static void K2_SetDeltaSeconds(UPARAM(Ref) FMovementRecord& OutMovementRecord, float DeltaSeconds);
+
+	/** Get all the move delta applied to the actor/record */
+	UFUNCTION(BlueprintPure, Category = Mover, DisplayName="Get Total Move Delta")
+	static const FVector& K2_GetTotalMoveDelta(const FMovementRecord& MovementRecord);
+
+	/** Get relevant move delta applied to the actor/record */
+	UFUNCTION(BlueprintPure, Category = Mover, DisplayName="Get Relevant Move Delta")
+	static const FVector& K2_GetRelevantMoveDelta(const FMovementRecord& MovementRecord);
+
+	/** Get relevant velocity applied to the actor/record */
+	UFUNCTION(BlueprintPure, Category = Mover, DisplayName="Get Relevant Velocity")
+	static FVector K2_GetRelevantVelocity(const FMovementRecord& MovementRecord);
+};
diff --git a/Engine/Plugins/Experimental/Mover/Source/Mover/Public/MoveLibrary/MovementUtils.h b/Engine/Plugins/Experimental/Mover/Source/Mover/Public/MoveLibrary/MovementUtils.h
--- a/Engine/Plugins/Experimental/Mover/Source/Mover/Public/MoveLibrary/MovementUtils.h
+++ b/Engine/Plugins/Experimental/Mover/Source/Mover/Public/MoveLibrary/MovementUtils.h
@@ -43,38 +43,38 @@ struct MOVER_API FComputeVelocityParams
 {
 	GENERATED_BODY()
 	
-	UPROPERTY(BlueprintReadOnly, EditAnywhere, Category = Mover)
+	UPROPERTY(BlueprintReadWrite, EditAnywhere, Category = Mover)
 	float DeltaSeconds = 0.f;
 	
-	UPROPERTY(BlueprintReadOnly, EditAnywhere, Category = Mover)
+	UPROPERTY(BlueprintReadWrite, EditAnywhere, Category = Mover)
 	FVector InitialVelocity = FVector::ZeroVector;
 	
-	UPROPERTY(BlueprintReadOnly, EditAnywhere, Category = Mover)
+	UPROPERTY(BlueprintReadWrite, EditAnywhere, Category = Mover)
 	FVector MoveDirectionIntent = FVector::ZeroVector;
 
 	// AuxState variables
-	UPROPERTY(BlueprintReadOnly, EditAnywhere, Category = Mover)
+	UPROPERTY(BlueprintReadWrite, EditAnywhere, Category = Mover)
 	float MaxSpeed = 0.f;
 	
-	UPROPERTY(BlueprintReadOnly, EditAnywhere, Category = Mover)
+	UPROPERTY(BlueprintReadWrite, EditAnywhere, Category = Mover)
 	float TurningBoost = 0.f;
 	
-	UPROPERTY(BlueprintReadOnly, EditAnywhere, Category = Mover)
+	UPROPERTY(BlueprintReadWrite, EditAnywhere, Category = Mover)
 	float Friction = 0.f;
 	
-	UPROPERTY(BlueprintReadOnly, EditAnywhere, Category = Mover)
+	UPROPERTY(BlueprintReadWrite, EditAnywhere, Category = Mover)
 	float Deceleration = 0.f;
 	
-	UPROPERTY(BlueprintReadOnly, EditAnywhere, Category = Mover)
+	UPROPERTY(BlueprintReadWrite, EditAnywhere, Category = Mover)
 	float Acceleration = 0.f;
 	
-	UPROPERTY(BlueprintReadOnly, EditAnywhere, Category = Mover)
+	UPROPERTY(BlueprintReadWrite, EditAnywhere, Category = Mover)
 	FVector MoveInput = FVector::ZeroVector;
 	
-	UPROPERTY(BlueprintReadOnly, EditAnywhere, Category = Mover)
+	UPROPERTY(BlueprintReadWrite, EditAnywhere, Category = Mover)
 	EMoveInputType MoveInputType = EMoveInputType::DirectionalIntent;
 	
-	UPROPERTY(BlueprintReadOnly, EditAnywhere, Category = Mover)
+	UPROPERTY(BlueprintReadWrite, EditAnywhere, Category = Mover)
 	bool bUseAccelerationForVelocityMove = true;
 };
 
@@ -84,35 +84,35 @@ struct MOVER_API FComputeCombinedVelocityParams
 {
 	GENERATED_BODY()
 	
-	UPROPERTY(BlueprintReadOnly, EditAnywhere, Category = Mover)
+	UPROPERTY(BlueprintReadWrite, EditAnywhere, Category = Mover)
 	float DeltaSeconds = 0.f;
 	
-	UPROPERTY(BlueprintReadOnly, EditAnywhere, Category = Mover)
+	UPROPERTY(BlueprintReadWrite, EditAnywhere, Category = Mover)
 	FVector InitialVelocity = FVector::ZeroVector;
 	
-	UPROPERTY(BlueprintReadOnly, EditAnywhere, Category = Mover)
+	UPROPERTY(BlueprintReadWrite, EditAnywhere, Category = Mover)
 	FVector MoveDirectionIntent = FVector::ZeroVector;
 
 	// AuxState variables
-	UPROPERTY(BlueprintReadOnly, EditAnywhere, Category = Mover)
+	UPROPERTY(BlueprintReadWrite, EditAnywhere, Category = Mover)
 	float MaxSpeed = 0.f;
 	
-	UPROPERTY(BlueprintReadOnly, EditAnywhere, Category = Mover)
+	UPROPERTY(BlueprintReadWrite, EditAnywhere, Category = Mover)
 	float TurningBoost = 0.f;
 	
-	UPROPERTY(BlueprintReadOnly, EditAnywhere, Category = Mover)
+	UPROPERTY(BlueprintReadWrite, EditAnywhere, Category = Mover)
 	float Friction = 0.f;
 	
-	UPROPERTY(BlueprintReadOnly, EditAnywhere, Category = Mover)
+	UPROPERTY(BlueprintReadWrite, EditAnywhere, Category = Mover)
 	float Deceleration = 0.f;
 	
-	UPROPERTY(BlueprintReadOnly, EditAnywhere, Category = Mover)
+	UPROPERTY(BlueprintReadWrite, EditAnywhere, Category = Mover)
 	float Acceleration = 0.f;
 
-	UPROPERTY(BlueprintReadOnly, EditAnywhere, Category = Mover)
+	UPROPERTY(BlueprintReadWrite, EditAnywhere, Category = Mover)
 	FVector ExternalAcceleration = FVector::ZeroVector;
 	
-	UPROPERTY(BlueprintReadOnly, EditAnywhere, Category = Mover)
+	UPROPERTY(BlueprintReadWrite, EditAnywhere, Category = Mover)
 	float OverallMaxSpeed = 0.f;
 };
 
@@ -187,15 +187,72 @@ public:
 	/** Returns an alternative move delta when we are in contact with 2 surfaces */
 	static FVector ComputeTwoWallAdjustedDelta(const FMovingComponentSet& MovingComps, const FVector& MoveDelta, const FHitResult& Hit, const FVector& OldHitNormal);
 
-	/** Attempts to move a component along a surface. Returns the percent of time applied, with 0.0 meaning no movement occurred. */
-	UFUNCTION(BlueprintCallable, Category=Mover)
-	static float TryMoveToSlideAlongSurface(const FMovingComponentSet& MovingComps, const FVector& Delta, float PctOfDeltaToMove, const FQuat Rotation, const FVector& Normal, FHitResult& Hit, bool bHandleImpact, FMovementRecord& MoveRecord);
+	/**
+	 * Attempts to move a component along a surface. Returns the percent of time applied, with 0.0 meaning no movement occurred.
+	 * Note: This function takes a movement record that collects moves applied to the actor see @FMovementRecord
+	 */
+	UFUNCTION(BlueprintCallable, Category=Mover, meta=(DisplayName="Try Move To Slide Along Surface (Movement Record)"))
+	static float TryMoveToSlideAlongSurface(const FMovingComponentSet& MovingComps, const FVector& Delta, float PctOfDeltaToMove, const FQuat Rotation, const FVector& Normal, UPARAM(Ref) FHitResult& Hit, bool bHandleImpact, UPARAM(Ref) FMovementRecord& MoveRecord);
+
+	/**
+	 * Attempts to move a component along a surface. Returns the percent of time applied, with 0.0 meaning no movement occurred.
+	 * Note: This function doesn't update a movement record so velocity should be gathered/set using a different method
+	 */
+	UFUNCTION(BlueprintCallable, Category=Mover, meta=(DisplayName="Try Move To Slide Along Surface (No Movement Record)"))
+	static float TryMoveToSlideAlongSurfaceNoMovementRecord(const FMovingComponentSet& MovingComps, const FVector& Delta, float PctOfDeltaToMove, const FQuat Rotation, const FVector& Normal, UPARAM(Ref) FHitResult& Hit, bool bHandleImpact);
 
 	// Component movement
 
-	/** Attempts to move a component and resolve any penetration issues with the proposed move Delta */
-	UFUNCTION(BlueprintCallable, Category=Mover)
-	static bool TrySafeMoveUpdatedComponent(const FMovingComponentSet& MovingComps, const FVector& Delta, const FQuat& NewRotation, bool bSweep, FHitResult& OutHit, ETeleportType Teleport, FMovementRecord& MoveRecord);
+	/**
+	 * Attempts to move a component and resolve any penetration issues with the proposed move Delta.
+	 * This function handles blocking moves and slides along the surface or stops on impact. It uses @TryMoveToSlideAlongSurface to slide
+	 * along the surface on hit, so if other behavior is needed for sliding (like Falling based sliding) consider using
+	 * @TrySafeMoveUpdatedComponent and a slide function or setting @bSlideAlongSurface to false and then using a separate sliding function.
+	 * Note: This function takes a movement record that collects moves applied to the actor see @FMovementRecord.
+	 * @param MovingComps Encapsulates components involved in movement
+	 * @param Delta The desired location change in world space
+	 * @param NewRotation The new desired rotation in world space
+	 * @param bSweep Whether we sweep to the destination location
+	 * @param OutHit Optional output describing the blocking hit that stopped the move, if any
+	 * @param Teleport Whether we teleport the physics state (if physics collision is enabled for this object)
+	 * @param MoveRecord Move record to add moves to
+	 * @param bSlideAlongSurface If true the actor slides along a blocking surface. If false the actor will stop if the move was blocked
+	 * @return Returns the percent of the move applied 0 to 1, where 1 represents the whole move being applied
+	 */
+	UFUNCTION(BlueprintCallable, Category=Mover, meta=(DisplayName="Try Safe Move And Slide (Movement Record)"))
+	static float TrySafeMoveAndSlideUpdatedComponent(const FMovingComponentSet& MovingComps, const FVector& Delta, const FQuat& NewRotation, bool bSweep, UPARAM(Ref) FHitResult& OutHit, ETeleportType Teleport, UPARAM(Ref) FMovementRecord& MoveRecord, bool bSlideAlongSurface = true);
+
+	/**
+ 	 * Attempts to move a component and resolve any penetration issues with the proposed move Delta.
+ 	 * This function handles blocking moves and slides along the surface or stops on impact. It uses @TryMoveToSlideAlongSurface to slide
+ 	 * along the surface on hit, so if other behavior is needed for sliding (like Falling based sliding) consider using
+ 	 * @TrySafeMoveUpdatedComponent and a slide function or setting @bSlideAlongSurface to false and then using a separate sliding function.
+ 	 * Note: This function doesn't update a movement record so velocity should be gathered/set using a different method
+ 	 * @param MovingComps Encapsulates components involved in movement
+ 	 * @param Delta The desired location change in world space
+ 	 * @param NewRotation The new desired rotation in world space
+ 	 * @param bSweep Whether we sweep to the destination location
+ 	 * @param OutHit Optional output describing the blocking hit that stopped the move, if any
+ 	 * @param Teleport Whether we teleport the physics state (if physics collision is enabled for this object)
+ 	 * @param bSlideAlongSurface If true the actor slides along a blocking surface. If false the actor will stop if the move was blocked
+ 	 * @return Returns the percent of the move applied 0 to 1, where 1 represents the whole move being applied
+ 	 */
+	UFUNCTION(BlueprintCallable, Category=Mover, meta=(DisplayName="Try Safe Move And Slide (No Movement Record)"))
+	static float TrySafeMoveAndSlideUpdatedComponentNoMovementRecord(const FMovingComponentSet& MovingComps, const FVector& Delta, const FQuat& NewRotation, bool bSweep, UPARAM(Ref) FHitResult& OutHit, ETeleportType Teleport, bool bSlideAlongSurface = true);
+	
+	/**
+	 * Attempts to move a component and resolve any penetration issues with the proposed move Delta
+	 * Note: This function takes a movement record that collects moves applied to the actor see @FMovementRecord
+	 */
+	UFUNCTION(BlueprintCallable, Category=Mover, meta=(DisplayName="Try Safe Move Updated Component (Movement Record)"))
+	static bool TrySafeMoveUpdatedComponent(const FMovingComponentSet& MovingComps, const FVector& Delta, const FQuat& NewRotation, bool bSweep, UPARAM(Ref) FHitResult& OutHit, ETeleportType Teleport, UPARAM(Ref) FMovementRecord& MoveRecord);
+	
+	/**
+ 	 * Attempts to move a component and resolve any penetration issues with the proposed move Delta
+ 	 * Note: This function doesn't update a movement record so velocity should be gathered/set using a different method
+ 	 */
+	UFUNCTION(BlueprintCallable, Category=Mover, meta=(DisplayName="Try Safe Move Updated Component (No Movement Record)"))
+	static bool TrySafeMoveUpdatedComponentNoMovementRecord(const FMovingComponentSet& MovingComps, const FVector& Delta, const FQuat& NewRotation, bool bSweep, UPARAM(Ref) FHitResult& OutHit, ETeleportType Teleport);
 
 	/** Returns a movement step that should get the subject of the hit result out of an initial penetration condition */
 	static FVector ComputePenetrationAdjustment(const FHitResult& Hit);
@@ -222,12 +279,12 @@ public:
 	static bool IsAngularVelocityZero(const FRotator& AngularVelocity);
 
 	/**
-	 * Try to find an acceptable non-colliding location to place TestActor as close to possible to PlaceLocation. Expects PlaceLocation to be a valid location inside the level.
-	 * Returns true if a location without blocking collision is found, in which case PlaceLocation is overwritten with the new clear location.
-	 * Returns false if no suitable location could be found, in which case PlaceLocation is unmodified.
+	 * Try to find an acceptable non-colliding location to place TestActor as close to possible to TestLocation. Expects TestLocation to be a valid location inside the level.
+	 * Returns true if a location without blocking collision is found, in which case TestLocation is overwritten with the new clear location.
+	 * Returns false if no suitable location could be found, in which case TestLocation is unmodified.
 	 */
 	UFUNCTION(BlueprintCallable, Category = Mover)
-	static bool FindTeleportSpot(const UMoverComponent* MoverComp, FVector& TestLocation, FRotator TestRotation);
+	static bool FindTeleportSpot(const UMoverComponent* MoverComp, UPARAM(Ref) FVector& TestLocation, FRotator TestRotation);
 
 	/** Returns whether MoverComp's Actor would encroach at TestLocation on something that blocks it. */
 	UFUNCTION(BlueprintCallable, Category = Mover)
diff --git a/Engine/Plugins/Experimental/Mover/Source/Mover/Public/MoveLibrary/WaterMovementUtils.h b/Engine/Plugins/Experimental/Mover/Source/Mover/Public/MoveLibrary/WaterMovementUtils.h
--- a/Engine/Plugins/Experimental/Mover/Source/Mover/Public/MoveLibrary/WaterMovementUtils.h
+++ b/Engine/Plugins/Experimental/Mover/Source/Mover/Public/MoveLibrary/WaterMovementUtils.h
@@ -17,49 +17,49 @@ struct MOVER_API FWaterMoveParams
 {
 	GENERATED_BODY()
 
-	UPROPERTY(BlueprintReadOnly, EditAnywhere, Category = Mover)
+	UPROPERTY(BlueprintReadWrite, EditAnywhere, Category = Mover)
 	EMoveInputType MoveInputType = EMoveInputType::DirectionalIntent;
 
-	UPROPERTY(BlueprintReadOnly, EditAnywhere, Category = Mover)
+	UPROPERTY(BlueprintReadWrite, EditAnywhere, Category = Mover)
 	FVector MoveInput = FVector::ZeroVector;
 
-	UPROPERTY(BlueprintReadOnly, EditAnywhere, Category = Mover)
+	UPROPERTY(BlueprintReadWrite, EditAnywhere, Category = Mover)
 	FRotator OrientationIntent = FRotator::ZeroRotator;
 
-	UPROPERTY(BlueprintReadOnly, EditAnywhere, Category = Mover)
+	UPROPERTY(BlueprintReadWrite, EditAnywhere, Category = Mover)
 	FVector PriorVelocity = FVector::ZeroVector;
 
-	UPROPERTY(BlueprintReadOnly, EditAnywhere, Category = Mover)
+	UPROPERTY(BlueprintReadWrite, EditAnywhere, Category = Mover)
 	FRotator PriorOrientation = FRotator::ZeroRotator;
 
-	UPROPERTY(BlueprintReadOnly, EditAnywhere, Category = Mover)
+	UPROPERTY(BlueprintReadWrite, EditAnywhere, Category = Mover)
 	float MaxSpeed = 800.f;
 
-	UPROPERTY(BlueprintReadOnly, EditAnywhere, Category = Mover)
+	UPROPERTY(BlueprintReadWrite, EditAnywhere, Category = Mover)
 	float Acceleration = 4000.f;
 
-	UPROPERTY(BlueprintReadOnly, EditAnywhere, Category = Mover)
+	UPROPERTY(BlueprintReadWrite, EditAnywhere, Category = Mover)
 	float Deceleration = 8000.f;
 
-	UPROPERTY(BlueprintReadOnly, EditAnywhere, Category = Mover)
+	UPROPERTY(BlueprintReadWrite, EditAnywhere, Category = Mover)
 	float Friction = 0.f;
 
-	UPROPERTY(BlueprintReadOnly, EditAnywhere, Category = Mover)
+	UPROPERTY(BlueprintReadWrite, EditAnywhere, Category = Mover)
 	float TurningRate = 500.f;
 
-	UPROPERTY(BlueprintReadOnly, EditAnywhere, Category = Mover)
+	UPROPERTY(BlueprintReadWrite, EditAnywhere, Category = Mover)
 	float TurningBoost = 8.f;
 
-	UPROPERTY(BlueprintReadOnly, EditAnywhere, Category = Mover)
+	UPROPERTY(BlueprintReadWrite, EditAnywhere, Category = Mover)
 	float DeltaSeconds = 0.f;
 	
-	UPROPERTY(BlueprintReadOnly, EditAnywhere, Category = Mover)
+	UPROPERTY(BlueprintReadWrite, EditAnywhere, Category = Mover)
 	FVector MoveAcceleration = FVector::ZeroVector;
 
-	UPROPERTY(BlueprintReadOnly, EditAnywhere, Category = Mover)
+	UPROPERTY(BlueprintReadWrite, EditAnywhere, Category = Mover)
 	float MoveSpeed = 0.f;
 	
-	UPROPERTY(BlueprintReadOnly, EditAnywhere, Category = Mover)
+	UPROPERTY(BlueprintReadWrite, EditAnywhere, Category = Mover)
 	FQuat WorldToGravityQuat = FQuat::Identity;
 };
 
@@ -159,22 +159,22 @@ struct MOVER_API FUpdateWaterSplineDataParams
 {
 	GENERATED_BODY()
 
-	UPROPERTY(BlueprintReadOnly, EditAnywhere, Category = Mover)
+	UPROPERTY(BlueprintReadWrite, EditAnywhere, Category = Mover)
 	float TargetImmersionDepth = 0.f;
 	
-	UPROPERTY(BlueprintReadOnly, EditAnywhere, Category = Mover)
+	UPROPERTY(BlueprintReadWrite, EditAnywhere, Category = Mover)
 	float WaterVelocityDepthForMax = 0.f;
 	
-	UPROPERTY(BlueprintReadOnly, EditAnywhere, Category = Mover)
+	UPROPERTY(BlueprintReadWrite, EditAnywhere, Category = Mover)
 	float WaterVelocityMinMultiplier = 0.f;
 
-	UPROPERTY(BlueprintReadOnly, EditAnywhere, Category = Mover)
+	UPROPERTY(BlueprintReadWrite, EditAnywhere, Category = Mover)
 	FVector PlayerVelocity = FVector::ZeroVector;
 	
-	UPROPERTY(BlueprintReadOnly, EditAnywhere, Category = Mover)
+	UPROPERTY(BlueprintReadWrite, EditAnywhere, Category = Mover)
 	FVector PlayerLocation = FVector::ZeroVector;
 
-	UPROPERTY(BlueprintReadOnly, EditAnywhere, Category = Mover)
+	UPROPERTY(BlueprintReadWrite, EditAnywhere, Category = Mover)
 	float CapsuleHalfHeight = 0.f;	
 };
 
@@ -193,5 +193,5 @@ public:
 	
 	/** Updates the water spline data used in calculated swimming movement */
 	UFUNCTION(BlueprintCallable, Category=Mover)
-	static void UpdateWaterSplineData(const FUpdateWaterSplineDataParams& UpdateWaterSplineDataParams, FWaterCheckResult& OutWaterResult);
+	static void UpdateWaterSplineData(const FUpdateWaterSplineDataParams& UpdateWaterSplineDataParams, UPARAM(ref) FWaterCheckResult& WaterCheckResult);
 };
diff --git a/Engine/Plugins/Experimental/Mover/Source/Mover/Public/MoverComponent.h b/Engine/Plugins/Experimental/Mover/Source/Mover/Public/MoverComponent.h
--- a/Engine/Plugins/Experimental/Mover/Source/Mover/Public/MoverComponent.h
+++ b/Engine/Plugins/Experimental/Mover/Source/Mover/Public/MoverComponent.h
@@ -482,7 +482,7 @@ public:	// Queries
    	 * @param TagsToAdd			Tags to add/append to the Mover Component
    	 */
 	UFUNCTION(BlueprintCallable, Category = Mover, meta = (Keywords = "Add Tag"))
-	void AddGameplayTags(UPARAM(Ref) const FGameplayTagContainer& TagsToAdd);
+	void AddGameplayTags(const FGameplayTagContainer& TagsToAdd);
 	
 	/**
    	 * Removes a gameplay tag from this Mover Component
@@ -496,7 +496,7 @@ public:	// Queries
 	 * @param TagsToRemove			Tags to remove from the Mover Component
 	 */
 	UFUNCTION(BlueprintCallable, Category = Mover, meta = (Keywords = "Remove Tag"))
-	void RemoveGameplayTags(UPARAM(Ref) const FGameplayTagContainer& TagsToRemove);
+	void RemoveGameplayTags(const FGameplayTagContainer& TagsToRemove);
 	
 protected:
 
