[Mover] Implementing required virtuals on FMoverLaunchInputs for use with the physics mover

diff --git a/Engine/Plugins/Experimental/Mover/Source/Mover/Private/PhysicsMover/PhysicsMoverSimulationTypes.cpp b/Engine/Plugins/Experimental/Mover/Source/Mover/Private/PhysicsMover/PhysicsMoverSimulationTypes.cpp
--- a/Engine/Plugins/Experimental/Mover/Source/Mover/Private/PhysicsMover/PhysicsMoverSimulationTypes.cpp
+++ b/Engine/Plugins/Experimental/Mover/Source/Mover/Private/PhysicsMover/PhysicsMoverSimulationTypes.cpp
@@ -180,3 +180,26 @@ void FMoverLaunchInputs::ToString(FAnsiStringBuilderBase& Out) const
 	Out.Appendf("LaunchVelocity: X=%.2f Y=%.2f Z=%.2f\n", LaunchVelocity.X, LaunchVelocity.Y, LaunchVelocity.Z);
 	Out.Appendf("Mode: %u\n", Mode);
 }
+
+bool FMoverLaunchInputs::ShouldReconcile(const FMoverDataStructBase& AuthorityState) const
+{
+	const FMoverLaunchInputs& TypedAuthority = static_cast<const FMoverLaunchInputs&>(AuthorityState);
+	return Mode != TypedAuthority.Mode || LaunchVelocity != TypedAuthority.LaunchVelocity;
+}
+
+void FMoverLaunchInputs::Interpolate(const FMoverDataStructBase& From, const FMoverDataStructBase& To, float Pct)
+{
+	if (Pct < .5f)
+	{
+		*this = static_cast<const FMoverLaunchInputs&>(From);
+	}
+	else
+	{
+		*this = static_cast<const FMoverLaunchInputs&>(To);
+	}
+}
+
+void FMoverLaunchInputs::Merge(const FMoverDataStructBase& From)
+{
+	
+}
diff --git a/Engine/Plugins/Experimental/Mover/Source/Mover/Public/PhysicsMover/PhysicsMoverSimulationTypes.h b/Engine/Plugins/Experimental/Mover/Source/Mover/Public/PhysicsMover/PhysicsMoverSimulationTypes.h
--- a/Engine/Plugins/Experimental/Mover/Source/Mover/Public/PhysicsMover/PhysicsMoverSimulationTypes.h
+++ b/Engine/Plugins/Experimental/Mover/Source/Mover/Public/PhysicsMover/PhysicsMoverSimulationTypes.h
@@ -217,30 +217,19 @@ struct MOVER_API FMoverLaunchInputs : public FMoverDataStructBase
 
 public:
 	// Launch velocity in cm/s
-	UPROPERTY(BlueprintReadWrite, Category = Mover)
-	FVector LaunchVelocity;
+	UPROPERTY(BlueprintReadWrite, Category = Mover, meta = (Units = "cm/s"))
+	FVector LaunchVelocity = FVector::ZeroVector;
 
 	UPROPERTY(BlueprintReadWrite, Category = Mover)
-	EMoverLaunchVelocityMode Mode;
-
-	FMoverLaunchInputs()
-		: LaunchVelocity(FVector::ZeroVector)
-		, Mode(EMoverLaunchVelocityMode::Additive)
-	{
-	}
+	EMoverLaunchVelocityMode Mode = EMoverLaunchVelocityMode::Additive;
 
-	virtual ~FMoverLaunchInputs() {}
-
-	// @return newly allocated copy of this FMovementSettingsInputs. Must be overridden by child classes
 	virtual FMoverDataStructBase* Clone() const override;
-
 	virtual bool NetSerialize(FArchive& Ar, UPackageMap* Map, bool& bOutSuccess) override;
-
 	virtual UScriptStruct* GetScriptStruct() const override { return StaticStruct(); }
-
 	virtual void ToString(FAnsiStringBuilderBase& Out) const override;
-
-	virtual void AddReferencedObjects(FReferenceCollector& Collector) override { Super::AddReferencedObjects(Collector); }
+	virtual bool ShouldReconcile(const FMoverDataStructBase& AuthorityState) const override;
+	virtual void Interpolate(const FMoverDataStructBase& From, const FMoverDataStructBase& To, float Pct) override;
+	virtual void Merge(const FMoverDataStructBase& From) override;
 };
 
 template<>
