Mover: Writing the move to the move record if NavWalkingMode went off the navmesh and EOffNavMeshBehavior allows the actor to move off the NavMesh. This properly applies velocity in the case the actor walks off the NavMesh with this setting turned on.

diff --git a/Engine/Plugins/Experimental/Mover/Source/Mover/Private/DefaultMovementSet/Modes/AsyncNavWalkingMode.cpp b/Engine/Plugins/Experimental/Mover/Source/Mover/Private/DefaultMovementSet/Modes/AsyncNavWalkingMode.cpp
--- a/Engine/Plugins/Experimental/Mover/Source/Mover/Private/DefaultMovementSet/Modes/AsyncNavWalkingMode.cpp
+++ b/Engine/Plugins/Experimental/Mover/Source/Mover/Private/DefaultMovementSet/Modes/AsyncNavWalkingMode.cpp
@@ -17,6 +17,8 @@
 
 #include UE_INLINE_GENERATED_CPP_BY_NAME(AsyncNavWalkingMode)
 
+static const FName MoveWithoutNavMeshSubstepName = "MoveWithoutNavMesh";
+
 UAsyncNavWalkingMode::UAsyncNavWalkingMode()
 	: bSweepWhileNavWalking(true)
 	, bProjectNavMeshWalking(false)
@@ -250,6 +252,7 @@ void UAsyncNavWalkingMode::SimulationTick_Implementation(const FSimulationTickPa
 					// allow the full move to occur 
 					LocationInProgress = StartingSyncState->GetLocation_WorldSpace() + (ProposedMove.LinearVelocity* DeltaSeconds);
 					RotationInProgress = TargetRotation;
+					MoveRecord.Append(FMovementSubstep(MoveWithoutNavMeshSubstepName, ProposedMove.LinearVelocity * DeltaSeconds, true));
 					break;
 
 				case EOffNavMeshBehavior::DoNotMove:
diff --git a/Engine/Plugins/Experimental/Mover/Source/Mover/Private/DefaultMovementSet/Modes/NavWalkingMode.cpp b/Engine/Plugins/Experimental/Mover/Source/Mover/Private/DefaultMovementSet/Modes/NavWalkingMode.cpp
--- a/Engine/Plugins/Experimental/Mover/Source/Mover/Private/DefaultMovementSet/Modes/NavWalkingMode.cpp
+++ b/Engine/Plugins/Experimental/Mover/Source/Mover/Private/DefaultMovementSet/Modes/NavWalkingMode.cpp
@@ -244,7 +244,7 @@ void UNavWalkingMode::SimulationTick_Implementation(const FSimulationTickParams&
 
 				case EOffNavMeshBehavior::MoveWithoutNavMesh:
 					// allow the full move to occur 
-					// TODO: Need to actually move the updated component
+					// TODO: Need to actually move the updated component and add the movement to the move record
 					ensureMsgf(false, TEXT("NavWalkingMode does not yet support MoveWithoutNavMesh"));
 					break;
 
