Mover: Properly cleaning up our movement modes when a MoverComponent is destroyed. Unregistered is now called on all movement modes on the state machine and Deactivate will be called on the active mode.

diff --git a/Engine/Plugins/Experimental/Mover/Source/Mover/Private/MovementModeStateMachine.cpp b/Engine/Plugins/Experimental/Mover/Source/Mover/Private/MovementModeStateMachine.cpp
--- a/Engine/Plugins/Experimental/Mover/Source/Mover/Private/MovementModeStateMachine.cpp
+++ b/Engine/Plugins/Experimental/Mover/Source/Mover/Private/MovementModeStateMachine.cpp
@@ -68,13 +68,18 @@ void UMovementModeStateMachine::UnregisterMovementMode(FName ModeName)
 
 void UMovementModeStateMachine::ClearAllMovementModes()
 {
-	Modes.Empty();
+	if (Modes.Contains(CurrentModeName))
+	{
+		Modes[CurrentModeName]->Deactivate();		
+	}
 
 	for (TPair<FName, TObjectPtr<UBaseMovementMode>>& Element : Modes)
 	{
 		Element.Value->OnUnregistered();
 	}
 
+	Modes.Empty();
+	
 	ConstructDefaultModes();	// Note that we're resetting to our defaults so we keep the null movement mode
 }
 
diff --git a/Engine/Plugins/Experimental/Mover/Source/Mover/Private/MoverComponent.cpp b/Engine/Plugins/Experimental/Mover/Source/Mover/Private/MoverComponent.cpp
--- a/Engine/Plugins/Experimental/Mover/Source/Mover/Private/MoverComponent.cpp
+++ b/Engine/Plugins/Experimental/Mover/Source/Mover/Private/MoverComponent.cpp
@@ -238,6 +238,13 @@ void UMoverComponent::BeginPlay()
 	}
 }
 
+void UMoverComponent::EndPlay(const EEndPlayReason::Type EndPlayReason)
+{
+	ModeFSM->ClearAllMovementModes();
+	
+	Super::EndPlay(EndPlayReason);
+}
+
 void UMoverComponent::BindProcessGeneratedMovement(FMover_ProcessGeneratedMovement ProcessGeneratedMovementEvent)
 {
 	ProcessGeneratedMovement = ProcessGeneratedMovementEvent;
diff --git a/Engine/Plugins/Experimental/Mover/Source/Mover/Public/MoverComponent.h b/Engine/Plugins/Experimental/Mover/Source/Mover/Public/MoverComponent.h
--- a/Engine/Plugins/Experimental/Mover/Source/Mover/Public/MoverComponent.h
+++ b/Engine/Plugins/Experimental/Mover/Source/Mover/Public/MoverComponent.h
@@ -73,7 +73,8 @@ public:
 	virtual void RegisterComponentTickFunctions(bool bRegister) override;
 	virtual void PostLoad() override;
 	virtual void BeginPlay() override;
-
+	virtual void EndPlay(const EEndPlayReason::Type EndPlayReason) override;
+	
 	// Broadcast before each simulation tick.
 	// Note - Guaranteed to run on the game thread (even in async simulation).
 	UPROPERTY(BlueprintAssignable, Category = Mover)
