Mover: added rw lock to manage access to cached state of standalone backend. Prevents crashes in async movement sims when multiple systems are reading/writing the state at the same time.

diff --git a/Engine/Plugins/Experimental/Mover/Source/Mover/Private/Backends/MoverStandaloneLiaison.cpp b/Engine/Plugins/Experimental/Mover/Source/Mover/Private/Backends/MoverStandaloneLiaison.cpp
--- a/Engine/Plugins/Experimental/Mover/Source/Mover/Private/Backends/MoverStandaloneLiaison.cpp
+++ b/Engine/Plugins/Experimental/Mover/Source/Mover/Private/Backends/MoverStandaloneLiaison.cpp
@@ -97,13 +97,21 @@ int32 UMoverStandaloneLiaisonComponent::GetCurrentSimFrame()
 
 bool UMoverStandaloneLiaisonComponent::ReadPendingSyncState(OUT FMoverSyncState& OutSyncState)
 {
-	OutSyncState = CachedLastSyncState;
+	{
+		FRWScopeLock Lock(StateDataLock, SLT_ReadOnly);
+		OutSyncState = CachedLastSyncState;
+	}
+
 	return true;
 }
 
 bool UMoverStandaloneLiaisonComponent::WritePendingSyncState(const FMoverSyncState& SyncStateToWrite)
 {
-	CachedLastSyncState = SyncStateToWrite;
+	{
+		FRWScopeLock Lock(StateDataLock, SLT_Write);
+		CachedLastSyncState = SyncStateToWrite;
+	}
+
 	return true;
 }
 
@@ -123,7 +131,8 @@ void UMoverStandaloneLiaisonComponent::BeginPlay()
 
 			MoverComp->InitMoverSimulation();
 
-			MoverComp->InitializeSimulationState(OUT & CachedLastSyncState, OUT & CachedLastAuxState);
+			FRWScopeLock Lock(StateDataLock, SLT_Write);
+			MoverComp->InitializeSimulationState(OUT &CachedLastSyncState, OUT &CachedLastAuxState);
 		}
 		else
 		{
@@ -340,23 +349,31 @@ void UMoverStandaloneLiaisonComponent::TickMovementSimulation(float DeltaSeconds
 	FMoverTickEndData EndData;
 
 	StartData.InputCmd = LastProducedInputCmd;
-	StartData.SyncState = CachedLastSyncState;
-	StartData.AuxState = CachedLastAuxState;
+
+	{
+		FRWScopeLock Lock(StateDataLock, SLT_ReadOnly);
+		StartData.SyncState = CachedLastSyncState;
+		StartData.AuxState = CachedLastAuxState;
+	}
 
 	{
 		SCOPED_NAMED_EVENT(StandaloneMoverBackendLiaison_Tick_SimulationTick, FColor::Blue);
 		MoverComp->SimulationTick(TimeStep, StartData, OUT EndData);
 	}
 
-	LastUsedInputCmd = StartData.InputCmd;
-	CachedLastSyncState = EndData.SyncState;
-	CachedLastAuxState = EndData.AuxState;
+	{
+		FRWScopeLock Lock(StateDataLock, SLT_Write);
+		CachedLastSyncState = EndData.SyncState;
+		CachedLastAuxState = EndData.AuxState;
+	}
 }
 
 void UMoverStandaloneLiaisonComponent::TickApplySimulationState(float DeltaSeconds)
 {
 	{
 		SCOPED_NAMED_EVENT(MoverBackendLiaison_Tick_FinalizeFrame, FColor::Green);
+
+		FRWScopeLock Lock(StateDataLock, SLT_ReadOnly);
 		MoverComp->FinalizeFrame(&CachedLastSyncState, &CachedLastAuxState);
 	}
 }
diff --git a/Engine/Plugins/Experimental/Mover/Source/Mover/Public/Backends/MoverStandaloneLiaison.h b/Engine/Plugins/Experimental/Mover/Source/Mover/Public/Backends/MoverStandaloneLiaison.h
--- a/Engine/Plugins/Experimental/Mover/Source/Mover/Public/Backends/MoverStandaloneLiaison.h
+++ b/Engine/Plugins/Experimental/Mover/Source/Mover/Public/Backends/MoverStandaloneLiaison.h
@@ -177,10 +177,12 @@ protected:
 	int32 CurrentSimFrame;
 
 	FMoverInputCmdContext LastProducedInputCmd;
-	FMoverInputCmdContext LastUsedInputCmd;
+
 	FMoverSyncState CachedLastSyncState;
 	FMoverAuxStateContext CachedLastAuxState;
 
+	FRWLock StateDataLock;	// used when reading/writing to our cached state (sync/aux) data
+
 	FMoverStandaloneProduceInputTickFunction ProduceInputTickFunction;
 	FMoverStandaloneSimulateMovementTickFunction SimulateMovementTickFunction;
 	FMoverStandaloneApplyStateTickFunction ApplyStateTickFunction;
