Mover: allowing move input type that's explicitly 'none', and making this the default. Will still log warnings if Invalid is passed in.

diff --git a/Engine/Plugins/Experimental/Mover/Source/Mover/Private/DefaultMovementSet/Modes/AsyncFallingMode.cpp b/Engine/Plugins/Experimental/Mover/Source/Mover/Private/DefaultMovementSet/Modes/AsyncFallingMode.cpp
--- a/Engine/Plugins/Experimental/Mover/Source/Mover/Private/DefaultMovementSet/Modes/AsyncFallingMode.cpp
+++ b/Engine/Plugins/Experimental/Mover/Source/Mover/Private/DefaultMovementSet/Modes/AsyncFallingMode.cpp
@@ -64,7 +64,7 @@ void UAsyncFallingMode::GenerateMove_Implementation(const FMoverTickStartData& S
 	}
 	else
 	{
-		Params.MoveInputType = EMoveInputType::Invalid;
+		Params.MoveInputType = EMoveInputType::None;
 		Params.MoveInput = FVector::ZeroVector;
 	}
 
diff --git a/Engine/Plugins/Experimental/Mover/Source/Mover/Private/DefaultMovementSet/Modes/AsyncFlyingMode.cpp b/Engine/Plugins/Experimental/Mover/Source/Mover/Private/DefaultMovementSet/Modes/AsyncFlyingMode.cpp
--- a/Engine/Plugins/Experimental/Mover/Source/Mover/Private/DefaultMovementSet/Modes/AsyncFlyingMode.cpp
+++ b/Engine/Plugins/Experimental/Mover/Source/Mover/Private/DefaultMovementSet/Modes/AsyncFlyingMode.cpp
@@ -36,7 +36,7 @@ void UAsyncFlyingMode::GenerateMove_Implementation(const FMoverTickStartData& St
 	}
 	else
 	{
-		Params.MoveInputType = EMoveInputType::Invalid;
+		Params.MoveInputType = EMoveInputType::None;
 		Params.MoveInput = FVector::ZeroVector;
 	}
 
diff --git a/Engine/Plugins/Experimental/Mover/Source/Mover/Private/DefaultMovementSet/Modes/AsyncNavWalkingMode.cpp b/Engine/Plugins/Experimental/Mover/Source/Mover/Private/DefaultMovementSet/Modes/AsyncNavWalkingMode.cpp
--- a/Engine/Plugins/Experimental/Mover/Source/Mover/Private/DefaultMovementSet/Modes/AsyncNavWalkingMode.cpp
+++ b/Engine/Plugins/Experimental/Mover/Source/Mover/Private/DefaultMovementSet/Modes/AsyncNavWalkingMode.cpp
@@ -85,7 +85,7 @@ void UAsyncNavWalkingMode::GenerateMove_Implementation(const FMoverTickStartData
 	}
 	else
 	{
-		Params.MoveInputType = EMoveInputType::Invalid;
+		Params.MoveInputType = EMoveInputType::None;
 		Params.MoveInput = FVector::ZeroVector;
 	}
 
diff --git a/Engine/Plugins/Experimental/Mover/Source/Mover/Private/DefaultMovementSet/Modes/AsyncWalkingMode.cpp b/Engine/Plugins/Experimental/Mover/Source/Mover/Private/DefaultMovementSet/Modes/AsyncWalkingMode.cpp
--- a/Engine/Plugins/Experimental/Mover/Source/Mover/Private/DefaultMovementSet/Modes/AsyncWalkingMode.cpp
+++ b/Engine/Plugins/Experimental/Mover/Source/Mover/Private/DefaultMovementSet/Modes/AsyncWalkingMode.cpp
@@ -79,7 +79,7 @@ void UAsyncWalkingMode::GenerateMove_Implementation(const FMoverTickStartData& S
 	}
 	else
 	{
-		Params.MoveInputType = EMoveInputType::Invalid;
+		Params.MoveInputType = EMoveInputType::None;
 		Params.MoveInput = FVector::ZeroVector;
 	}
 
diff --git a/Engine/Plugins/Experimental/Mover/Source/Mover/Private/DefaultMovementSet/Modes/FallingMode.cpp b/Engine/Plugins/Experimental/Mover/Source/Mover/Private/DefaultMovementSet/Modes/FallingMode.cpp
--- a/Engine/Plugins/Experimental/Mover/Source/Mover/Private/DefaultMovementSet/Modes/FallingMode.cpp
+++ b/Engine/Plugins/Experimental/Mover/Source/Mover/Private/DefaultMovementSet/Modes/FallingMode.cpp
@@ -56,7 +56,7 @@ void UFallingMode::GenerateMove_Implementation(const FMoverTickStartData& StartS
 	}
 	else
 	{
-		Params.MoveInputType = EMoveInputType::Invalid;
+		Params.MoveInputType = EMoveInputType::None;
 		Params.MoveInput = FVector::ZeroVector;
 	}
 
diff --git a/Engine/Plugins/Experimental/Mover/Source/Mover/Private/DefaultMovementSet/Modes/FlyingMode.cpp b/Engine/Plugins/Experimental/Mover/Source/Mover/Private/DefaultMovementSet/Modes/FlyingMode.cpp
--- a/Engine/Plugins/Experimental/Mover/Source/Mover/Private/DefaultMovementSet/Modes/FlyingMode.cpp
+++ b/Engine/Plugins/Experimental/Mover/Source/Mover/Private/DefaultMovementSet/Modes/FlyingMode.cpp
@@ -34,7 +34,7 @@ void UFlyingMode::GenerateMove_Implementation(const FMoverTickStartData& StartSt
 	}
 	else
 	{
-		Params.MoveInputType = EMoveInputType::Invalid;
+		Params.MoveInputType = EMoveInputType::None;
 		Params.MoveInput = FVector::ZeroVector;
 	}
 
diff --git a/Engine/Plugins/Experimental/Mover/Source/Mover/Private/DefaultMovementSet/Modes/NavWalkingMode.cpp b/Engine/Plugins/Experimental/Mover/Source/Mover/Private/DefaultMovementSet/Modes/NavWalkingMode.cpp
--- a/Engine/Plugins/Experimental/Mover/Source/Mover/Private/DefaultMovementSet/Modes/NavWalkingMode.cpp
+++ b/Engine/Plugins/Experimental/Mover/Source/Mover/Private/DefaultMovementSet/Modes/NavWalkingMode.cpp
@@ -75,7 +75,7 @@ void UNavWalkingMode::GenerateMove_Implementation(const FMoverTickStartData& Sta
 	}
 	else
 	{
-		Params.MoveInputType = EMoveInputType::Invalid;
+		Params.MoveInputType = EMoveInputType::None;
 		Params.MoveInput = FVector::ZeroVector;
 	}
 
diff --git a/Engine/Plugins/Experimental/Mover/Source/Mover/Private/DefaultMovementSet/Modes/SwimmingMode.cpp b/Engine/Plugins/Experimental/Mover/Source/Mover/Private/DefaultMovementSet/Modes/SwimmingMode.cpp
--- a/Engine/Plugins/Experimental/Mover/Source/Mover/Private/DefaultMovementSet/Modes/SwimmingMode.cpp
+++ b/Engine/Plugins/Experimental/Mover/Source/Mover/Private/DefaultMovementSet/Modes/SwimmingMode.cpp
@@ -138,7 +138,7 @@ void USwimmingMode::GenerateMove_Implementation(const FMoverTickStartData& Start
 		}
 		else
 		{
-			Params.MoveInputType = EMoveInputType::Invalid;
+			Params.MoveInputType = EMoveInputType::None;
 			Params.MoveInput = FVector::ZeroVector;
 		}
 		Params.OrientationIntent = IntendedOrientation_WorldSpace;
diff --git a/Engine/Plugins/Experimental/Mover/Source/Mover/Private/DefaultMovementSet/Modes/WalkingMode.cpp b/Engine/Plugins/Experimental/Mover/Source/Mover/Private/DefaultMovementSet/Modes/WalkingMode.cpp
--- a/Engine/Plugins/Experimental/Mover/Source/Mover/Private/DefaultMovementSet/Modes/WalkingMode.cpp
+++ b/Engine/Plugins/Experimental/Mover/Source/Mover/Private/DefaultMovementSet/Modes/WalkingMode.cpp
@@ -73,7 +73,7 @@ void UWalkingMode::GenerateMove_Implementation(const FMoverTickStartData& StartS
 	}
 	else
 	{
-		Params.MoveInputType = EMoveInputType::Invalid;
+		Params.MoveInputType = EMoveInputType::None;
 		Params.MoveInput = FVector::ZeroVector;
 	}
 
diff --git a/Engine/Plugins/Experimental/Mover/Source/Mover/Private/MoveLibrary/MovementUtils.cpp b/Engine/Plugins/Experimental/Mover/Source/Mover/Private/MoveLibrary/MovementUtils.cpp
--- a/Engine/Plugins/Experimental/Mover/Source/Mover/Private/MoveLibrary/MovementUtils.cpp
+++ b/Engine/Plugins/Experimental/Mover/Source/Mover/Private/MoveLibrary/MovementUtils.cpp
@@ -96,7 +96,7 @@ FVector UMovementUtils::ComputeVelocity(const FComputeVelocityParams& InParams)
 	}
 	else
 	{
-		UE_LOG(LogMover, Warning, TEXT("Mover Compute Velocity has recieved a bad input type and no velocity will be generated!"));
+		UE_CLOG((InParams.MoveInputType == EMoveInputType::Invalid), LogMover, Warning, TEXT("Mover Compute Velocity has received an invalid input type and no velocity will be generated!"));
 		return FVector::ZeroVector;
 	}
 
diff --git a/Engine/Plugins/Experimental/Mover/Source/Mover/Public/MoverDataModelTypes.h b/Engine/Plugins/Experimental/Mover/Source/Mover/Public/MoverDataModelTypes.h
--- a/Engine/Plugins/Experimental/Mover/Source/Mover/Public/MoverDataModelTypes.h
+++ b/Engine/Plugins/Experimental/Mover/Source/Mover/Public/MoverDataModelTypes.h
@@ -21,6 +21,9 @@ enum class EMoveInputType : uint8
 
 	/** Move with a given velocity (units per second) */
 	Velocity,
+
+	/** No move input of any type */
+	None,
 };
 
 
@@ -79,7 +82,7 @@ public:
 
 
 	FCharacterDefaultInputs()
-		: MoveInputType(EMoveInputType::Invalid)
+		: MoveInputType(EMoveInputType::None)
 		, MoveInput(ForceInitToZero)
 		, OrientationIntent(ForceInitToZero)
 		, ControlRotation(ForceInitToZero)
