Physics Mover - Ensure we always have a valid input for character motion.

diff --git a/Engine/Plugins/Experimental/Mover/Source/Mover/Private/Backends/MoverNetworkPhysicsLiaison.cpp b/Engine/Plugins/Experimental/Mover/Source/Mover/Private/Backends/MoverNetworkPhysicsLiaison.cpp
--- a/Engine/Plugins/Experimental/Mover/Source/Mover/Private/Backends/MoverNetworkPhysicsLiaison.cpp
+++ b/Engine/Plugins/Experimental/Mover/Source/Mover/Private/Backends/MoverNetworkPhysicsLiaison.cpp
@@ -360,10 +360,17 @@ void UMoverNetworkPhysicsLiaisonComponent::PerformPreSimulate_Internal(const FPh
 	FMoverDefaultSyncState& InputSyncState = Input.SyncState.SyncStateCollection.FindOrAddMutableDataByType<FMoverDefaultSyncState>();
 	InputSyncState.SetTransforms_WorldSpace(CharacterParticle->GetX(), FRotator(CharacterParticle->GetR()), CharacterParticle->GetV() - LocalGroundVelocity + AIMoveVelocity);
 
-	FCharacterDefaultInputs* CharacterDefaultInputs = Input.InputCmd.InputCollection.FindMutableDataByType<FCharacterDefaultInputs>();
-	if (CharacterDefaultInputs && !CharacterDefaultInputs->SuggestedMovementMode.IsNone())
+	FCharacterDefaultInputs& CharacterDefaultInputs = Input.InputCmd.InputCollection.FindOrAddMutableDataByType<FCharacterDefaultInputs>();
+	if (!CharacterDefaultInputs.SuggestedMovementMode.IsNone())
 	{
-		MoverComp.QueueNextMode(CharacterDefaultInputs->SuggestedMovementMode);
+		MoverComp.QueueNextMode(CharacterDefaultInputs.SuggestedMovementMode);
+		CharacterDefaultInputs.SuggestedMovementMode = NAME_None;
+	}
+
+	// Make sure we have a valid input for the update
+	if (CharacterDefaultInputs.GetMoveInputType() == EMoveInputType::Invalid)
+	{
+		CharacterDefaultInputs.SetMoveInput(EMoveInputType::DirectionalIntent, FVector::ZeroVector);
 	}
 
 	// Cache the particle XR before calling the parent, because the parent will teleport if the target is too far away and change the particle's transform
@@ -453,11 +460,6 @@ void UMoverNetworkPhysicsLiaisonComponent::PerformPreSimulate_Internal(const FPh
 			0.5f
 			});
 	}
-
-	if (CharacterDefaultInputs)
-	{
-		CharacterDefaultInputs->SuggestedMovementMode = NAME_None; // Should have already been set this frame
-	}
 }
 
 Chaos::FPhysicsObject* UMoverNetworkPhysicsLiaisonComponent::GetControlledPhysicsObject() const
diff --git a/Engine/Plugins/Experimental/Mover/Source/Mover/Private/Backends/MoverNetworkPhysicsLiaisonBase.cpp b/Engine/Plugins/Experimental/Mover/Source/Mover/Private/Backends/MoverNetworkPhysicsLiaisonBase.cpp
--- a/Engine/Plugins/Experimental/Mover/Source/Mover/Private/Backends/MoverNetworkPhysicsLiaisonBase.cpp
+++ b/Engine/Plugins/Experimental/Mover/Source/Mover/Private/Backends/MoverNetworkPhysicsLiaisonBase.cpp
@@ -406,9 +406,9 @@ void UMoverNetworkPhysicsLiaisonComponentBase::BeginPlay()
 	if (UPhysicsMoverManager* MoverManager = UWorld::GetSubsystem<UPhysicsMoverManager>(GetWorld()))
 	{
 		MoverManager->RegisterPhysicsMoverComponent(this);
-	}
 
-	InitializeSimOutputData();
+		InitializeSimOutputData();
+	}
 }
 
 void UMoverNetworkPhysicsLiaisonComponentBase::EndPlay(const EEndPlayReason::Type EndPlayReason)
@@ -781,7 +781,9 @@ void UMoverNetworkPhysicsLiaisonComponentBase::PerformPreSimulate_Internal(const
 	//@todo DanH: Invoking the FSM OnSimulationTick can sometimes trigger AdvanceToNextMode to happen on the PT
 	// Tick the actual simulation. This is where the proposed moves are queried and executed, affecting change to the moving actor's gameplay state and captured in the output sim state
 	FMoverTickStartData TickStartData(Input.InputCmd, Input.SyncState, FMoverAuxStateContext());
-	FMoverTickEndData TickEndData(&Input.SyncState, nullptr);
+	FMoverTickEndData TickEndData;
+	TickEndData.SyncState = TickStartData.SyncState;
+
 	MoverComp.ModeFSM->OnSimulationTick(MoverComp.UpdatedComponent, MoverComp.UpdatedCompAsPrimitive, MoverComp.GetSimBlackboard_Mutable(), TickStartData, TimeStep, TickEndData);
 
 	// Set the output sync state and fill in the movement mode
