[Mover] Tuning swimming parameters and fixing some depth calculation bugs

#changelist validated
#virtualized

diff --git a/Engine/Plugins/Experimental/Mover/Source/Mover/Private/DefaultMovementSet/Modes/SwimmingMode.cpp b/Engine/Plugins/Experimental/Mover/Source/Mover/Private/DefaultMovementSet/Modes/SwimmingMode.cpp
--- a/Engine/Plugins/Experimental/Mover/Source/Mover/Private/DefaultMovementSet/Modes/SwimmingMode.cpp
+++ b/Engine/Plugins/Experimental/Mover/Source/Mover/Private/DefaultMovementSet/Modes/SwimmingMode.cpp
@@ -52,7 +52,7 @@ void USwimmingMode::GenerateMove_Implementation(const FMoverTickStartData& Start
 
 		const float ImmersionDepth = WaterData.ImmersionDepth + CapsuleHalfHeight;
 		const float ImmersionPercent = FMath::Clamp(ImmersionDepth * ReciprocalCapsuleHeight, UE_KINDA_SMALL_NUMBER, 1.f);
-		const float IdealDepth = CommonLegacySettings->SwimmingIdealImmersionDepth;
+		const float IdealDepth = CommonLegacySettings->SwimmingIdealImmersionDepth + CapsuleHalfHeight;
 		const float IdealImmersionPercent = FMath::Clamp(IdealDepth * ReciprocalOriginalCapsuleHeight, UE_KINDA_SMALL_NUMBER, 1.f);
 
 		const float GravityForce = GetMoverComponent()->GetGravityAcceleration().Z;
diff --git a/Engine/Plugins/Experimental/Mover/Source/Mover/Private/PhysicsMover/Modes/PhysicsDrivenSwimmingMode.cpp b/Engine/Plugins/Experimental/Mover/Source/Mover/Private/PhysicsMover/Modes/PhysicsDrivenSwimmingMode.cpp
--- a/Engine/Plugins/Experimental/Mover/Source/Mover/Private/PhysicsMover/Modes/PhysicsDrivenSwimmingMode.cpp
+++ b/Engine/Plugins/Experimental/Mover/Source/Mover/Private/PhysicsMover/Modes/PhysicsDrivenSwimmingMode.cpp
@@ -149,7 +149,7 @@ void UPhysicsDrivenSwimmingMode::SimulationTick_Implementation(const FSimulation
 
 	if (WaterResult.IsSwimmableVolume())
 	{
-		const bool bIsWithinReach = FloorResult.FloorDist <= TargetHeight;
+		const bool bIsWithinReach = FloorResult.FloorDist <= TargetHeight + CommonLegacySettings->MaxStepHeight;
 		const bool bWalkTrigger = WaterResult.WaterSplineData.ImmersionDepth < CommonLegacySettings->SwimmingStopImmersionDepth;
 		const bool bFallTrigger = FMath::Clamp((WaterResult.WaterSplineData.ImmersionDepth + TargetHeight) / (2 * TargetHeight), -2.f, 2.f) < -1.f;
 	
