Mover: reducing costs of answering MoverComp queries by making sure we always have a valid cached state and inputs

diff --git a/Engine/Plugins/Experimental/Mover/Source/Mover/Private/Backends/MoverNetworkPhysicsLiaisonBase.cpp b/Engine/Plugins/Experimental/Mover/Source/Mover/Private/Backends/MoverNetworkPhysicsLiaisonBase.cpp
--- a/Engine/Plugins/Experimental/Mover/Source/Mover/Private/Backends/MoverNetworkPhysicsLiaisonBase.cpp
+++ b/Engine/Plugins/Experimental/Mover/Source/Mover/Private/Backends/MoverNetworkPhysicsLiaisonBase.cpp
@@ -624,10 +624,8 @@ void UMoverNetworkPhysicsLiaisonComponentBase::PostPhysicsUpdate_External()
 	// The MoverComponent relies on its CachedLastSyncState for a lot of information, so setting it here is what makes the resulting state of the completed sim frame
 	// "real" from the perspective the MoverComp and any objects that call getters on it
 	MoverComp.CachedLastSyncState = InterpolatedSyncState;
-	MoverComp.bHasValidCachedState = true;
-
+	MoverComp.LastMoverDefaultSyncState = MoverComp.CachedLastSyncState.SyncStateCollection.FindDataByType<FMoverDefaultSyncState>();
 	MoverComp.CachedLastUsedInputCmd = InterpolatedInputCmd;
-	MoverComp.bHasValidCachedUsedInput = true;
 
 	// Note this may return something different now because its dependent on the CachedLastSyncState that we just changed above
 	const FName NextMode = MoverComp.GetMovementModeName();
diff --git a/Engine/Plugins/Experimental/Mover/Source/Mover/Private/Debug/GameplayDebuggerCategory_Mover.cpp b/Engine/Plugins/Experimental/Mover/Source/Mover/Private/Debug/GameplayDebuggerCategory_Mover.cpp
--- a/Engine/Plugins/Experimental/Mover/Source/Mover/Private/Debug/GameplayDebuggerCategory_Mover.cpp
+++ b/Engine/Plugins/Experimental/Mover/Source/Mover/Private/Debug/GameplayDebuggerCategory_Mover.cpp
@@ -143,32 +143,26 @@ void FGameplayDebuggerCategory_Mover::CollectData(APlayerController* OwnerPC, AA
 			DataPack.ActiveTransitions.Add(FString::Printf(TEXT("%s (global)"), (Transition ? *Transition->GetClass()->GetName() : TEXT("null"))));
 		}
 		
-		if (MyMoverComponent->HasValidCachedState())
-		{
-			const FMoverSyncState& SyncState = MyMoverComponent->GetSyncState();
-
-			for (const TSharedPtr<FLayeredMoveBase>& ActiveMove : SyncState.LayeredMoves.GetActiveMoves())
-			{
-				DataPack.ActiveLayeredMoves.Add(ActiveMove->ToSimpleString());
-			}
+		const FMoverSyncState& SyncState = MyMoverComponent->GetSyncState();
 
-			for (auto it = SyncState.MovementModifiers.GetActiveModifiersIterator(); it; ++it)
-			{
-				DataPack.ActiveModifiers.Add(*it->Get()->ToSimpleString());
-			}
+		for (const TSharedPtr<FLayeredMoveBase>& ActiveMove : SyncState.LayeredMoves.GetActiveMoves())
+		{
+			DataPack.ActiveLayeredMoves.Add(ActiveMove->ToSimpleString());
 		}
 
-		if (MyMoverComponent->HasValidCachedInputCmd())
+		for (auto It = SyncState.MovementModifiers.GetActiveModifiersIterator(); It; ++It)
 		{
-			const FMoverInputCmdContext& LastInputCmd = MyMoverComponent->GetLastInputCmd();
+			DataPack.ActiveModifiers.Add(*It->Get()->ToSimpleString());
+		}
 
-			if (const FCharacterDefaultInputs* DefaultInputs = LastInputCmd.InputCollection.FindDataByType<FCharacterDefaultInputs>())
-			{
-				DataPack.MoveInputType = (int8)DefaultInputs->GetMoveInputType();
-				DataPack.MoveInput = DefaultInputs->GetMoveInput_WorldSpace();
-				DataPack.OrientIntentDir = DefaultInputs->GetOrientationIntentDir_WorldSpace();
-				DataPack.SuggestedModeName = DefaultInputs->SuggestedMovementMode.ToString();
-			}
+		const FMoverInputCmdContext& LastInputCmd = MyMoverComponent->GetLastInputCmd();
+
+		if (const FCharacterDefaultInputs* DefaultInputs = LastInputCmd.InputCollection.FindDataByType<FCharacterDefaultInputs>())
+		{
+			DataPack.MoveInputType = (int8)DefaultInputs->GetMoveInputType();
+			DataPack.MoveInput = DefaultInputs->GetMoveInput_WorldSpace();
+			DataPack.OrientIntentDir = DefaultInputs->GetOrientationIntentDir_WorldSpace();
+			DataPack.SuggestedModeName = DefaultInputs->SuggestedMovementMode.ToString();
 		}
 
 	}
@@ -287,7 +281,7 @@ void FGameplayDebuggerCategory_Mover::DrawInWorldInfo(AActor& DebugActor, FGamep
 		}
 
 		// Draw overlaid arrows showing target orientation and actual
-		if (MoverComp && MoverComp->HasValidCachedState())
+		if (MoverComp)
 		{
 			const FMoverSyncState& LastState = MoverComp->GetSyncState();
 			const FMoverDefaultSyncState* MoverState = LastState.SyncStateCollection.FindDataByType<FMoverDefaultSyncState>();
diff --git a/Engine/Plugins/Experimental/Mover/Source/Mover/Private/Debug/MoverDebugComponent.cpp b/Engine/Plugins/Experimental/Mover/Source/Mover/Private/Debug/MoverDebugComponent.cpp
--- a/Engine/Plugins/Experimental/Mover/Source/Mover/Private/Debug/MoverDebugComponent.cpp
+++ b/Engine/Plugins/Experimental/Mover/Source/Mover/Private/Debug/MoverDebugComponent.cpp
@@ -64,7 +64,7 @@ void UMoverDebugComponent::TickComponent(float DeltaTime, ELevelTick TickType, F
 	if (bIsTrackingHistory)
 	{
 		UMoverComponent* MoverComp = GetOwner()->FindComponentByClass<UMoverComponent>();
-		if (MoverComp && MoverComp->HasValidCachedState())
+		if (MoverComp && MoverComp->CachedLastSimTickTimeStep.ServerFrame > 0)
 		{
 			UpdateHistoryTrackingForFrame(MoverComp->CachedLastSimTickTimeStep, &MoverComp->CachedLastSyncState, &MoverComp->CachedLastAuxState);
 		}
diff --git a/Engine/Plugins/Experimental/Mover/Source/Mover/Private/MoveLibrary/BasedMovementUtils.cpp b/Engine/Plugins/Experimental/Mover/Source/Mover/Private/MoveLibrary/BasedMovementUtils.cpp
--- a/Engine/Plugins/Experimental/Mover/Source/Mover/Private/MoveLibrary/BasedMovementUtils.cpp
+++ b/Engine/Plugins/Experimental/Mover/Source/Mover/Private/MoveLibrary/BasedMovementUtils.cpp
@@ -337,169 +337,167 @@ void UBasedMovementUtils::UpdateSimpleBasedMovement(UMoverComponent* TargetMover
 	}
 
 	bool bDidGetUpToDate = false;
-	if (TargetMoverComp->HasValidCachedState())
-	{
-		FRelativeBaseInfo LastFoundBaseInfo;	// Last-found is the most recent capture during movement, likely set this sim frame
-		FRelativeBaseInfo LastAppliedBaseInfo;	// Last-applied is the one that our based movement is up to date with, likely set in the last sim frame
-		FRelativeBaseInfo CurrentBaseInfo;		// Current info is the current snapshot of the current base, with up-to-date transform that may be different than last-found.
 
-		const bool bHasLastFoundInfo = SimBlackboard->TryGet(CommonBlackboard::LastFoundDynamicMovementBase, LastFoundBaseInfo);
-		const bool bHasLastAppliedInfo = SimBlackboard->TryGet(CommonBlackboard::LastAppliedDynamicMovementBase, LastAppliedBaseInfo);
-		if (bHasLastFoundInfo)
+	FRelativeBaseInfo LastFoundBaseInfo;	// Last-found is the most recent capture during movement, likely set this sim frame
+	FRelativeBaseInfo LastAppliedBaseInfo;	// Last-applied is the one that our based movement is up to date with, likely set in the last sim frame
+	FRelativeBaseInfo CurrentBaseInfo;		// Current info is the current snapshot of the current base, with up-to-date transform that may be different than last-found.
+
+	const bool bHasLastFoundInfo = SimBlackboard->TryGet(CommonBlackboard::LastFoundDynamicMovementBase, LastFoundBaseInfo);
+	const bool bHasLastAppliedInfo = SimBlackboard->TryGet(CommonBlackboard::LastAppliedDynamicMovementBase, LastAppliedBaseInfo);
+	if (bHasLastFoundInfo)
+	{
+		if (!bHasLastAppliedInfo || !LastFoundBaseInfo.UsesSameBase(LastAppliedBaseInfo))
 		{
-			if (!bHasLastAppliedInfo || !LastFoundBaseInfo.UsesSameBase(LastAppliedBaseInfo))
-			{
-				LastAppliedBaseInfo = LastFoundBaseInfo;	// This is the first time we've checked this base, so start with the last-found capture
-			}
+			LastAppliedBaseInfo = LastFoundBaseInfo;	// This is the first time we've checked this base, so start with the last-found capture
+		}
 
-			if (!ensureMsgf(LastFoundBaseInfo.HasRelativeInfo() && LastFoundBaseInfo.UsesSameBase(LastAppliedBaseInfo),
-					TEXT("Attempting to update based movement with a missing or mismatched base. This may indicate a logic problem with detecting bases.")))
-			{ 
-				SimBlackboard->Invalidate(CommonBlackboard::LastFoundDynamicMovementBase);
-				SimBlackboard->Invalidate(CommonBlackboard::LastAppliedDynamicMovementBase);
-				return;
-			}
+		if (!ensureMsgf(LastFoundBaseInfo.HasRelativeInfo() && LastFoundBaseInfo.UsesSameBase(LastAppliedBaseInfo),
+				TEXT("Attempting to update based movement with a missing or mismatched base. This may indicate a logic problem with detecting bases.")))
+		{ 
+			SimBlackboard->Invalidate(CommonBlackboard::LastFoundDynamicMovementBase);
+			SimBlackboard->Invalidate(CommonBlackboard::LastAppliedDynamicMovementBase);
+			return;
+		}
 
-			CurrentBaseInfo.SetFromComponent(LastFoundBaseInfo.MovementBase.Get(), LastFoundBaseInfo.BoneName);
-			CurrentBaseInfo.ContactLocalPosition = LastFoundBaseInfo.ContactLocalPosition;
+		CurrentBaseInfo.SetFromComponent(LastFoundBaseInfo.MovementBase.Get(), LastFoundBaseInfo.BoneName);
+		CurrentBaseInfo.ContactLocalPosition = LastFoundBaseInfo.ContactLocalPosition;
 
-			FVector CurrentBaseLocation;
-			FQuat CurrentBaseQuat;
+		FVector CurrentBaseLocation;
+		FQuat CurrentBaseQuat;
 			
 
-			if (UBasedMovementUtils::GetMovementBaseTransform(CurrentBaseInfo.MovementBase.Get(), CurrentBaseInfo.BoneName, OUT CurrentBaseLocation, OUT CurrentBaseQuat))
-			{
-				const bool bDidBaseRotationChange = !LastAppliedBaseInfo.Rotation.Equals(CurrentBaseQuat, UE_SMALL_NUMBER);
-				const bool bDidBaseLocationChange = (LastAppliedBaseInfo.Location != CurrentBaseLocation);
+		if (UBasedMovementUtils::GetMovementBaseTransform(CurrentBaseInfo.MovementBase.Get(), CurrentBaseInfo.BoneName, OUT CurrentBaseLocation, OUT CurrentBaseQuat))
+		{
+			const bool bDidBaseRotationChange = !LastAppliedBaseInfo.Rotation.Equals(CurrentBaseQuat, UE_SMALL_NUMBER);
+			const bool bDidBaseLocationChange = (LastAppliedBaseInfo.Location != CurrentBaseLocation);
 
-				FQuat DeltaQuat = FQuat::Identity;
-				FVector WorldDeltaLocation = FVector::ZeroVector;
-				FQuat WorldTargetQuat = UpdatedComponent->GetComponentQuat();
+			FQuat DeltaQuat = FQuat::Identity;
+			FVector WorldDeltaLocation = FVector::ZeroVector;
+			FQuat WorldTargetQuat = UpdatedComponent->GetComponentQuat();
 
-				// Find change in rotation
+			// Find change in rotation
 
-				if (bDidBaseRotationChange && !bIgnoreBaseRotation)
-				{
-					DeltaQuat = CurrentBaseQuat * LastAppliedBaseInfo.Rotation.Inverse();
-					WorldTargetQuat = DeltaQuat * WorldTargetQuat;
+			if (bDidBaseRotationChange && !bIgnoreBaseRotation)
+			{
+				DeltaQuat = CurrentBaseQuat * LastAppliedBaseInfo.Rotation.Inverse();
+				WorldTargetQuat = DeltaQuat * WorldTargetQuat;
 
-					FVector TargetForwVector = WorldTargetQuat.GetForwardVector();
-					TargetForwVector = FVector::VectorPlaneProject(TargetForwVector, -TargetMoverComp->GetUpDirection());
-					TargetForwVector.Normalize();
+				FVector TargetForwVector = WorldTargetQuat.GetForwardVector();
+				TargetForwVector = FVector::VectorPlaneProject(TargetForwVector, -TargetMoverComp->GetUpDirection());
+				TargetForwVector.Normalize();
 
-					FVector TargetRightVector = WorldTargetQuat.GetRightVector();
-					TargetRightVector = FVector::VectorPlaneProject(TargetRightVector, -TargetMoverComp->GetUpDirection());
-					TargetRightVector.Normalize();
+				FVector TargetRightVector = WorldTargetQuat.GetRightVector();
+				TargetRightVector = FVector::VectorPlaneProject(TargetRightVector, -TargetMoverComp->GetUpDirection());
+				TargetRightVector.Normalize();
 					
-					WorldTargetQuat = UKismetMathLibrary::MakeRotFromXY(TargetForwVector, TargetRightVector).Quaternion();
-				}
+				WorldTargetQuat = UKismetMathLibrary::MakeRotFromXY(TargetForwVector, TargetRightVector).Quaternion();
+			}
 
-				if (bDidBaseLocationChange || bDidBaseRotationChange)
-				{
-					// Calculate new transform matrix of base actor (ignoring scale).
-					const FQuatRotationTranslationMatrix OldLocalToWorld(LastAppliedBaseInfo.Rotation, LastAppliedBaseInfo.Location);
-					const FQuatRotationTranslationMatrix NewLocalToWorld(CurrentBaseQuat, CurrentBaseLocation);
-
-					// Find change in location
-					// NOTE that we are using the floor hit location, not the actor's root position which may be floating above the base
-					const FVector NewWorldBaseContactPos = NewLocalToWorld.TransformPosition(CurrentBaseInfo.ContactLocalPosition);
-					const FVector OldWorldBaseContactPos = OldLocalToWorld.TransformPosition(CurrentBaseInfo.ContactLocalPosition);
-					WorldDeltaLocation = NewWorldBaseContactPos - OldWorldBaseContactPos;
-
-					const FVector OldWorldLocation = UpdatedComponent->GetComponentLocation();
-					EMoveComponentFlags MoveComponentFlags = MOVECOMP_IgnoreBases;
-					const bool bSweep = true;
-					FHitResult MoveHitResult;
-
-					bool bDidMove = UMovementUtils::TryMoveUpdatedComponent_Internal(FMovingComponentSet(TargetMoverComp), WorldDeltaLocation, WorldTargetQuat, bSweep, MoveComponentFlags, &MoveHitResult, ETeleportType::None);
+			if (bDidBaseLocationChange || bDidBaseRotationChange)
+			{
+				// Calculate new transform matrix of base actor (ignoring scale).
+				const FQuatRotationTranslationMatrix OldLocalToWorld(LastAppliedBaseInfo.Rotation, LastAppliedBaseInfo.Location);
+				const FQuatRotationTranslationMatrix NewLocalToWorld(CurrentBaseQuat, CurrentBaseLocation);
+
+				// Find change in location
+				// NOTE that we are using the floor hit location, not the actor's root position which may be floating above the base
+				const FVector NewWorldBaseContactPos = NewLocalToWorld.TransformPosition(CurrentBaseInfo.ContactLocalPosition);
+				const FVector OldWorldBaseContactPos = OldLocalToWorld.TransformPosition(CurrentBaseInfo.ContactLocalPosition);
+				WorldDeltaLocation = NewWorldBaseContactPos - OldWorldBaseContactPos;
+
+				const FVector OldWorldLocation = UpdatedComponent->GetComponentLocation();
+				EMoveComponentFlags MoveComponentFlags = MOVECOMP_IgnoreBases;
+				const bool bSweep = true;
+				FHitResult MoveHitResult;
+
+				bool bDidMove = UMovementUtils::TryMoveUpdatedComponent_Internal(FMovingComponentSet(TargetMoverComp), WorldDeltaLocation, WorldTargetQuat, bSweep, MoveComponentFlags, &MoveHitResult, ETeleportType::None);
 					
-					const FVector NewWorldLocation = UpdatedComponent->GetComponentLocation();
+				const FVector NewWorldLocation = UpdatedComponent->GetComponentLocation();
 
-					if ((NewWorldLocation - (OldWorldLocation + WorldDeltaLocation)).IsNearlyZero() == false)
-					{
-						// Find the remaining delta that wasn't achieved
-						const FVector UnachievedWorldDelta = (OldWorldLocation + WorldDeltaLocation) - NewWorldLocation;
+				if ((NewWorldLocation - (OldWorldLocation + WorldDeltaLocation)).IsNearlyZero() == false)
+				{
+					// Find the remaining delta that wasn't achieved
+					const FVector UnachievedWorldDelta = (OldWorldLocation + WorldDeltaLocation) - NewWorldLocation;
 
-						// Convert the remaining delta to current base space
-						FVector UnachievedLocalDelta;
-						UBasedMovementUtils::TransformLocationToLocal(CurrentBaseLocation, CurrentBaseQuat, UnachievedWorldDelta, OUT UnachievedLocalDelta);
+					// Convert the remaining delta to current base space
+					FVector UnachievedLocalDelta;
+					UBasedMovementUtils::TransformLocationToLocal(CurrentBaseLocation, CurrentBaseQuat, UnachievedWorldDelta, OUT UnachievedLocalDelta);
 						
-						// Subtract the remaining delta to reflect the change in the contact position
-						CurrentBaseInfo.ContactLocalPosition -= UnachievedLocalDelta;
-					}
+					// Subtract the remaining delta to reflect the change in the contact position
+					CurrentBaseInfo.ContactLocalPosition -= UnachievedLocalDelta;
+				}
 
-					// Propagate the movement changes to the backend's state, if supported
+				// Propagate the movement changes to the backend's state, if supported
 
-					// Note that this is occurring out-of-band with the movement simulation, in order to support based movement regardless of update order or
-					// whether the movement base is also simulated through Mover.
+				// Note that this is occurring out-of-band with the movement simulation, in order to support based movement regardless of update order or
+				// whether the movement base is also simulated through Mover.
 
-					FMoverSyncState PendingSimSyncState;
-					if (TargetMoverComp->BackendLiaisonComp->ReadPendingSyncState(OUT PendingSimSyncState))
+				FMoverSyncState PendingSimSyncState;
+				if (TargetMoverComp->BackendLiaisonComp->ReadPendingSyncState(OUT PendingSimSyncState))
+				{
+					// Modify the PENDING sync state that has not yet been committed to simulation history nor replicated
+					if (FMoverDefaultSyncState* PendingMoverState = PendingSimSyncState.SyncStateCollection.FindMutableDataByType<FMoverDefaultSyncState>())
 					{
-						// Modify the PENDING sync state that has not yet been committed to simulation history nor replicated
-						if (FMoverDefaultSyncState* PendingMoverState = PendingSimSyncState.SyncStateCollection.FindMutableDataByType<FMoverDefaultSyncState>())
-						{
-							FTransform OldSyncTransformWs = PendingMoverState->GetTransform_WorldSpace();
-							FTransform NewSyncTransformWs = UpdatedComponent->GetComponentTransform();
+						FTransform OldSyncTransformWs = PendingMoverState->GetTransform_WorldSpace();
+						FTransform NewSyncTransformWs = UpdatedComponent->GetComponentTransform();
 
-							PendingMoverState->SetTransforms_WorldSpace(
-								NewSyncTransformWs.GetLocation(),
-								NewSyncTransformWs.GetRotation().Rotator(),
-								PendingMoverState->GetVelocity_WorldSpace(),	// keep same velocity and base
-								PendingMoverState->GetMovementBase(), PendingMoverState->GetMovementBaseBoneName());
+						PendingMoverState->SetTransforms_WorldSpace(
+							NewSyncTransformWs.GetLocation(),
+							NewSyncTransformWs.GetRotation().Rotator(),
+							PendingMoverState->GetVelocity_WorldSpace(),	// keep same velocity and base
+							PendingMoverState->GetMovementBase(), PendingMoverState->GetMovementBaseBoneName());
 
-							TargetMoverComp->BackendLiaisonComp->WritePendingSyncState(PendingSimSyncState);	// writes pending Simulation state
+						TargetMoverComp->BackendLiaisonComp->WritePendingSyncState(PendingSimSyncState);	// writes pending Simulation state
 
-							// If smoothing, modify presentation-related states as well so that the visual offset location stays anchored to the movement base
-							if (TargetMoverComp->SmoothingMode != EMoverSmoothingMode::None)
-							{
-								const FTransform OldToNewTransform = NewSyncTransformWs.GetRelativeTransform(OldSyncTransformWs);
+						// If smoothing, modify presentation-related states as well so that the visual offset location stays anchored to the movement base
+						if (TargetMoverComp->SmoothingMode != EMoverSmoothingMode::None)
+						{
+							const FTransform OldToNewTransform = NewSyncTransformWs.GetRelativeTransform(OldSyncTransformWs);
 
-								// Modify the PRESENTATION sync state that we're smoothing TO
-								FMoverSyncState PresentationSyncState;
-								if (TargetMoverComp->BackendLiaisonComp->ReadPresentationSyncState(OUT PresentationSyncState))
+							// Modify the PRESENTATION sync state that we're smoothing TO
+							FMoverSyncState PresentationSyncState;
+							if (TargetMoverComp->BackendLiaisonComp->ReadPresentationSyncState(OUT PresentationSyncState))
+							{
+								if (FMoverDefaultSyncState* PresentationMoverState = PresentationSyncState.SyncStateCollection.FindMutableDataByType<FMoverDefaultSyncState>())
 								{
-									if (FMoverDefaultSyncState* PresentationMoverState = PresentationSyncState.SyncStateCollection.FindMutableDataByType<FMoverDefaultSyncState>())
-									{
-										OldSyncTransformWs = PresentationMoverState->GetTransform_WorldSpace();
-										NewSyncTransformWs = OldToNewTransform * OldSyncTransformWs;
-
-										PresentationMoverState->SetTransforms_WorldSpace(
-											NewSyncTransformWs.GetLocation(),
-											NewSyncTransformWs.GetRotation().Rotator(),
-											PresentationMoverState->GetVelocity_WorldSpace(),	// keep same velocity and base
-											PresentationMoverState->GetMovementBase(), PresentationMoverState->GetMovementBaseBoneName());
-
-										TargetMoverComp->BackendLiaisonComp->WritePresentationSyncState(PresentationSyncState);
-									}
+									OldSyncTransformWs = PresentationMoverState->GetTransform_WorldSpace();
+									NewSyncTransformWs = OldToNewTransform * OldSyncTransformWs;
+
+									PresentationMoverState->SetTransforms_WorldSpace(
+										NewSyncTransformWs.GetLocation(),
+										NewSyncTransformWs.GetRotation().Rotator(),
+										PresentationMoverState->GetVelocity_WorldSpace(),	// keep same velocity and base
+										PresentationMoverState->GetMovementBase(), PresentationMoverState->GetMovementBaseBoneName());
+
+									TargetMoverComp->BackendLiaisonComp->WritePresentationSyncState(PresentationSyncState);
 								}
+							}
 
-								// Modify the PREV PRESENTATION sync state that we're smoothing FROM
-								FMoverSyncState PrevPresentationSyncState;
-								if (TargetMoverComp->BackendLiaisonComp->ReadPrevPresentationSyncState(OUT PrevPresentationSyncState))
+							// Modify the PREV PRESENTATION sync state that we're smoothing FROM
+							FMoverSyncState PrevPresentationSyncState;
+							if (TargetMoverComp->BackendLiaisonComp->ReadPrevPresentationSyncState(OUT PrevPresentationSyncState))
+							{
+								if (FMoverDefaultSyncState* PrevPresentationMoverState = PrevPresentationSyncState.SyncStateCollection.FindMutableDataByType<FMoverDefaultSyncState>())
 								{
-									if (FMoverDefaultSyncState* PrevPresentationMoverState = PrevPresentationSyncState.SyncStateCollection.FindMutableDataByType<FMoverDefaultSyncState>())
-									{
-										OldSyncTransformWs = PrevPresentationMoverState->GetTransform_WorldSpace();
-										NewSyncTransformWs = OldToNewTransform * OldSyncTransformWs;
-
-										PrevPresentationMoverState->SetTransforms_WorldSpace(
-											NewSyncTransformWs.GetLocation(),
-											NewSyncTransformWs.GetRotation().Rotator(),
-											PrevPresentationMoverState->GetVelocity_WorldSpace(),	// keep same velocity and base
-											PrevPresentationMoverState->GetMovementBase(), PrevPresentationMoverState->GetMovementBaseBoneName());
-
-										TargetMoverComp->BackendLiaisonComp->WritePrevPresentationSyncState(PrevPresentationSyncState);	
-									}
+									OldSyncTransformWs = PrevPresentationMoverState->GetTransform_WorldSpace();
+									NewSyncTransformWs = OldToNewTransform * OldSyncTransformWs;
+
+									PrevPresentationMoverState->SetTransforms_WorldSpace(
+										NewSyncTransformWs.GetLocation(),
+										NewSyncTransformWs.GetRotation().Rotator(),
+										PrevPresentationMoverState->GetVelocity_WorldSpace(),	// keep same velocity and base
+										PrevPresentationMoverState->GetMovementBase(), PrevPresentationMoverState->GetMovementBaseBoneName());
+
+									TargetMoverComp->BackendLiaisonComp->WritePrevPresentationSyncState(PrevPresentationSyncState);	
 								}
 							}
 						}
 					}
 				}
-
-				SimBlackboard->Set(CommonBlackboard::LastAppliedDynamicMovementBase, CurrentBaseInfo);
-				bDidGetUpToDate = true;
 			}
+
+			SimBlackboard->Set(CommonBlackboard::LastAppliedDynamicMovementBase, CurrentBaseInfo);
+			bDidGetUpToDate = true;
 		}
 	}
 
diff --git a/Engine/Plugins/Experimental/Mover/Source/Mover/Private/MoverComponent.cpp b/Engine/Plugins/Experimental/Mover/Source/Mover/Private/MoverComponent.cpp
--- a/Engine/Plugins/Experimental/Mover/Source/Mover/Private/MoverComponent.cpp
+++ b/Engine/Plugins/Experimental/Mover/Source/Mover/Private/MoverComponent.cpp
@@ -80,6 +80,8 @@ void UMoverComponent::InitializeComponent()
 {
 	TGuardValue<bool> InInitializeComponentGuard(bInInitializeComponent, true);
 
+	Super::InitializeComponent();
+
 	const UWorld* MyWorld = GetWorld();
 
 	if (MyWorld && MyWorld->IsGameWorld())
@@ -105,7 +107,10 @@ void UMoverComponent::InitializeComponent()
 		}
 	}
 
-	Super::InitializeComponent();
+	// Gather initial state to fulfill queries
+	CreateDefaultInputAndState(CachedLastProducedInputCmd, CachedLastSyncState, CachedLastAuxState);
+	CachedLastUsedInputCmd = CachedLastProducedInputCmd;
+	LastMoverDefaultSyncState = CachedLastSyncState.SyncStateCollection.FindDataByType<FMoverDefaultSyncState>();
 }
 
 
@@ -266,7 +271,6 @@ void UMoverComponent::ProduceInput(const int32 DeltaTimeMS, FMoverInputCmdContex
 	}
 
 	CachedLastProducedInputCmd = *Cmd;
-	bHasValidLastProducedInput = true;
 }
 
 void UMoverComponent::RestoreFrame(const FMoverSyncState* SyncState, const FMoverAuxStateContext* AuxState)
@@ -294,10 +298,10 @@ void UMoverComponent::FinalizeFrame(const FMoverSyncState* SyncState, const FMov
 	{
 		// TODO: This is distasteful duplication -- consider moving to a util function
 		CachedLastSyncState = *SyncState;
+		LastMoverDefaultSyncState = CachedLastSyncState.SyncStateCollection.FindDataByType<FMoverDefaultSyncState>();
 		CachedLastAuxState = *AuxState;
 		CachedLastSimTickTimeStep.BaseSimTimeMs = BackendLiaisonComp->GetCurrentSimTimeMs();
 		CachedLastSimTickTimeStep.ServerFrame = BackendLiaisonComp->GetCurrentSimFrame();
-		bHasValidCachedState = true;
 	}
 
 	if (OnPostFinalize.IsBound())
@@ -406,27 +410,13 @@ void UMoverComponent::InitializeSimulationState(FMoverSyncState* OutSync, FMover
 	npCheckSlow(OutSync);
 	npCheckSlow(OutAux);
 
-	// Add all initial persistent sync state types
-	for (const FMoverDataPersistence& PersistentSyncEntry : PersistentSyncStateDataTypes)
-	{
-		check(PersistentSyncEntry.RequiredType && PersistentSyncEntry.RequiredType->IsChildOf(FMoverDataStructBase::StaticStruct()));
-		OutSync->SyncStateCollection.FindOrAddDataByType(PersistentSyncEntry.RequiredType);
-	}
-
-	if (FMoverDefaultSyncState* MoverState = OutSync->SyncStateCollection.FindMutableDataByType<FMoverDefaultSyncState>())
-	{
-		MoverState->SetTransforms_WorldSpace(UpdatedComponent->GetComponentLocation(),
-		                                     UpdatedComponent->GetComponentRotation(),
-		                                     FVector::ZeroVector);	// no initial velocity
-	}
-
-	OutSync->MovementMode = StartingMovementMode;
-
-	*OutAux = FMoverAuxStateContext();
+	CreateDefaultInputAndState(CachedLastProducedInputCmd, *OutSync, *OutAux);
 
+	CachedLastUsedInputCmd = CachedLastProducedInputCmd;
 	CachedLastSyncState = *OutSync;
+	LastMoverDefaultSyncState = CachedLastSyncState.SyncStateCollection.FindDataByType<FMoverDefaultSyncState>();;
 	CachedLastAuxState = *OutAux;
-	bHasValidCachedState = true;
+
 }
 
 void UMoverComponent::SimulationTick(const FMoverTimeStep& InTimeStep, const FMoverTickStartData& SimInput, OUT FMoverTickEndData& SimOutput)
@@ -546,7 +536,6 @@ void UMoverComponent::SimulationTick(const FMoverTimeStep& InTimeStep, const FMo
 	}
 
 	CachedLastUsedInputCmd = SimInput.InputCmd;
-	bHasValidCachedUsedInput = true;
 
 	if (bSupportsKinematicBasedMovement)
 	{ 
@@ -638,28 +627,25 @@ bool UMoverComponent::IsModifierActiveOrQueued(const FMovementModifierHandle& Mo
 
 const FMovementModifierBase* UMoverComponent::FindMovementModifier(const FMovementModifierHandle& ModifierHandle) const
 {
-	if (bHasValidCachedState)
+	// Check active modifiers for modifier handle
+	for (auto ActiveModifierFromSyncStateIt = CachedLastSyncState.MovementModifiers.GetActiveModifiersIterator(); ActiveModifierFromSyncStateIt; ++ActiveModifierFromSyncStateIt)
 	{
-		// Check active modifiers for modifier handle
-		for (auto ActiveModifierFromSyncStateIt = CachedLastSyncState.MovementModifiers.GetActiveModifiersIterator(); ActiveModifierFromSyncStateIt; ++ActiveModifierFromSyncStateIt)
-		{
-			const TSharedPtr<FMovementModifierBase> ActiveModifierFromSyncState = *ActiveModifierFromSyncStateIt;
+		const TSharedPtr<FMovementModifierBase> ActiveModifierFromSyncState = *ActiveModifierFromSyncStateIt;
 
-			if (ModifierHandle == ActiveModifierFromSyncState->GetHandle())
-			{
-				return ActiveModifierFromSyncState.Get();
-			}
+		if (ModifierHandle == ActiveModifierFromSyncState->GetHandle())
+		{
+			return ActiveModifierFromSyncState.Get();
 		}
+	}
 
-		// Check queued modifiers for modifier handle
-		for (auto QueuedModifierFromSyncStateIt = CachedLastSyncState.MovementModifiers.GetQueuedModifiersIterator(); QueuedModifierFromSyncStateIt; ++QueuedModifierFromSyncStateIt)
-		{
-			const TSharedPtr<FMovementModifierBase> QueuedModifierFromSyncState = *QueuedModifierFromSyncStateIt;
+	// Check queued modifiers for modifier handle
+	for (auto QueuedModifierFromSyncStateIt = CachedLastSyncState.MovementModifiers.GetQueuedModifiersIterator(); QueuedModifierFromSyncStateIt; ++QueuedModifierFromSyncStateIt)
+	{
+		const TSharedPtr<FMovementModifierBase> QueuedModifierFromSyncState = *QueuedModifierFromSyncStateIt;
 
-			if (ModifierHandle == QueuedModifierFromSyncState->GetHandle())
-			{
-				return QueuedModifierFromSyncState.Get();
-			}
+		if (ModifierHandle == QueuedModifierFromSyncState->GetHandle())
+		{
+			return QueuedModifierFromSyncState.Get();
 		}
 	}
 	
@@ -668,28 +654,25 @@ const FMovementModifierBase* UMoverComponent::FindMovementModifier(const FMoveme
 
 const FMovementModifierBase* UMoverComponent::FindMovementModifierByType(const UScriptStruct* DataStructType) const
 {
-	if (bHasValidCachedState)
+	// Check active modifiers for modifier handle
+	for (auto ActiveModifierFromSyncStateIt = CachedLastSyncState.MovementModifiers.GetActiveModifiersIterator(); ActiveModifierFromSyncStateIt; ++ActiveModifierFromSyncStateIt)
 	{
-		// Check active modifiers for modifier handle
-		for (auto ActiveModifierFromSyncStateIt = CachedLastSyncState.MovementModifiers.GetActiveModifiersIterator(); ActiveModifierFromSyncStateIt; ++ActiveModifierFromSyncStateIt)
-		{
-			const TSharedPtr<FMovementModifierBase> ActiveModifierFromSyncState = *ActiveModifierFromSyncStateIt;
+		const TSharedPtr<FMovementModifierBase> ActiveModifierFromSyncState = *ActiveModifierFromSyncStateIt;
 
-			if (DataStructType == ActiveModifierFromSyncState->GetScriptStruct())
-			{
-				return ActiveModifierFromSyncState.Get();
-			}
+		if (DataStructType == ActiveModifierFromSyncState->GetScriptStruct())
+		{
+			return ActiveModifierFromSyncState.Get();
 		}
+	}
 
-		// Check queued modifiers for modifier handle
-		for (auto QueuedModifierFromSyncStateIt = CachedLastSyncState.MovementModifiers.GetQueuedModifiersIterator(); QueuedModifierFromSyncStateIt; ++QueuedModifierFromSyncStateIt)
-		{
-			const TSharedPtr<FMovementModifierBase> QueuedModifierFromSyncState = *QueuedModifierFromSyncStateIt;
+	// Check queued modifiers for modifier handle
+	for (auto QueuedModifierFromSyncStateIt = CachedLastSyncState.MovementModifiers.GetQueuedModifiersIterator(); QueuedModifierFromSyncStateIt; ++QueuedModifierFromSyncStateIt)
+	{
+		const TSharedPtr<FMovementModifierBase> QueuedModifierFromSyncState = *QueuedModifierFromSyncStateIt;
 
-			if (DataStructType == QueuedModifierFromSyncState->GetScriptStruct())
-			{
-				return QueuedModifierFromSyncState.Get();
-			}
+		if (DataStructType == QueuedModifierFromSyncState->GetScriptStruct())
+		{
+			return QueuedModifierFromSyncState.Get();
 		}
 	}
 	
@@ -713,36 +696,33 @@ bool UMoverComponent::HasGameplayTag(FGameplayTag TagToFind, bool bExactMatch) c
 		}
 	}
 	
-	if (bHasValidCachedState)
+	// Search Movement Modes
+	if (const UBaseMovementMode* ActiveMovementMode = GetMovementMode())
 	{
-		// Search Movement Modes
-		if (const UBaseMovementMode* ActiveMovementMode = GetMovementMode())
+		if (ActiveMovementMode->HasGameplayTag(TagToFind, bExactMatch))
 		{
-			if (ActiveMovementMode->HasGameplayTag(TagToFind, bExactMatch))
-			{
-				return true;
-			}
+			return true;
 		}
+	}
 		
-		// Search Movement Modifiers
-		for (auto ModifierFromSyncStateIt = CachedLastSyncState.MovementModifiers.GetActiveModifiersIterator(); ModifierFromSyncStateIt; ++ModifierFromSyncStateIt)
+	// Search Movement Modifiers
+	for (auto ModifierFromSyncStateIt = CachedLastSyncState.MovementModifiers.GetActiveModifiersIterator(); ModifierFromSyncStateIt; ++ModifierFromSyncStateIt)
+	{
+		if (const TSharedPtr<FMovementModifierBase> ModifierFromSyncState = *ModifierFromSyncStateIt)
 		{
-			if (const TSharedPtr<FMovementModifierBase> ModifierFromSyncState = *ModifierFromSyncStateIt)
+			if (ModifierFromSyncState.IsValid() && ModifierFromSyncState->HasGameplayTag(TagToFind, bExactMatch))
 			{
-				if (ModifierFromSyncState.IsValid() && ModifierFromSyncState->HasGameplayTag(TagToFind, bExactMatch))
-				{
-					return true;
-				}
+				return true;
 			}
 		}
+	}
 		
-		// Search Layered Moves
-		for (const TSharedPtr<FLayeredMoveBase>& LayeredMove :  CachedLastSyncState.LayeredMoves.GetActiveMoves())
+	// Search Layered Moves
+	for (const TSharedPtr<FLayeredMoveBase>& LayeredMove :  CachedLastSyncState.LayeredMoves.GetActiveMoves())
+	{
+		if (LayeredMove->HasGameplayTag(TagToFind, bExactMatch))
 		{
-			if (LayeredMove->HasGameplayTag(TagToFind, bExactMatch))
-			{
-				return true;
-			}
+			return true;
 		}
 	}
 	
@@ -775,10 +755,10 @@ void UMoverComponent::SetFrameStateFromContext(const FMoverSyncState* SyncState,
 	// TODO: This is distasteful duplication -- consider moving to a util function
 	// Copy these as the last official state, so they can be queried by other systems outside of the NP simulation
 	CachedLastSyncState = *SyncState;
+	LastMoverDefaultSyncState = CachedLastSyncState.SyncStateCollection.FindDataByType<FMoverDefaultSyncState>();
 	CachedLastAuxState = *AuxState;
 	CachedLastSimTickTimeStep.BaseSimTimeMs = BackendLiaisonComp->GetCurrentSimTimeMs();
 	CachedLastSimTickTimeStep.ServerFrame = BackendLiaisonComp->GetCurrentSimFrame();
-	bHasValidCachedState = true;
 
 	if (FMoverDefaultSyncState* MoverState = CachedLastSyncState.SyncStateCollection.FindMutableDataByType<FMoverDefaultSyncState>())
 	{
@@ -849,6 +829,35 @@ bool UMoverComponent::InitMoverSimulation()
 	return bHasMatchingStartingState;
 }
 
+void UMoverComponent::CreateDefaultInputAndState(FMoverInputCmdContext& OutInputCmd, FMoverSyncState& OutSyncState, FMoverAuxStateContext& OutAuxState) const
+{
+	OutInputCmd = FMoverInputCmdContext();
+	// TODO: here is where we'd add persistent input cmd struct types once they're supported
+
+	OutSyncState = FMoverSyncState();
+
+	// Add all initial persistent sync state types
+	for (const FMoverDataPersistence& PersistentSyncEntry : PersistentSyncStateDataTypes)
+	{
+		OutSyncState.SyncStateCollection.FindOrAddDataByType(PersistentSyncEntry.RequiredType);
+	}
+
+	// Mirror the scene component transform if we have one, otherwise it will be left at origin
+	FMoverDefaultSyncState* MoverState = OutSyncState.SyncStateCollection.FindMutableDataByType<FMoverDefaultSyncState>();
+	if (MoverState && UpdatedComponent)
+	{
+		MoverState->SetTransforms_WorldSpace(
+			UpdatedComponent->GetComponentLocation(),
+			UpdatedComponent->GetComponentRotation(),
+			FVector::ZeroVector);	// no initial velocity
+	}
+
+	OutSyncState.MovementMode = StartingMovementMode;
+	
+	OutAuxState = FMoverAuxStateContext();
+
+}
+
 void UMoverComponent::HandleImpact(FMoverOnImpactParams& ImpactParams)
 {
 	if (ImpactParams.MovementModeName.IsNone())
@@ -1537,12 +1546,7 @@ DEFINE_FUNCTION(UMoverComponent::execK2_FindActiveLayeredMove)
 
 const FLayeredMoveBase* UMoverComponent::FindActiveLayeredMoveByType(const UScriptStruct* LayeredMoveStructType) const
 {
-	if (bHasValidCachedState)
-	{
-		return CachedLastSyncState.LayeredMoves.FindActiveMove(LayeredMoveStructType);
-	}
-
-	return nullptr;
+	return CachedLastSyncState.LayeredMoves.FindActiveMove(LayeredMoveStructType);
 }
 
 void UMoverComponent::QueueNextMode(FName DesiredModeName, bool bShouldReenter)
@@ -1766,12 +1770,9 @@ void UMoverComponent::SetPrimaryVisualComponent(USceneComponent* SceneComponent)
 
 FVector UMoverComponent::GetVelocity() const
 { 
-	if (bHasValidCachedState)
+	if (LastMoverDefaultSyncState)
 	{
-		if (const FMoverDefaultSyncState* SyncState = CachedLastSyncState.SyncStateCollection.FindDataByType<FMoverDefaultSyncState>())
-		{
-			return SyncState->GetVelocity_WorldSpace();
-		}
+		return LastMoverDefaultSyncState->GetVelocity_WorldSpace();
 	}
 
 	return FVector::ZeroVector;
@@ -1780,12 +1781,9 @@ FVector UMoverComponent::GetVelocity() const
 
 FVector UMoverComponent::GetMovementIntent() const
 { 
-	if (bHasValidCachedState)
+	if (LastMoverDefaultSyncState)
 	{
-		if (const FMoverDefaultSyncState* SyncState = CachedLastSyncState.SyncStateCollection.FindDataByType<FMoverDefaultSyncState>())
-		{
-			return SyncState->GetIntent_WorldSpace();
-		}
+		return LastMoverDefaultSyncState->GetIntent_WorldSpace();
 	}
 
 	return FVector::ZeroVector; 
@@ -1795,26 +1793,20 @@ FVector UMoverComponent::GetMovementIntent() const
 FRotator UMoverComponent::GetTargetOrientation() const
 {
 	// Prefer the input's intended orientation, but if it can't be determined, assume it matches the actual orientation
-	if (bHasValidCachedUsedInput)
+	const FMoverInputCmdContext& LastInputCmd = GetLastInputCmd();
+	if (const FCharacterDefaultInputs* MoverInputs = LastInputCmd.InputCollection.FindDataByType<FCharacterDefaultInputs>())
 	{
-		const FMoverInputCmdContext& LastInputCmd = GetLastInputCmd();
-		if (const FCharacterDefaultInputs* MoverInputs = LastInputCmd.InputCollection.FindDataByType<FCharacterDefaultInputs>())
-		{
-			const FVector TargetOrientationDir = MoverInputs->GetOrientationIntentDir_WorldSpace();
+		const FVector TargetOrientationDir = MoverInputs->GetOrientationIntentDir_WorldSpace();
 
-			if (!TargetOrientationDir.IsNearlyZero())
-			{
-				return TargetOrientationDir.ToOrientationRotator();
-			}
+		if (!TargetOrientationDir.IsNearlyZero())
+		{
+			return TargetOrientationDir.ToOrientationRotator();
 		}
 	}
 	
-	if (bHasValidCachedState)
+	if (LastMoverDefaultSyncState)
 	{
-		if (const FMoverDefaultSyncState* SyncState = CachedLastSyncState.SyncStateCollection.FindDataByType<FMoverDefaultSyncState>())
-		{
-			return SyncState->GetOrientation_WorldSpace();
-		}
+		return LastMoverDefaultSyncState->GetOrientation_WorldSpace();
 	}
 
 	return GetOwner() ? GetOwner()->GetActorRotation() : FRotator::ZeroRotator;
@@ -2013,12 +2005,7 @@ TArray<FTrajectorySampleInfo> UMoverComponent::GetPredictedTrajectory(FMoverPred
 
 FName UMoverComponent::GetMovementModeName() const
 { 
-	if (bHasValidCachedState)
-	{
-		return CachedLastSyncState.MovementMode;
-	}
-
-	return NAME_None;
+	return CachedLastSyncState.MovementMode;
 }
 
 const UBaseMovementMode* UMoverComponent::GetMovementMode() const
@@ -2028,12 +2015,9 @@ const UBaseMovementMode* UMoverComponent::GetMovementMode() const
 
 UPrimitiveComponent* UMoverComponent::GetMovementBase() const
 {
-	if (bHasValidCachedState)
+	if (LastMoverDefaultSyncState)
 	{
-		if (const FMoverDefaultSyncState* SyncState = CachedLastSyncState.SyncStateCollection.FindDataByType<FMoverDefaultSyncState>())
-		{
-			return SyncState->GetMovementBase();
-		}
+		return LastMoverDefaultSyncState->GetMovementBase();
 	}
 
 	return nullptr;
@@ -2041,30 +2025,21 @@ UPrimitiveComponent* UMoverComponent::GetMovementBase() const
 
 FName UMoverComponent::GetMovementBaseBoneName() const
 {
-	if (bHasValidCachedState)
+	if (LastMoverDefaultSyncState)
 	{
-		if (const FMoverDefaultSyncState* SyncState = CachedLastSyncState.SyncStateCollection.FindDataByType<FMoverDefaultSyncState>())
-		{
-			return SyncState->GetMovementBaseBoneName();
-		}
+		return LastMoverDefaultSyncState->GetMovementBaseBoneName();
 	}
 
 	return NAME_None;
 }
 
-
 bool UMoverComponent::HasValidCachedState() const
 {
-	return bHasValidCachedState;
+	return true;
 }
 
 const FMoverSyncState& UMoverComponent::GetSyncState() const
 {
-	if (!bHasValidCachedState)
-	{
-		UE_LOG(LogMover, Warning, TEXT("Attempting direct access to the last-cached sync state before one has been set. Results will be unreliable. Use the HasValidCachedState function to check if CachedLastSyncState is valid or not. (%s)"), *GetNameSafe(GetOwner()));
-	}
-
 	return CachedLastSyncState;
 }
 
@@ -2091,16 +2066,11 @@ UMoverBlackboard* UMoverComponent::GetSimBlackboard_Mutable() const
 
 bool UMoverComponent::HasValidCachedInputCmd() const
 {
-	return bHasValidCachedUsedInput;
+	return true;
 }
 
 const FMoverInputCmdContext& UMoverComponent::GetLastInputCmd() const
 {
-	if (!bHasValidCachedUsedInput)
-	{
-		UE_LOG(LogMover, Warning, TEXT("Attempting direct access to the last-cached used input cmd before one has been set. Results will be unreliable. Use the HasValidCachedInputCmd function to check if CachedLastUsedInputCmd is valid or not. (%s)"), *GetNameSafe(GetOwner()));
-	}
-
 	return CachedLastUsedInputCmd;
 }
 
@@ -2178,15 +2148,12 @@ UBaseMovementMode* UMoverComponent::FindMode_Mutable(TSubclassOf<UBaseMovementMo
 
 UBaseMovementMode* UMoverComponent::GetActiveModeInternal(TSubclassOf<UBaseMovementMode> ModeType, bool bRequireExactClass) const
 {
-	if (bHasValidCachedState)
+	if (const TObjectPtr<UBaseMovementMode>* CurrentMode = MovementModes.Find(CachedLastSyncState.MovementMode))
 	{
-		if (const TObjectPtr<UBaseMovementMode>* CurrentMode = MovementModes.Find(CachedLastSyncState.MovementMode))
+		if ((!bRequireExactClass && CurrentMode->IsA(ModeType)) ||
+			CurrentMode->GetClass() == ModeType)
 		{
-			if ((!bRequireExactClass && CurrentMode->IsA(ModeType)) ||
-				CurrentMode->GetClass() == ModeType)
-			{
-				return CurrentMode->Get();
-			}
+			return CurrentMode->Get();
 		}
 	}
 
@@ -2198,10 +2165,9 @@ void UMoverComponent::SetSimulationOutput(const FMoverTimeStep& TimeStep, const
 	CachedLastSimTickTimeStep = TimeStep;
 
 	CachedLastUsedInputCmd = OutputData.LastUsedInputCmd;
-	bHasValidCachedUsedInput = true;
 
 	CachedLastSyncState = OutputData.SyncState;
-	bHasValidCachedState = true;
+	LastMoverDefaultSyncState = CachedLastSyncState.SyncStateCollection.FindDataByType<FMoverDefaultSyncState>();
 
 	// This is for things like the ground info that we want to cache and interpolate but isn't part of the networked sync state
 	// AdditionalOutputData is generic because ground info might not be useful for platforms, say, but platforms may want to store something else.
diff --git a/Engine/Plugins/Experimental/Mover/Source/Mover/Public/MoverComponent.h b/Engine/Plugins/Experimental/Mover/Source/Mover/Public/MoverComponent.h
--- a/Engine/Plugins/Experimental/Mover/Source/Mover/Public/MoverComponent.h
+++ b/Engine/Plugins/Experimental/Mover/Source/Mover/Public/MoverComponent.h
@@ -347,18 +347,20 @@ public:	// Queries
 	FName GetMovementBaseBoneName() const;
 
 	// Signals whether we have a sync state saved yet. If not, most queries will not be meaningful.
-	UFUNCTION(BlueprintPure, Category = Mover)
+	UE_DEPRECATED(5.6, "HasValidCachedState has been deprecated, and is not needed since we no longer wait until movement simulation begins before providing a valid sync state.")
+	UFUNCTION(BlueprintPure, Category = Mover, meta=(DeprecatedFunction, DeprecationMessage="HasValidCachedState has been deprecated, and is not needed since we no longer wait until movement simulation begins before providing a valid sync state."))
 	bool HasValidCachedState() const;
 
-	// Access the most recent captured sync state. Check @HasValidCachedState first.
+	// Access the most recent captured sync state.
 	UFUNCTION(BlueprintPure, Category = Mover)
 	const FMoverSyncState& GetSyncState() const;
 
 	// Signals whether we have input data saved yet. If not, input queries will not be meaningful.
-	UFUNCTION(BlueprintPure, Category = Mover)
+	UE_DEPRECATED(5.6, "HasValidCachedInputCmd has been deprecated, and is not needed since we no longer wait until movement simulation begins before providing a valid input cmd.")
+	UFUNCTION(BlueprintPure, Category = Mover, meta = (DeprecatedFunction, DeprecationMessage = "HasValidCachedInputCmd has been deprecated, and is not needed since we no longer wait until movement simulation begins before providing a valid input cmd."))
 	bool HasValidCachedInputCmd() const;
 
-	// Access the most recently-used inputs. Check @HasValidCachedInputCmd first.
+	// Access the most recently-used inputs.
 	UFUNCTION(BlueprintPure, Category = Mover)
 	const FMoverInputCmdContext& GetLastInputCmd() const;
 
@@ -507,7 +509,9 @@ protected:
 	void SetFrameStateFromContext(const FMoverSyncState* SyncState, const FMoverAuxStateContext* AuxState, bool bRebaseBasedState);
 
 public:
-	bool InitMoverSimulation();
+	bool InitMoverSimulation();	// Intended to be called by the backend only
+
+	virtual void CreateDefaultInputAndState(FMoverInputCmdContext& OutInputCmd, FMoverSyncState& OutSyncState, FMoverAuxStateContext& OutAuxState) const;
 
 	/** Handle a blocking impact.*/
 	UFUNCTION(BlueprintCallable, Category = Mover)
@@ -582,6 +586,7 @@ protected:
 	FMoverInputCmdContext CachedLastUsedInputCmd;
 	
 	FMoverSyncState CachedLastSyncState;	
+	const FMoverDefaultSyncState* LastMoverDefaultSyncState = nullptr;
 
 	FMoverTimeStep CachedLastSimTickTimeStep;	// Saved timestep info from our last simulation tick, used during rollback handling. This will rewind during corrections.
 	FMoverTimeStep CachedNewestSimTickTimeStep;	// Saved timestep info from the newest (farthest-advanced) simulation tick. This will not rewind during corrections.
@@ -622,11 +627,6 @@ public:
 
 	void SetSimulationOutput(const FMoverTimeStep& TimeStep, const UE::Mover::FSimulationOutputData& OutputData);
 
-protected:
-	bool bHasValidLastProducedInput = false;
-	bool bHasValidCachedUsedInput = false;
-	bool bHasValidCachedState = false;
-
 private:
 	/** Whether or not gravity is overridden on this actor. Otherwise, fall back on world settings. See @SetGravityOverride */
 	UPROPERTY(EditDefaultsOnly, Category="Mover|Gravity")
