[Mover] [Swimming] OriginalHalfHeight calculation in OnRegister cannot be called on the CDO as the collider component might not be registered yet, especially if not in editor.

diff --git a/Engine/Plugins/Experimental/Mover/Source/Mover/Private/DefaultMovementSet/Modes/SwimmingMode.cpp b/Engine/Plugins/Experimental/Mover/Source/Mover/Private/DefaultMovementSet/Modes/SwimmingMode.cpp
--- a/Engine/Plugins/Experimental/Mover/Source/Mover/Private/DefaultMovementSet/Modes/SwimmingMode.cpp
+++ b/Engine/Plugins/Experimental/Mover/Source/Mover/Private/DefaultMovementSet/Modes/SwimmingMode.cpp
@@ -174,7 +174,7 @@ void USwimmingMode::OnRegistered(const FName ModeName)
 	if (UMoverComponent* MoverComp = GetMoverComponent())
 	{
 		CommonLegacySettings = MoverComp->FindSharedSettings<UCommonLegacyMovementSettings>();
-		OriginalCapsuleHalfHeight = MoverComp->GetOwner()->GetClass()->GetDefaultObject<AActor>()->GetSimpleCollisionHalfHeight();
+		OriginalCapsuleHalfHeight = MoverComp->GetOwner()->GetSimpleCollisionHalfHeight();
 	}
 	ensureMsgf(CommonLegacySettings, TEXT("Failed to find instance of CommonLegacyMovementSettings on %s. Movement may not function properly."), *GetPathNameSafe(this));
 }
