[Mover] Character movement - CalcVelocity()

diff --git a/Engine/Plugins/Experimental/Mover/Source/Mover/Private/MoveLibrary/MovementUtils.cpp b/Engine/Plugins/Experimental/Mover/Source/Mover/Private/MoveLibrary/MovementUtils.cpp
--- a/Engine/Plugins/Experimental/Mover/Source/Mover/Private/MoveLibrary/MovementUtils.cpp
+++ b/Engine/Plugins/Experimental/Mover/Source/Mover/Private/MoveLibrary/MovementUtils.cpp
@@ -66,18 +66,20 @@ FVector UMovementUtils::ComputeVelocity(const FComputeVelocityParams& InParams)
 		const FVector ControlAcceleration = InParams.MoveDirectionIntent.GetClampedToMaxSize(1.f);
 		const float AnalogInputModifier = (ControlAcceleration.SizeSquared() > 0.f ? ControlAcceleration.Size() : 0.f);
 		DesiredSpeed = InParams.MaxSpeed * AnalogInputModifier;
-		const bool bExceedingMaxSpeed = IsExceedingMaxSpeed(Velocity, DesiredSpeed);
+		const FVector VelocityAlongInput = Velocity.ProjectOnTo(InParams.MoveDirectionIntent);
+		const bool bExceedingMaxSpeedAlongInput = IsExceedingMaxSpeed(VelocityAlongInput, DesiredSpeed);
 
 		if (Velocity.SizeSquared() > 0.f)
 		{
-			if (AnalogInputModifier > 0.f && !bExceedingMaxSpeed)
+			if (AnalogInputModifier > 0.f && !bExceedingMaxSpeedAlongInput)
 			{
 				// Apply change in velocity direction
 				// Change direction faster than only using acceleration, but never increase velocity magnitude.
 				const float TimeScale = FMath::Clamp(InParams.DeltaSeconds * InParams.TurningBoost, 0.f, 1.f);
 				Velocity = Velocity + (ControlAcceleration * Velocity.Size() - Velocity) * FMath::Min(TimeScale * InParams.Friction, 1.f);
 			}
-			else
+			const bool bExceedingMaxSpeed = IsExceedingMaxSpeed(Velocity, DesiredSpeed);
+			if (bExceedingMaxSpeed)
 			{
 				// Dampen velocity magnitude based on deceleration.
 				const FVector OldVelocity = Velocity;
