Chaos Visual Debugger: Adding local sim data to what CVD records for Mover

diff --git a/Engine/Plugins/Experimental/Mover/Source/Mover/Private/ChaosVisualDebugger/MoverCVDRuntimeTrace.cpp b/Engine/Plugins/Experimental/Mover/Source/Mover/Private/ChaosVisualDebugger/MoverCVDRuntimeTrace.cpp
--- a/Engine/Plugins/Experimental/Mover/Source/Mover/Private/ChaosVisualDebugger/MoverCVDRuntimeTrace.cpp
+++ b/Engine/Plugins/Experimental/Mover/Source/Mover/Private/ChaosVisualDebugger/MoverCVDRuntimeTrace.cpp
@@ -19,7 +19,8 @@
 #include "Interfaces/IPhysicsComponent.h"
 #include "Engine/World.h"
 
-CVD_DEFINE_OPTIONAL_DATA_CHANNEL(MoverData, EChaosVDDataChannelInitializationFlags::CanChangeEnabledState)
+CVD_DEFINE_OPTIONAL_DATA_CHANNEL(MoverNetworkedData, EChaosVDDataChannelInitializationFlags::CanChangeEnabledState)
+CVD_DEFINE_OPTIONAL_DATA_CHANNEL(MoverLocalSimData, EChaosVDDataChannelInitializationFlags::CanChangeEnabledState)
 
 // FSkipObjectRefsMemoryWriter and FSkipObjectRefsMemoryReader are a 
 // workaround for serializing mover info structs with object references in them, such as the mover base
@@ -52,7 +53,7 @@ public:
 	}
 };
 
-void FMoverCVDRuntimeTrace::UnwrapSimData(const FMoverCVDSimDataWrapper& InSimDataWrapper, TSharedPtr<FMoverInputCmdContext>& OutInputCmd, TSharedPtr<FMoverSyncState>& OutSyncState)
+void FMoverCVDRuntimeTrace::UnwrapSimData(const FMoverCVDSimDataWrapper& InSimDataWrapper, TSharedPtr<FMoverInputCmdContext>& OutInputCmd, TSharedPtr<FMoverSyncState>& OutSyncState, TSharedPtr<FMoverDataCollection>& OutLocalSimData)
 {
 	// Consider using Chaos::VisualDebugger::FChaosVDMemoryReader
 
@@ -60,25 +61,26 @@ void FMoverCVDRuntimeTrace::UnwrapSimData(const FMoverCVDSimDataWrapper& InSimDa
 	{
 		// Deserialize into a new struct allocated dynamically
 		FMemoryReader ArReader(InSimDataWrapper.InputCmdBytes, true);
-		FMoverInputCmdContext* NewStruct = new FMoverInputCmdContext;
-		FMoverInputCmdContext::StaticStruct()->SerializeBin(ArReader, NewStruct);
+		OutInputCmd = MakeShared<FMoverInputCmdContext>();
+		FMoverInputCmdContext::StaticStruct()->SerializeBin(ArReader, OutInputCmd.Get());
 		// Input cmd's Collection of Custom Input Structs
 		FSkipObjectRefsMemoryReader ArInputCollectionReader(InSimDataWrapper.InputMoverDataCollectionBytes, true);
-		NewStruct->InputCollection.SerializeDebugData(ArInputCollectionReader);
-		OutInputCmd = TSharedPtr<FMoverInputCmdContext>(NewStruct);
+		OutInputCmd->InputCollection.SerializeDebugData(ArInputCollectionReader);
 	}
 
 	// Sync State
 	{
 		// Deserialize into a new struct allocated dynamically
 		FMemoryReader ArReader(InSimDataWrapper.SyncStateBytes, true);
-		FMoverSyncState* NewStruct = new FMoverSyncState;
-		FMoverSyncState::StaticStruct()->SerializeBin(ArReader, NewStruct);
-		OutSyncState = TSharedPtr<FMoverSyncState>(NewStruct);
+		OutSyncState = MakeShared<FMoverSyncState>();
+		FMoverSyncState::StaticStruct()->SerializeBin(ArReader, OutSyncState.Get());
+
+		OutLocalSimData = MakeShared<FMoverDataCollection>();
+		OutLocalSimData->SerializeDebugData(ArReader);
 	}
 }
 
-void FMoverCVDRuntimeTrace::WrapSimData(uint32 SolverID, uint32 ParticleID, const FMoverInputCmdContext& InInputCmd, const FMoverSyncState& InSyncState, FMoverCVDSimDataWrapper& OutSimDataWrapper)
+void FMoverCVDRuntimeTrace::WrapSimData(uint32 SolverID, uint32 ParticleID, const FMoverInputCmdContext& InInputCmd, const FMoverSyncState& InSyncState, const FMoverDataCollection* LocalSimData, FMoverCVDSimDataWrapper& OutSimDataWrapper)
 {
 	OutSimDataWrapper.SolverID = SolverID;
 	OutSimDataWrapper.ParticleID = ParticleID;
@@ -103,21 +105,32 @@ void FMoverCVDRuntimeTrace::WrapSimData(uint32 SolverID, uint32 ParticleID, cons
 		FMemoryWriter ArWriter(OutSimDataWrapper.SyncStateBytes, true);
 		// This is not version friendly, we need to instead use SerializeTagParticles. Slower and Sergio is working on a faster version, but it's not available yet.
 		InSyncState.StaticStruct()->SerializeBin(ArWriter, &const_cast<FMoverSyncState&>(InSyncState));
+
+		// Local Sim Data Collection
+		if (LocalSimData)
+		{
+			(const_cast<FMoverDataCollection*>(LocalSimData))->SerializeDebugData(ArWriter);
+		}
+		else
+		{
+			FMoverDataCollection EmptyDataCollection;
+			EmptyDataCollection.SerializeDebugData(ArWriter);
+		}
 	}
 }
 
-void FMoverCVDRuntimeTrace::TraceMoverData(UMoverComponent* MoverComponent, const FMoverInputCmdContext* InputCmd, const FMoverSyncState* SyncState)
+void FMoverCVDRuntimeTrace::TraceMoverData(UMoverComponent* MoverComponent, const FMoverInputCmdContext* InputCmd, const FMoverSyncState* SyncState, const FMoverDataCollection* LocalSimData)
 {
 	if (!FChaosVisualDebuggerTrace::IsTracing())
 	{
 		return;
 	}
 
-	if (!CVDDC_MoverData->IsChannelEnabled())
+	if (!CVDDC_MoverNetworkedData->IsChannelEnabled())
 	{
 		return;
 	}
-	
+
 	if (!InputCmd || !SyncState)
 	{
 		return;
@@ -135,14 +148,36 @@ void FMoverCVDRuntimeTrace::TraceMoverData(UMoverComponent* MoverComponent, cons
 		}
 
 		int32 SolverID = CVD_TRACE_GET_SOLVER_ID_FROM_WORLD(World);
-		TraceMoverDataPrivate(SolverID, ParticleID, InputCmd, SyncState);
+		// LocalSimState could add a lot of extra bytes, especially without some sort of delta serialization, so it is only optionally recorded
+		const FMoverDataCollection* RecordedLocalSimData = CVDDC_MoverLocalSimData->IsChannelEnabled() ? LocalSimData : nullptr;
+		TraceMoverDataPrivate(SolverID, ParticleID, InputCmd, SyncState, RecordedLocalSimData);
+	}
+}
+
+void FMoverCVDRuntimeTrace::TraceMoverData(uint32 SolverID, uint32 ParticleID, const FMoverInputCmdContext* InputCmd, const FMoverSyncState* SyncState, const FMoverDataCollection* LocalSimData)
+{
+	if (!FChaosVisualDebuggerTrace::IsTracing())
+	{
+		return;
+	}
+
+	if (!CVDDC_MoverNetworkedData->IsChannelEnabled())
+	{
+		return;
+	}
+
+	if (!InputCmd || !SyncState)
+	{
+		return;
 	}
+		
+	TraceMoverDataPrivate(SolverID, ParticleID, InputCmd, SyncState, LocalSimData);
 }
 
-void FMoverCVDRuntimeTrace::TraceMoverDataPrivate(uint32 SolverID, uint32 ParticleID, const FMoverInputCmdContext* InputCmd, const FMoverSyncState* SyncState)
+void FMoverCVDRuntimeTrace::TraceMoverDataPrivate(uint32 SolverID, uint32 ParticleID, const FMoverInputCmdContext* InputCmd, const FMoverSyncState* SyncState, const FMoverDataCollection* LocalSimData)
 {
 	FMoverCVDSimDataWrapper SimDataWrapper;
-	WrapSimData(SolverID, ParticleID, *InputCmd, *SyncState, SimDataWrapper);
+	WrapSimData(SolverID, ParticleID, *InputCmd, *SyncState, LocalSimData, SimDataWrapper);
 	SimDataWrapper.MarkAsValid();
 
 	FChaosVDScopedTLSBufferAccessor TLSDataBuffer;
diff --git a/Engine/Plugins/Experimental/Mover/Source/Mover/Public/ChaosVisualDebugger/MoverCVDRuntimeTrace.h b/Engine/Plugins/Experimental/Mover/Source/Mover/Public/ChaosVisualDebugger/MoverCVDRuntimeTrace.h
--- a/Engine/Plugins/Experimental/Mover/Source/Mover/Public/ChaosVisualDebugger/MoverCVDRuntimeTrace.h
+++ b/Engine/Plugins/Experimental/Mover/Source/Mover/Public/ChaosVisualDebugger/MoverCVDRuntimeTrace.h
@@ -11,6 +11,7 @@
 struct FMoverCVDSimDataWrapper;
 struct FMoverInputCmdContext;
 struct FMoverSyncState;
+struct FMoverDataCollection;
 class UMoverComponent;
 
 #if WITH_CHAOS_VISUAL_DEBUGGER
@@ -19,16 +20,20 @@ class UMoverComponent;
 class FMoverCVDRuntimeTrace
 {
 public:
-	MOVER_API static void TraceMoverData(UMoverComponent* MoverComponent, const FMoverInputCmdContext* InputCmd, const FMoverSyncState* SyncState);
+	MOVER_API static void TraceMoverData(UMoverComponent* MoverComponent, const FMoverInputCmdContext* InputCmd, const FMoverSyncState* SyncState, const FMoverDataCollection* LocalSimState = nullptr);
+	MOVER_API static void TraceMoverData(uint32 SolverID, uint32 ParticleID, const FMoverInputCmdContext* InputCmd, const FMoverSyncState* SyncState, const FMoverDataCollection* LocalSimState = nullptr);
 
-	MOVER_API static void UnwrapSimData(const FMoverCVDSimDataWrapper& InSimDataWrapper, TSharedPtr<FMoverInputCmdContext>& OutInputCmd, TSharedPtr<FMoverSyncState>& OutSyncState);
-	static void WrapSimData(uint32 SolverID, uint32 ParticleID, const FMoverInputCmdContext& InInputCmd, const FMoverSyncState& InSyncState, FMoverCVDSimDataWrapper& OutSimDataWrapper);
+	MOVER_API static void UnwrapSimData(const FMoverCVDSimDataWrapper& InSimDataWrapper, TSharedPtr<FMoverInputCmdContext>& OutInputCmd, TSharedPtr<FMoverSyncState>& OutSyncState, TSharedPtr<FMoverDataCollection>& OutLocalSimState);
+	MOVER_API static void WrapSimData(uint32 SolverID, uint32 ParticleID, const FMoverInputCmdContext& InInputCmd, const FMoverSyncState& InSyncState, const FMoverDataCollection* LocalSimState, FMoverCVDSimDataWrapper& OutSimDataWrapper);
 
 private:
-	static void TraceMoverDataPrivate(uint32 SolverID, uint32 ParticleID, const FMoverInputCmdContext* InputCmd, const FMoverSyncState* SyncState);
+	static void TraceMoverDataPrivate(uint32 SolverID, uint32 ParticleID, const FMoverInputCmdContext* InputCmd, const FMoverSyncState* SyncState, const FMoverDataCollection* LocalSimState = nullptr);
 };
 
-CVD_DECLARE_OPTIONAL_DATA_CHANNEL_EXTERN(MoverData, MOVER_API)
+// This is all mover data that is networked, either input command (client to server) or sync state (server to client)
+CVD_DECLARE_OPTIONAL_DATA_CHANNEL_EXTERN(MoverNetworkedData, MOVER_API)
+// This is additional mover data, local to each end point's simulation
+CVD_DECLARE_OPTIONAL_DATA_CHANNEL_EXTERN(MoverLocalSimData, MOVER_API)
 
 #else
 
@@ -36,9 +41,10 @@ CVD_DECLARE_OPTIONAL_DATA_CHANNEL_EXTERN(MoverData, MOVER_API)
 class FMoverCVDRuntimeTrace
 {
 public:
-	MOVER_API static void TraceMoverData(UMoverComponent* MoverComponent, const FMoverInputCmdContext* InputCmd, const FMoverSyncState* SyncState) {}
-	MOVER_API static void UnwrapSimData(const FMoverCVDSimDataWrapper& InSimDataWrapper, TSharedPtr<FMoverInputCmdContext>& OutInputCmd, TSharedPtr<FMoverSyncState>& OutSyncState) {}
-	static void WrapSimData(uint32 SolverID, uint32 ParticleID, const FMoverInputCmdContext& InInputCmd, const FMoverSyncState& InSyncState, FMoverCVDSimDataWrapper& OutSimDataWrapper) {}
+	MOVER_API static void TraceMoverData(UMoverComponent* MoverComponent, const FMoverInputCmdContext* InputCmd, const FMoverSyncState* SyncState, const FMoverDataCollection* FMoverLocalSimState = nullptr) {}
+	MOVER_API static void TraceMoverData(uint32 SolverID, uint32 ParticleID, const FMoverInputCmdContext* InputCmd, const FMoverSyncState* SyncState, const FMoverDataCollection* FMoverLocalSimState = nullptr) {}
+	MOVER_API static void UnwrapSimData(const FMoverCVDSimDataWrapper& InSimDataWrapper, TSharedPtr<FMoverInputCmdContext>& OutInputCmd, TSharedPtr<FMoverSyncState>& OutSyncState, TSharedPtr<FMoverDataCollection>& FMoverLocalSimState) {}
+	MOVER_API static void WrapSimData(uint32 SolverID, uint32 ParticleID, const FMoverInputCmdContext& InInputCmd, const FMoverSyncState& InSyncState, const FMoverDataCollection* FMoverLocalSimState, FMoverCVDSimDataWrapper& OutSimDataWrapper) {}
 };
 
 #endif // WITH_CHAOS_VISUAL_DEBUGGER
diff --git a/Engine/Plugins/Experimental/Mover/Source/MoverCVDData/Public/MoverCVDDataWrappers.h b/Engine/Plugins/Experimental/Mover/Source/MoverCVDData/Public/MoverCVDDataWrappers.h
--- a/Engine/Plugins/Experimental/Mover/Source/MoverCVDData/Public/MoverCVDDataWrappers.h
+++ b/Engine/Plugins/Experimental/Mover/Source/MoverCVDData/Public/MoverCVDDataWrappers.h
@@ -8,18 +8,6 @@
 
 #include "MoverCVDDataWrappers.generated.h"
 
-USTRUCT(DisplayName="Mover Particle Info")
-struct FMoverCVDSimDataParticleInfo
-{
-	GENERATED_BODY()
-
-	UPROPERTY(VisibleAnywhere, Category="Mover Info")
-	int32 SolverID = INDEX_NONE;
-
-	UPROPERTY(VisibleAnywhere, Category="Mover Info")
-	int32 ParticleID = INDEX_NONE;
-};
-
 USTRUCT(DisplayName="Mover Sim Data")
 struct FMoverCVDSimDataWrapper : public FChaosVDWrapperDataBase
 {
@@ -27,8 +15,12 @@ struct FMoverCVDSimDataWrapper : public FChaosVDWrapperDataBase
 	
 	inline static FStringView WrapperTypeName = TEXT("FMoverCVDSimDataWrapper");
 
+	UPROPERTY(VisibleAnywhere, Category="Mover Info")
 	int32 SolverID = INDEX_NONE;
+
+	UPROPERTY(VisibleAnywhere, Category="Mover Info")
 	int32 ParticleID = INDEX_NONE;
+
 	TArray<uint8> SyncStateBytes;
 	TArray<uint8> InputCmdBytes;
 	TArray<uint8> InputMoverDataCollectionBytes;
diff --git a/Engine/Plugins/Experimental/Mover/Source/MoverCVDEditor/Private/MoverCVDSimDataComponent.cpp b/Engine/Plugins/Experimental/Mover/Source/MoverCVDEditor/Private/MoverCVDSimDataComponent.cpp
--- a/Engine/Plugins/Experimental/Mover/Source/MoverCVDEditor/Private/MoverCVDSimDataComponent.cpp
+++ b/Engine/Plugins/Experimental/Mover/Source/MoverCVDEditor/Private/MoverCVDSimDataComponent.cpp
@@ -23,8 +23,7 @@ void UMoverCVDSimDataComponent::UpdateFromSolverFrameData(const FChaosVDSolverFr
 			FrameSimDataArray = *RecordedData;
 
 			// Also, clear all cached deserialized data, we're starting from scratch at a new frame
-			DeserializedSyncStates.Empty();
-			DeserializedInputCmds.Empty();
+			DeserializedStates.Empty();
 		}
 	}
 }
@@ -34,7 +33,7 @@ void UMoverCVDSimDataComponent::ClearData()
 	FrameSimDataArray.Reset();
 }
 
-bool UMoverCVDSimDataComponent::FindAndUnwrapSimDataForParticle(uint32 ParticleID, TSharedPtr<FMoverCVDSimDataWrapper>& OutSimDataWrapper, TSharedPtr<FMoverSyncState>& OutSyncState, TSharedPtr<FMoverInputCmdContext>& OutInputCmd)
+bool UMoverCVDSimDataComponent::FindAndUnwrapSimDataForParticle(uint32 ParticleID, TSharedPtr<FMoverCVDSimDataWrapper>& OutSimDataWrapper, TSharedPtr<FMoverSyncState>& OutSyncState, TSharedPtr<FMoverInputCmdContext>& OutInputCmd, TSharedPtr<FMoverDataCollection>& OutLocalSimData)
 {
 	// Look for a sim data corresponding to ParticleID
 	TSharedPtr<FMoverCVDSimDataWrapper>* FoundSimData = FrameSimDataArray.FindByPredicate
@@ -55,13 +54,19 @@ bool UMoverCVDSimDataComponent::FindAndUnwrapSimDataForParticle(uint32 ParticleI
 	OutSimDataWrapper = *FoundSimData;
 
 	// Did we previously deserialize?
-	OutSyncState = DeserializedSyncStates.FindOrAdd(SimDataPtr);
-	OutInputCmd = DeserializedInputCmds.FindOrAdd(SimDataPtr);
-	if (!OutSyncState || !OutInputCmd)
+	TSharedPtr<FDeserializedMoverStates> DeserializedMoverStates = DeserializedStates.FindOrAdd(SimDataPtr);
+	if (!DeserializedMoverStates)
 	{
+		DeserializedMoverStates = MakeShared<FDeserializedMoverStates>();
 		// This means that the SimData wasn't deserialized yet, so we do it now
 		// Otherwise we use the cached version
-		FMoverCVDRuntimeTrace::UnwrapSimData(*SimDataPtr, OutInputCmd, OutSyncState);
+		FMoverCVDRuntimeTrace::UnwrapSimData(*SimDataPtr, DeserializedMoverStates->InputCommand, DeserializedMoverStates->SyncState, DeserializedMoverStates->LocalSimData);
+	}
+	if (DeserializedMoverStates)
+	{
+		OutSyncState = DeserializedMoverStates->SyncState;
+		OutInputCmd = DeserializedMoverStates->InputCommand;
+		OutLocalSimData = DeserializedMoverStates->LocalSimData;
 	}
 
 	return true;
diff --git a/Engine/Plugins/Experimental/Mover/Source/MoverCVDEditor/Private/MoverCVDSimDataComponent.h b/Engine/Plugins/Experimental/Mover/Source/MoverCVDEditor/Private/MoverCVDSimDataComponent.h
--- a/Engine/Plugins/Experimental/Mover/Source/MoverCVDEditor/Private/MoverCVDSimDataComponent.h
+++ b/Engine/Plugins/Experimental/Mover/Source/MoverCVDEditor/Private/MoverCVDSimDataComponent.h
@@ -8,6 +8,14 @@
 struct FMoverSyncState;
 struct FMoverCVDSimDataWrapper;
 struct FMoverInputCmdContext;
+struct FMoverDataCollection;
+
+struct FDeserializedMoverStates
+{
+	TSharedPtr<FMoverSyncState> SyncState;
+	TSharedPtr<FMoverInputCmdContext> InputCommand;
+	TSharedPtr<FMoverDataCollection> LocalSimData;
+};
 
 /** Component holding Mover data for the current visualized frame */
 UCLASS()
@@ -33,13 +41,13 @@ public:
 		uint32 ParticleID,
 		TSharedPtr<FMoverCVDSimDataWrapper>& OutSimDataWrapper,
 		TSharedPtr<FMoverSyncState>& OutSyncState,
-		TSharedPtr<FMoverInputCmdContext>& OutInputCmd);
+		TSharedPtr<FMoverInputCmdContext>& OutInputCmd,
+		TSharedPtr<FMoverDataCollection>& OutLocalSimData);
 
 private:
 	// This is the array of FMoverCVDSimDataWrapper for the current frame (as it is updated in UpdateFromNewGameFrameData)
 	// and corresponding to this DataComponent's SolverID
 	TArray<TSharedPtr<FMoverCVDSimDataWrapper>> FrameSimDataArray;
 
-	TMap<FMoverCVDSimDataWrapper*, TSharedPtr<FMoverSyncState>> DeserializedSyncStates;
-	TMap<FMoverCVDSimDataWrapper*, TSharedPtr<FMoverInputCmdContext>> DeserializedInputCmds;
+	TMap<FMoverCVDSimDataWrapper*, TSharedPtr<FDeserializedMoverStates>> DeserializedStates;
 };
diff --git a/Engine/Plugins/Experimental/Mover/Source/MoverCVDEditor/Private/MoverCVDTab.cpp b/Engine/Plugins/Experimental/Mover/Source/MoverCVDEditor/Private/MoverCVDTab.cpp
--- a/Engine/Plugins/Experimental/Mover/Source/MoverCVDEditor/Private/MoverCVDTab.cpp
+++ b/Engine/Plugins/Experimental/Mover/Source/MoverCVDEditor/Private/MoverCVDTab.cpp
@@ -120,7 +120,7 @@ void FMoverCVDTab::DisplaySingleParticleInfo(int32 SelectedSolverID, int32 Selec
 			UMoverCVDSimDataComponent* SimDataComponent = SimDataComponentSharedPtr.Get();
 			if (IsValid(SimDataComponent))
 			{
-				if (SimDataComponent->FindAndUnwrapSimDataForParticle(SelectedParticleID, MoverSimDataWrapper, MoverSyncState, MoverInputCmd))
+				if (SimDataComponent->FindAndUnwrapSimDataForParticle(SelectedParticleID, MoverSimDataWrapper, MoverSyncState, MoverInputCmd, MoverLocalSimData))
 				{
 					if (MoverSimDataWrapper && MoverInputCmd && MoverSyncState)
 					{
@@ -128,19 +128,39 @@ void FMoverCVDTab::DisplaySingleParticleInfo(int32 SelectedSolverID, int32 Selec
 						{
 							MultiViewWrapper.Clear();
 
-							FMoverCVDSimDataParticleInfo MoverParticleInfo;
-							MoverParticleInfo.ParticleID = MoverSimDataWrapper->ParticleID;
-							MoverParticleInfo.SolverID = MoverSimDataWrapper->SolverID;
+							// This displays the particle ID and solver ID
 							MultiViewWrapper.AddData(MoverSimDataWrapper.Get());
 
+							// Input Command Context structs
 							for (const TSharedPtr<FMoverDataStructBase>& InputDataStructSharedPtr : MoverInputCmd->InputCollection.GetDataArray())
 							{
-								if (FMoverDataStructBase* MoverDataStructBase = MoverInputCmd->InputCollection.GetDataArray()[0].Get())
+								if (FMoverDataStructBase* MoverDataStructBase = InputDataStructSharedPtr.Get())
 								{
 									MultiViewWrapper.AddData(MakeShared<FStructOnScope>(MoverDataStructBase->GetScriptStruct(), reinterpret_cast<uint8*>(MoverDataStructBase)));
 								}
 							}
+
 							MultiViewWrapper.AddData(MoverSyncState.Get());
+							// Sync State Data Collection structs
+							for (const TSharedPtr<FMoverDataStructBase>& InputDataStructSharedPtr : MoverSyncState->SyncStateCollection.GetDataArray())
+							{
+								if (FMoverDataStructBase* MoverDataStructBase = InputDataStructSharedPtr.Get())
+								{
+									MultiViewWrapper.AddData(MakeShared<FStructOnScope>(MoverDataStructBase->GetScriptStruct(), reinterpret_cast<uint8*>(MoverDataStructBase)));
+								}
+							}
+
+							// Local Simulation State structs
+							if (MoverLocalSimData)
+							{
+								for (const TSharedPtr<FMoverDataStructBase>& LocalSimDataStructSharedPtr : MoverLocalSimData->GetDataArray())
+								{
+									if (FMoverDataStructBase* MoverDataStructBase = LocalSimDataStructSharedPtr.Get())
+									{
+										MultiViewWrapper.AddData(MakeShared<FStructOnScope>(MoverDataStructBase->GetScriptStruct(), reinterpret_cast<uint8*>(MoverDataStructBase)));
+									}
+								}
+							}
 
 							WasDataFoundForParticle = true;
 
diff --git a/Engine/Plugins/Experimental/Mover/Source/MoverCVDEditor/Private/MoverCVDTab.h b/Engine/Plugins/Experimental/Mover/Source/MoverCVDEditor/Private/MoverCVDTab.h
--- a/Engine/Plugins/Experimental/Mover/Source/MoverCVDEditor/Private/MoverCVDTab.h
+++ b/Engine/Plugins/Experimental/Mover/Source/MoverCVDEditor/Private/MoverCVDTab.h
@@ -15,6 +15,7 @@ class FName;
 struct FMoverSyncState;
 struct FMoverInputCmdContext;
 struct FMoverCVDSimDataWrapper;
+struct FMoverDataCollection;
 class UMoverCVDSimDataComponent;
 
 /** This tab is an additional details tab displaying mover info corresponding to the selected particles if they are moved by a Mover component */
@@ -55,6 +56,7 @@ private:
 	TSharedPtr<FMoverCVDSimDataWrapper> MoverSimDataWrapper;
 	TSharedPtr<FMoverSyncState> MoverSyncState;
 	TSharedPtr<FMoverInputCmdContext> MoverInputCmd;
+	TSharedPtr<FMoverDataCollection> MoverLocalSimData;
 };
 
 
diff --git a/Engine/Plugins/Experimental/MoverExamples/Source/MoverExamples/Public/CharacterVariants/AbilityInputs.h b/Engine/Plugins/Experimental/MoverExamples/Source/MoverExamples/Public/CharacterVariants/AbilityInputs.h
--- a/Engine/Plugins/Experimental/MoverExamples/Source/MoverExamples/Public/CharacterVariants/AbilityInputs.h
+++ b/Engine/Plugins/Experimental/MoverExamples/Source/MoverExamples/Public/CharacterVariants/AbilityInputs.h
@@ -12,19 +12,19 @@ struct MOVEREXAMPLES_API FMoverExampleAbilityInputs : public FMoverDataStructBas
 {
 	GENERATED_USTRUCT_BODY()
 
-	UPROPERTY(BlueprintReadWrite, Category = Mover)
+	UPROPERTY(VisibleAnywhere, BlueprintReadWrite, Category = Mover)
 	bool bIsDashJustPressed = false;
 
-	UPROPERTY(BlueprintReadWrite, Category = Mover)
+	UPROPERTY(VisibleAnywhere, BlueprintReadWrite, Category = Mover)
     bool bIsAimPressed = false;
 
-	UPROPERTY(BlueprintReadWrite, Category = Mover)
+	UPROPERTY(VisibleAnywhere, BlueprintReadWrite, Category = Mover)
 	bool bIsVaultJustPressed = false;
 
-	UPROPERTY(BlueprintReadWrite, Category = Mover)
+	UPROPERTY(VisibleAnywhere, BlueprintReadWrite, Category = Mover)
 	bool bWantsToStartZiplining = false;
 
-	UPROPERTY(BlueprintReadWrite, Category = Mover)
+	UPROPERTY(VisibleAnywhere, BlueprintReadWrite, Category = Mover)
 	bool bWantsToBeCrouched = false;
 
 	// Implementation of FMoverDataStructBase
