Mover: Adding an ensure to not crash when a transition it None upon mode registration. There is data validation for it, but it's possible to test a character without the validation running in editor, and get a crash still.

diff --git a/Engine/Plugins/Experimental/Mover/Source/Mover/Private/MovementMode.cpp b/Engine/Plugins/Experimental/Mover/Source/Mover/Private/MovementMode.cpp
--- a/Engine/Plugins/Experimental/Mover/Source/Mover/Private/MovementMode.cpp
+++ b/Engine/Plugins/Experimental/Mover/Source/Mover/Private/MovementMode.cpp
@@ -31,7 +31,14 @@ void UBaseMovementMode::OnRegistered(const FName ModeName)
 {
 	for (TObjectPtr<UBaseMovementModeTransition>& Transition : Transitions)
 	{
-		Transition->OnRegistered();
+		if (Transition)
+		{
+			Transition->OnRegistered();
+		}
+		else
+		{
+			UE_LOG(LogMover, Error, TEXT("Invalid or missing transition object on mode of type %s of component %s for actor %s"), *GetPathNameSafe(this), *GetNameSafe(GetOuter()), *GetNameSafe(GetOutermost()));
+		}
 	}
 	
 	K2_OnRegistered(ModeName);
@@ -41,7 +48,14 @@ void UBaseMovementMode::OnUnregistered()
 {
 	for (TObjectPtr<UBaseMovementModeTransition>& Transition : Transitions)
 	{
-		Transition->OnUnregistered();
+		if (Transition)
+		{
+			Transition->OnUnregistered();
+		}
+		else
+		{
+			UE_LOG(LogMover, Error, TEXT("Invalid or missing transition object on mode of type %s of component %s for actor %s"), *GetPathNameSafe(this), *GetNameSafe(GetOuter()), *GetNameSafe(GetOutermost()));
+		}
 	}
 
 	K2_OnUnregistered();
