Optimize mover standalone tick functions and add a function for updating cached state

diff --git a/Engine/Plugins/Experimental/Mover/Source/Mover/Private/Backends/MoverStandaloneLiaison.cpp b/Engine/Plugins/Experimental/Mover/Source/Mover/Private/Backends/MoverStandaloneLiaison.cpp
--- a/Engine/Plugins/Experimental/Mover/Source/Mover/Private/Backends/MoverStandaloneLiaison.cpp
+++ b/Engine/Plugins/Experimental/Mover/Source/Mover/Private/Backends/MoverStandaloneLiaison.cpp
@@ -396,12 +396,11 @@ void UMoverStandaloneLiaisonComponent::OnControllerChanged(APawn* Pawn, AControl
 
 void FMoverStandaloneProduceInputTickFunction::ExecuteTick(float DeltaTime, enum ELevelTick TickType, ENamedThreads::Type CurrentThread, const FGraphEventRef& MyCompletionGraphEvent)
 {
-	FActorComponentTickFunction::ExecuteTickHelper(Target.Get(), /*bTickInEditor=*/ false, DeltaTime, TickType, [this](float DilatedTime)
+	// ExecuteTickHelper does the null check
+	UMoverStandaloneLiaisonComponent* TargetComp = Target.Get();
+	FActorComponentTickFunction::ExecuteTickHelper(TargetComp, /*bTickInEditor=*/ false, DeltaTime, TickType, [TargetComp](float DilatedTime)
 		{
-			if (Target.IsValid())
-			{
-				Target->TickInputProduction(DilatedTime);
-			}
+			TargetComp->TickInputProduction(DilatedTime);
 		});
 }
 
@@ -424,12 +423,10 @@ FName FMoverStandaloneProduceInputTickFunction::DiagnosticContext(bool bDetailed
 
 void FMoverStandaloneSimulateMovementTickFunction::ExecuteTick(float DeltaTime, enum ELevelTick TickType, ENamedThreads::Type CurrentThread, const FGraphEventRef& MyCompletionGraphEvent)
 {
-	FActorComponentTickFunction::ExecuteTickHelper(Target.Get(), /*bTickInEditor=*/ false, DeltaTime, TickType, [this](float DilatedTime)
+	UMoverStandaloneLiaisonComponent* TargetComp = Target.Get();
+	FActorComponentTickFunction::ExecuteTickHelper(TargetComp, /*bTickInEditor=*/ false, DeltaTime, TickType, [TargetComp](float DilatedTime)
 		{
-			if (Target.IsValid())
-			{
-				Target->TickMovementSimulation(DilatedTime);
-			}
+			TargetComp->TickMovementSimulation(DilatedTime);
 		});
 }
 
@@ -453,12 +450,10 @@ FName FMoverStandaloneSimulateMovementTickFunction::DiagnosticContext(bool bDeta
 
 void FMoverStandaloneApplyStateTickFunction::ExecuteTick(float DeltaTime, enum ELevelTick TickType, ENamedThreads::Type CurrentThread, const FGraphEventRef& MyCompletionGraphEvent)
 {
-	FActorComponentTickFunction::ExecuteTickHelper(Target.Get(), /*bTickInEditor=*/ false, DeltaTime, TickType, [this](float DilatedTime)
+	UMoverStandaloneLiaisonComponent* TargetComp = Target.Get();
+	FActorComponentTickFunction::ExecuteTickHelper(TargetComp, /*bTickInEditor=*/ false, DeltaTime, TickType, [TargetComp](float DilatedTime)
 		{
-			if (Target.IsValid())
-			{
-				Target->TickApplySimulationState(DilatedTime);
-			}
+			TargetComp->TickApplySimulationState(DilatedTime);
 		});
 }
 
diff --git a/Engine/Plugins/Experimental/Mover/Source/Mover/Private/MoverComponent.cpp b/Engine/Plugins/Experimental/Mover/Source/Mover/Private/MoverComponent.cpp
--- a/Engine/Plugins/Experimental/Mover/Source/Mover/Private/MoverComponent.cpp
+++ b/Engine/Plugins/Experimental/Mover/Source/Mover/Private/MoverComponent.cpp
@@ -296,12 +296,7 @@ void UMoverComponent::FinalizeFrame(const FMoverSyncState* SyncState, const FMov
 	}
 	else
 	{
-		// TODO: This is distasteful duplication -- consider moving to a util function
-		CachedLastSyncState = *SyncState;
-		LastMoverDefaultSyncState = CachedLastSyncState.SyncStateCollection.FindDataByType<FMoverDefaultSyncState>();
-		CachedLastAuxState = *AuxState;
-		CachedLastSimTickTimeStep.BaseSimTimeMs = BackendLiaisonComp->GetCurrentSimTimeMs();
-		CachedLastSimTickTimeStep.ServerFrame = BackendLiaisonComp->GetCurrentSimFrame();
+		UpdateCachedFrameState(SyncState, AuxState);
 	}
 
 	if (OnPostFinalize.IsBound())
@@ -749,18 +744,21 @@ void UMoverComponent::RemoveGameplayTags(const FGameplayTagContainer& TagsToRemo
 	ExternalGameplayTags.RemoveTags(TagsToRemove);
 }
 
-void UMoverComponent::SetFrameStateFromContext(const FMoverSyncState* SyncState, const FMoverAuxStateContext* AuxState, bool bRebaseBasedState)
+void UMoverComponent::UpdateCachedFrameState(const FMoverSyncState* SyncState, const FMoverAuxStateContext* AuxState)
 {
-
-	// TODO: This is distasteful duplication -- consider moving to a util function
-	// Copy these as the last official state, so they can be queried by other systems outside of the NP simulation
+	// TODO integrate dirty tracking
 	CachedLastSyncState = *SyncState;
 	LastMoverDefaultSyncState = CachedLastSyncState.SyncStateCollection.FindDataByType<FMoverDefaultSyncState>();
 	CachedLastAuxState = *AuxState;
 	CachedLastSimTickTimeStep.BaseSimTimeMs = BackendLiaisonComp->GetCurrentSimTimeMs();
 	CachedLastSimTickTimeStep.ServerFrame = BackendLiaisonComp->GetCurrentSimFrame();
+}
+
+void UMoverComponent::SetFrameStateFromContext(const FMoverSyncState* SyncState, const FMoverAuxStateContext* AuxState, bool bRebaseBasedState)
+{
+	UpdateCachedFrameState(SyncState, AuxState);
 
-	if (FMoverDefaultSyncState* MoverState = CachedLastSyncState.SyncStateCollection.FindMutableDataByType<FMoverDefaultSyncState>())
+	if (FMoverDefaultSyncState* MoverState = const_cast<FMoverDefaultSyncState*>(LastMoverDefaultSyncState))
 	{
 		if (bRebaseBasedState && MoverState->GetMovementBase())
 		{
diff --git a/Engine/Plugins/Experimental/Mover/Source/Mover/Public/MoverComponent.h b/Engine/Plugins/Experimental/Mover/Source/Mover/Public/MoverComponent.h
--- a/Engine/Plugins/Experimental/Mover/Source/Mover/Public/MoverComponent.h
+++ b/Engine/Plugins/Experimental/Mover/Source/Mover/Public/MoverComponent.h
@@ -508,6 +508,9 @@ protected:
 	*/
 	void SetFrameStateFromContext(const FMoverSyncState* SyncState, const FMoverAuxStateContext* AuxState, bool bRebaseBasedState);
 
+	/** Update cached frame state if it has changed */
+	void UpdateCachedFrameState(const FMoverSyncState* SyncState, const FMoverAuxStateContext* AuxState);
+
 public:
 	bool InitMoverSimulation();	// Intended to be called by the backend only
 
