Mover - Add WIP simulation class. The intention is that this will be the object that updates the mover state and that this will replace the update functionality in the mover component

diff --git a/Engine/Plugins/Experimental/Mover/Source/Mover/Private/MoverSimulation.cpp b/Engine/Plugins/Experimental/Mover/Source/Mover/Private/MoverSimulation.cpp
new file mode 100644
--- /dev/null
+++ b/Engine/Plugins/Experimental/Mover/Source/Mover/Private/MoverSimulation.cpp
@@ -0,0 +1,22 @@
+// Copyright Epic Games, Inc. All Rights Reserved.
+
+#include "MoverSimulation.h"
+
+#include "MoveLibrary/MoverBlackboard.h"
+
+#include UE_INLINE_GENERATED_CPP_BY_NAME(MoverSimulation)
+
+UMoverSimulation::UMoverSimulation()
+{
+	Blackboard = CreateDefaultSubobject<UMoverBlackboard>(TEXT("MoverSimulationBlackboard"));
+}
+
+const UMoverBlackboard* UMoverSimulation::GetBlackboard() const
+{
+	return Blackboard;
+}
+
+UMoverBlackboard* UMoverSimulation::GetBlackboard_Mutable()
+{
+	return Blackboard;
+}
\ No newline at end of file
diff --git a/Engine/Plugins/Experimental/Mover/Source/Mover/Public/InstantMovementEffect.h b/Engine/Plugins/Experimental/Mover/Source/Mover/Public/InstantMovementEffect.h
--- a/Engine/Plugins/Experimental/Mover/Source/Mover/Public/InstantMovementEffect.h
+++ b/Engine/Plugins/Experimental/Mover/Source/Mover/Public/InstantMovementEffect.h
@@ -6,6 +6,7 @@
 #include "InstantMovementEffect.generated.h"
 
 class UMoverComponent;
+class UMoverSimulation;
 struct FMoverTimeStep;
 struct FMoverTickStartData;
 struct FMoverSyncState;
@@ -29,6 +30,7 @@ struct FApplyMovementEffectParams
  */
 struct FApplyMovementEffectParams_Async
 {
+	UMoverSimulation* Simulation;
 	const FMoverTickStartData* StartState;
 	const FMoverTimeStep* TimeStep;
 };
diff --git a/Engine/Plugins/Experimental/Mover/Source/Mover/Public/MoverSimulation.h b/Engine/Plugins/Experimental/Mover/Source/Mover/Public/MoverSimulation.h
new file mode 100644
--- /dev/null
+++ b/Engine/Plugins/Experimental/Mover/Source/Mover/Public/MoverSimulation.h
@@ -0,0 +1,33 @@
+// Copyright Epic Games, Inc. All Rights Reserved.
+
+#pragma once
+
+#include "UObject/Object.h"
+
+#include "MoverSimulation.generated.h"
+
+class UMoverBlackboard;
+
+/**
+* WIP Base class for a Mover simulation.
+* The simulation is intended to be the thing that updates the Mover
+* state and should be safe to run on an async thread
+*/
+UCLASS(BlueprintType)
+class MOVER_API UMoverSimulation : public UObject
+{
+	GENERATED_BODY()
+
+public:
+	UMoverSimulation();
+
+	UFUNCTION(BlueprintPure, Category = Mover)
+	const UMoverBlackboard* GetBlackboard() const;
+
+	UFUNCTION(BlueprintPure, Category = Mover)
+	UMoverBlackboard* GetBlackboard_Mutable();
+
+protected:
+	UPROPERTY(Transient)
+	TObjectPtr<UMoverBlackboard> Blackboard = nullptr;
+};
\ No newline at end of file
