Mover: Caching NavData in NavWalkingModes to avoid searching for it every frame and potential issues with searching through components off of the game thread

diff --git a/Engine/Plugins/Experimental/Mover/Source/Mover/Private/DefaultMovementSet/Modes/AsyncNavWalkingMode.cpp b/Engine/Plugins/Experimental/Mover/Source/Mover/Private/DefaultMovementSet/Modes/AsyncNavWalkingMode.cpp
--- a/Engine/Plugins/Experimental/Mover/Source/Mover/Private/DefaultMovementSet/Modes/AsyncNavWalkingMode.cpp
+++ b/Engine/Plugins/Experimental/Mover/Source/Mover/Private/DefaultMovementSet/Modes/AsyncNavWalkingMode.cpp
@@ -28,6 +28,7 @@ UAsyncNavWalkingMode::UAsyncNavWalkingMode()
 	, NavMeshProjectionInterpSpeed(12.f)
 	, NavMeshProjectionTimer(0)
 	, NavMoverComponent(nullptr)
+	, NavDataInterface(nullptr)
 	, bProjectNavMeshOnBothWorldChannels(true)
 {
 	GameplayTags.AddTag(Mover_IsOnGround);
@@ -178,7 +179,11 @@ void UAsyncNavWalkingMode::SimulationTick_Implementation(const FSimulationTickPa
 	float SimpleHalfHeight = 0;
 	NavMoverComponent->GetSimpleCollisionCylinder(OUT SimpleRadius, OUT SimpleHalfHeight);
 
-	const INavigationDataInterface* NavData = nullptr;
+	if (!NavDataInterface.IsValid())
+	{
+		NavDataInterface = GetNavData();
+	}
+	
 	bool bSameNavLocation = false;
 	if (CachedNavLocation.NodeRef != INVALID_NAVNODEREF)
 	{
@@ -200,11 +205,9 @@ void UAsyncNavWalkingMode::SimulationTick_Implementation(const FSimulationTickPa
 
 		if (bDeltaMoveNearlyZero && bSameNavLocation)
 		{
-			NavData = GetNavData();
-			
-			if (NavData)
+			if (NavDataInterface.IsValid())
 			{
-				if (!NavData->IsNodeRefValid(CachedNavLocation.NodeRef))
+				if (!NavDataInterface->IsNodeRefValid(CachedNavLocation.NodeRef))
 				{
 					CachedNavLocation.NodeRef = INVALID_NAVNODEREF;
 					bSameNavLocation = false;
@@ -227,14 +230,9 @@ void UAsyncNavWalkingMode::SimulationTick_Implementation(const FSimulationTickPa
 		{
 			UMovementUtils::SetGravityVerticalComponent(TargetFeetLocation, CachedNavLocation.Location.Dot(UpDirection), UpDirection);
 		}
-
-		if (!NavData)
-		{
-			NavData = GetNavData();
-		}
 		
 		// Find the point on the NavMesh
-		const bool bHasNavigationData = FindNavFloor(TargetFeetLocation, DestNavLocation, NavData);
+		const bool bHasNavigationData = FindNavFloor(TargetFeetLocation, DestNavLocation, NavDataInterface.Get());
 		if (!bHasNavigationData)
 		{
 			// Can't find nav mesh at this location, so we need to do something else
@@ -356,6 +354,8 @@ void UAsyncNavWalkingMode::Activate()
 		}
 	}
 
+	NavDataInterface = GetNavData();
+	
 	NavMeshProjectionTimer = (NavMeshProjectionInterval > 0.f) ? FMath::FRandRange(-NavMeshProjectionInterval, 0.f) : 0.f;
 }
 
@@ -379,7 +379,7 @@ const INavigationDataInterface* UAsyncNavWalkingMode::GetNavData() const
 	}
 
 
-	const INavigationDataInterface* NavData = FNavigationSystem::GetNavDataForActor(*MoverComponentOwner);
+	INavigationDataInterface* NavData = FNavigationSystem::GetNavDataForActor(*MoverComponentOwner);
 
 	return NavData;
 }
@@ -564,6 +564,7 @@ void UAsyncNavWalkingMode::OnUnregistered()
 {
 	CommonLegacySettings = nullptr;
 	NavMoverComponent = nullptr;
+	NavDataInterface = nullptr;
 
 	Super::OnUnregistered();
 }
diff --git a/Engine/Plugins/Experimental/Mover/Source/Mover/Private/DefaultMovementSet/Modes/NavWalkingMode.cpp b/Engine/Plugins/Experimental/Mover/Source/Mover/Private/DefaultMovementSet/Modes/NavWalkingMode.cpp
--- a/Engine/Plugins/Experimental/Mover/Source/Mover/Private/DefaultMovementSet/Modes/NavWalkingMode.cpp
+++ b/Engine/Plugins/Experimental/Mover/Source/Mover/Private/DefaultMovementSet/Modes/NavWalkingMode.cpp
@@ -23,6 +23,7 @@ UNavWalkingMode::UNavWalkingMode()
 	, NavMeshProjectionInterpSpeed(12.f)
 	, NavMeshProjectionTimer(0)
 	, NavMoverComponent(nullptr)
+	, NavDataInterface(nullptr)
 	, bProjectNavMeshOnBothWorldChannels(true)
 {
 	GameplayTags.AddTag(Mover_IsOnGround);
@@ -172,7 +173,11 @@ void UNavWalkingMode::SimulationTick_Implementation(const FSimulationTickParams&
 	float SimpleHalfHeight = 0;
 	NavMoverComponent->GetSimpleCollisionCylinder(SimpleRadius, SimpleHalfHeight);
 	
-	const INavigationDataInterface* NavData = nullptr;
+	if (!NavDataInterface.IsValid())
+	{
+		NavDataInterface = GetNavData();
+	}
+	
 	bool bSameNavLocation = false;
 	if (CachedNavLocation.NodeRef != INVALID_NAVNODEREF)
 	{
@@ -194,11 +199,9 @@ void UNavWalkingMode::SimulationTick_Implementation(const FSimulationTickParams&
 
 		if (bDeltaMoveNearlyZero && bSameNavLocation)
 		{
-			NavData = GetNavData();
-			
-			if (NavData)
+			if (NavDataInterface.IsValid())
 			{
-				if (!NavData->IsNodeRefValid(CachedNavLocation.NodeRef))
+				if (!NavDataInterface->IsNodeRefValid(CachedNavLocation.NodeRef))
 				{
 					CachedNavLocation.NodeRef = INVALID_NAVNODEREF;
 					bSameNavLocation = false;
@@ -221,14 +224,9 @@ void UNavWalkingMode::SimulationTick_Implementation(const FSimulationTickParams&
 		{
 			UMovementUtils::SetGravityVerticalComponent(AdjustedDest, CachedNavLocation.Location.Dot(UpDirection), UpDirection);
 		}
-
-		if (!NavData)
-		{
-			NavData = GetNavData();
-		}
 		
 		// Find the point on the NavMesh
-		const bool bHasNavigationData = FindNavFloor(AdjustedDest, DestNavLocation, NavData);
+		const bool bHasNavigationData = FindNavFloor(AdjustedDest, DestNavLocation, NavDataInterface.Get());
 		if (!bHasNavigationData)
 		{
 			// Can't find nav mesh at this location, so we need to do something else
@@ -369,6 +367,8 @@ void UNavWalkingMode::Activate()
 {
 	Super::Activate();
 	SetCollisionForNavWalking(true);
+
+	NavDataInterface = GetNavData();
 }
 
 void UNavWalkingMode::Deactivate()
@@ -579,6 +579,7 @@ void UNavWalkingMode::OnRegistered(const FName ModeName)
 void UNavWalkingMode::OnUnregistered()
 {
 	CommonLegacySettings = nullptr;
+	NavDataInterface = nullptr;
 
 	Super::OnUnregistered();
 }
diff --git a/Engine/Plugins/Experimental/Mover/Source/Mover/Public/DefaultMovementSet/Modes/AsyncNavWalkingMode.h b/Engine/Plugins/Experimental/Mover/Source/Mover/Public/DefaultMovementSet/Modes/AsyncNavWalkingMode.h
--- a/Engine/Plugins/Experimental/Mover/Source/Mover/Public/DefaultMovementSet/Modes/AsyncNavWalkingMode.h
+++ b/Engine/Plugins/Experimental/Mover/Source/Mover/Public/DefaultMovementSet/Modes/AsyncNavWalkingMode.h
@@ -95,6 +95,9 @@ protected:
 	/** Associated Movement component that will actually move the actor */
 	TWeakObjectPtr<UNavMoverComponent> NavMoverComponent;
 
+	// Note: This isn't guaranteed to be valid at all times. It will be most of the time but re-entering this mode to call Activate() will search for NavData again and update it accordingly.
+	TWeakInterfacePtr<const INavigationDataInterface> NavDataInterface;
+	
 	/** Use both WorldStatic and WorldDynamic channels for NavWalking geometry conforming */
 	UPROPERTY(Category = "NavMesh Movement", EditAnywhere, BlueprintReadOnly, AdvancedDisplay)
 	uint8 bProjectNavMeshOnBothWorldChannels : 1;
diff --git a/Engine/Plugins/Experimental/Mover/Source/Mover/Public/DefaultMovementSet/Modes/NavWalkingMode.h b/Engine/Plugins/Experimental/Mover/Source/Mover/Public/DefaultMovementSet/Modes/NavWalkingMode.h
--- a/Engine/Plugins/Experimental/Mover/Source/Mover/Public/DefaultMovementSet/Modes/NavWalkingMode.h
+++ b/Engine/Plugins/Experimental/Mover/Source/Mover/Public/DefaultMovementSet/Modes/NavWalkingMode.h
@@ -117,6 +117,9 @@ protected:
 	/** Associated Movement component that will actually move the actor */ 
 	UPROPERTY(BlueprintReadOnly, Category="Nav Movement")
 	TObjectPtr<UNavMoverComponent> NavMoverComponent;
+
+	// Note: This isn't guaranteed to be valid at all times. It will be most of the time but re-entering this mode to call Activate() will search for NavData again and update it accordingly.
+	TWeakInterfacePtr<const INavigationDataInterface> NavDataInterface;
 	
 	/** Use both WorldStatic and WorldDynamic channels for NavWalking geometry conforming */
 	UPROPERTY(Category = "NavMesh Movement", EditAnywhere, BlueprintReadOnly, AdvancedDisplay)
