Mover: adding new options for handling attempted off-nav-mesh excursions so a character can rotate only but will not move, and a separate option for allowing sliding along the edges.

diff --git a/Engine/Plugins/Experimental/Mover/Source/Mover/Private/DefaultMovementSet/Modes/AsyncNavWalkingMode.cpp b/Engine/Plugins/Experimental/Mover/Source/Mover/Private/DefaultMovementSet/Modes/AsyncNavWalkingMode.cpp
--- a/Engine/Plugins/Experimental/Mover/Source/Mover/Private/DefaultMovementSet/Modes/AsyncNavWalkingMode.cpp
+++ b/Engine/Plugins/Experimental/Mover/Source/Mover/Private/DefaultMovementSet/Modes/AsyncNavWalkingMode.cpp
@@ -185,7 +185,7 @@ void UAsyncNavWalkingMode::SimulationTick_Implementation(const FSimulationTickPa
 	}
 	
 	bool bSameNavLocation = false;
-	if (CachedNavLocation.NodeRef != INVALID_NAVNODEREF)
+	if (CachedNavLocation.HasNodeRef())
 	{
 		if (bProjectNavMeshWalking)
 		{
@@ -232,7 +232,23 @@ void UAsyncNavWalkingMode::SimulationTick_Implementation(const FSimulationTickPa
 		}
 		
 		// Find the point on the NavMesh
-		const bool bHasNavigationData = FindNavFloor(TargetFeetLocation, DestNavLocation, NavDataInterface.Get());
+		bool bHasNavigationData = false;
+
+		if (bSlideAlongNavMeshEdge && CachedNavLocation.HasNodeRef())
+		{
+			bHasNavigationData = NavDataInterface.Get()->FindMoveAlongSurface(CachedNavLocation, TargetFeetLocation, OUT DestNavLocation);
+
+			if (bHasNavigationData)
+			{
+				TargetFeetLocation = UMovementUtils::ProjectToGravityFloor(DestNavLocation.Location, UpDirection) + UMovementUtils::GetGravityVerticalComponent(TargetFeetLocation, UpDirection);
+			}
+		}
+		else
+		{
+			bHasNavigationData = FindNavFloor(TargetFeetLocation, OUT DestNavLocation, NavDataInterface.Get());
+		}
+
+
 		if (!bHasNavigationData)
 		{
 			// Can't find nav mesh at this location, so we need to do something else
@@ -257,6 +273,10 @@ void UAsyncNavWalkingMode::SimulationTick_Implementation(const FSimulationTickPa
 					UE_LOG(LogMover, Verbose, TEXT("%s could not find valid navigation data at location %s. Cannot move."), *GetNameSafe(MoverComp->GetOwner()), *TargetFeetLocation.ToCompactString());
 					// nothing to be done
 					break;
+
+				case EOffNavMeshBehavior::RotateOnly:
+					RotationInProgress = TargetRotation;
+					break;
 			}
 
 			CaptureFinalState(LocationInProgress, RotationInProgress.Rotator(), MoveRecord, OutputSyncState);
diff --git a/Engine/Plugins/Experimental/Mover/Source/Mover/Private/DefaultMovementSet/Modes/NavWalkingMode.cpp b/Engine/Plugins/Experimental/Mover/Source/Mover/Private/DefaultMovementSet/Modes/NavWalkingMode.cpp
--- a/Engine/Plugins/Experimental/Mover/Source/Mover/Private/DefaultMovementSet/Modes/NavWalkingMode.cpp
+++ b/Engine/Plugins/Experimental/Mover/Source/Mover/Private/DefaultMovementSet/Modes/NavWalkingMode.cpp
@@ -226,7 +226,23 @@ void UNavWalkingMode::SimulationTick_Implementation(const FSimulationTickParams&
 		}
 		
 		// Find the point on the NavMesh
-		const bool bHasNavigationData = FindNavFloor(AdjustedDest, DestNavLocation, NavDataInterface.Get());
+		bool bHasNavigationData = false;
+
+		if (bSlideAlongNavMeshEdge && CachedNavLocation.HasNodeRef())
+		{
+			bHasNavigationData = NavDataInterface.Get()->FindMoveAlongSurface(CachedNavLocation, AdjustedDest, OUT DestNavLocation);
+
+			if (bHasNavigationData)
+			{
+				AdjustedDest = UMovementUtils::ProjectToGravityFloor(DestNavLocation.Location, UpDirection) + UMovementUtils::GetGravityVerticalComponent(AdjustedDest, UpDirection);
+			}
+		}
+		else
+		{
+			bHasNavigationData = FindNavFloor(AdjustedDest, DestNavLocation, NavDataInterface.Get());
+		}
+
+
 		if (!bHasNavigationData)
 		{
 			// Can't find nav mesh at this location, so we need to do something else
@@ -250,6 +266,11 @@ void UNavWalkingMode::SimulationTick_Implementation(const FSimulationTickParams&
 					UE_LOG(LogMover, Verbose, TEXT("%s could not find valid navigation data at location %s. Cannot move."), *GetNameSafe(MoverComp->GetOwner()), *AdjustedDest.ToCompactString());
 					// nothing to be done
 					break;
+
+				case EOffNavMeshBehavior::RotateOnly:
+					FHitResult MoveHitResult;
+					UMovementUtils::TrySafeMoveUpdatedComponent(Params.MovingComps, FVector::ZeroVector, TargetOrientQuat, /*bSweep?*/ false, MoveHitResult, ETeleportType::None, MoveRecord);
+					break;
 			}
 
 			CaptureFinalState(UpdatedComponent, MoveRecord, OutputSyncState);
diff --git a/Engine/Plugins/Experimental/Mover/Source/Mover/Public/DefaultMovementSet/Modes/AsyncNavWalkingMode.h b/Engine/Plugins/Experimental/Mover/Source/Mover/Public/DefaultMovementSet/Modes/AsyncNavWalkingMode.h
--- a/Engine/Plugins/Experimental/Mover/Source/Mover/Public/DefaultMovementSet/Modes/AsyncNavWalkingMode.h
+++ b/Engine/Plugins/Experimental/Mover/Source/Mover/Public/DefaultMovementSet/Modes/AsyncNavWalkingMode.h
@@ -67,6 +67,10 @@ public:
 	UPROPERTY(Category = "NavMesh Movement", EditAnywhere, BlueprintReadWrite)
 	EOffNavMeshBehavior BehaviorOffNavMesh = EOffNavMeshBehavior::SwitchToWalking;
 
+	/** If attempting to stray off the nav mesh, should we slide along the edge instead? See @BehaviorOffNavMesh for cases where a sliding move can't be determined. */
+	UPROPERTY(Category = "NavMesh Movement", EditAnywhere, BlueprintReadWrite)
+	bool bSlideAlongNavMeshEdge = false;
+	
 
 	UPROPERTY(Transient)
 	float NavMeshProjectionTimer;
diff --git a/Engine/Plugins/Experimental/Mover/Source/Mover/Public/DefaultMovementSet/Modes/NavWalkingMode.h b/Engine/Plugins/Experimental/Mover/Source/Mover/Public/DefaultMovementSet/Modes/NavWalkingMode.h
--- a/Engine/Plugins/Experimental/Mover/Source/Mover/Public/DefaultMovementSet/Modes/NavWalkingMode.h
+++ b/Engine/Plugins/Experimental/Mover/Source/Mover/Public/DefaultMovementSet/Modes/NavWalkingMode.h
@@ -22,14 +22,17 @@ namespace NavWalkingBlackBoard
 UENUM(BlueprintType)
 enum class EOffNavMeshBehavior : uint8
 {
-	/** We change to normal walking mode and re-attempt the move. Something will later need to switch back to nav walking. */
+	/** We change to normal walking mode and re-attempt the move. Typically this is more expensive than nav walking, and we'll later need to switch back to nav walking. */
 	SwitchToWalking = 0        UMETA(DisplayName = "Switch to Walking Mode"),
 
 	/** We proceed to move as directed, but height may diverge from the floor until we return to a valid nav mesh. */
 	MoveWithoutNavMesh = 1     UMETA(DisplayName = "Move Without Nav Mesh"),
 
-	/** We do not move. Movement will continue later if nav mesh becomes available. */
-	DoNotMove = 2			   UMETA(DisplayName = "Do Not Move"),
+	/** We do not move. Movement will continue once instructed to move over nav mesh again. */
+	DoNotMove = 2			   UMETA(DisplayName = "Do Not Move or Rotate"),
+
+	/** We do not move, but will allow rotation in place. Movement will continue once instructed to move over nav mesh again. */
+	RotateOnly = 3			   UMETA(DisplayName = "Rotate Only"),
 };
 
 
@@ -89,6 +92,9 @@ public:
 	UPROPERTY(Category = "NavMesh Movement", EditAnywhere, BlueprintReadWrite)
 	EOffNavMeshBehavior BehaviorOffNavMesh = EOffNavMeshBehavior::SwitchToWalking;
 
+	/** If attempting to stray off the nav mesh, should we attempt to slide along the edge instead? See @BehaviorOffNavMesh for cases where a sliding move can't be determined. */
+	UPROPERTY(Category = "NavMesh Movement", EditAnywhere, BlueprintReadWrite)
+	bool bSlideAlongNavMeshEdge = false;
 	
 	UPROPERTY(Transient)
 	float NavMeshProjectionTimer;
