[Mover] Disabling multithreading for Physics Mover

diff --git a/Engine/Plugins/Experimental/Mover/Source/Mover/Public/PhysicsMover/PhysicsMoverSimulationTypes.h b/Engine/Plugins/Experimental/Mover/Source/Mover/Public/PhysicsMover/PhysicsMoverSimulationTypes.h
--- a/Engine/Plugins/Experimental/Mover/Source/Mover/Public/PhysicsMover/PhysicsMoverSimulationTypes.h
+++ b/Engine/Plugins/Experimental/Mover/Source/Mover/Public/PhysicsMover/PhysicsMoverSimulationTypes.h
@@ -28,7 +28,7 @@ struct FPhysicsDrivenMotionDebugParams
 	float TeleportThreshold = 1000.0f;
 	float MinStepUpDistance = 5.0f;
 	float MaxCharacterGroundMassRatio = 1.5f;
-	bool EnableMultithreading = true;
+	bool EnableMultithreading = false;
 	bool DebugDrawGroundQueries = false;
 };
 
