Mover: Making mover Component Blueprintable so users can set up BP mover components in content. Ideally we would like to get rid of derived C++ classes of mover components and only differentiate mover components by how they are set up (which modes, transitions and other movement objects they are made of)

diff --git a/Engine/Plugins/Experimental/Mover/Source/Mover/Public/MoverComponent.h b/Engine/Plugins/Experimental/Mover/Source/Mover/Public/MoverComponent.h
--- a/Engine/Plugins/Experimental/Mover/Source/Mover/Public/MoverComponent.h
+++ b/Engine/Plugins/Experimental/Mover/Source/Mover/Public/MoverComponent.h
@@ -58,7 +58,7 @@ DECLARE_DYNAMIC_DELEGATE_ThreeParams(FMover_ProcessGeneratedMovement, const FMov
 /**
  * 
  */
-UCLASS(BlueprintType, meta = (BlueprintSpawnableComponent))
+UCLASS(Blueprintable, BlueprintType, meta = (BlueprintSpawnableComponent))
 class MOVER_API UMoverComponent : public UActorComponent
 {
 	GENERATED_BODY()
