PBCM Fix for multithread cvar using inverse of value
#rnx

diff --git a/Engine/Plugins/Experimental/Mover/Source/Mover/Private/PhysicsMover/PhysicsMoverManagerAsyncCallback.cpp b/Engine/Plugins/Experimental/Mover/Source/Mover/Private/PhysicsMover/PhysicsMoverManagerAsyncCallback.cpp
--- a/Engine/Plugins/Experimental/Mover/Source/Mover/Private/PhysicsMover/PhysicsMoverManagerAsyncCallback.cpp
+++ b/Engine/Plugins/Experimental/Mover/Source/Mover/Private/PhysicsMover/PhysicsMoverManagerAsyncCallback.cpp
@@ -37,7 +37,7 @@ void FPhysicsMoverManagerAsyncCallback::ProcessInputs_Internal(int32 PhysicsStep
 		}
 	};
 
-	bool ForceSingleThread = GPhysicsDrivenMotionDebugParams.EnableMultithreading;
+	bool ForceSingleThread = !GPhysicsDrivenMotionDebugParams.EnableMultithreading;
 	Chaos::PhysicsParallelFor(InputDataArray.Num(), LambdaParallelUpdate, ForceSingleThread);
 }
 
@@ -70,7 +70,7 @@ void FPhysicsMoverManagerAsyncCallback::OnPreSimulate_Internal()
 		}
 	};
 
-	bool ForceSingleThread = GPhysicsDrivenMotionDebugParams.EnableMultithreading;
+	bool ForceSingleThread = !GPhysicsDrivenMotionDebugParams.EnableMultithreading;
 	Chaos::PhysicsParallelFor(InputDataArray.Num(), LambdaParallelUpdate, ForceSingleThread);
 }
 
@@ -99,6 +99,6 @@ void FPhysicsMoverManagerAsyncCallback::OnContactModification_Internal(Chaos::FC
 		}
 	};
 
-	bool ForceSingleThread = GPhysicsDrivenMotionDebugParams.EnableMultithreading;
+	bool ForceSingleThread = !GPhysicsDrivenMotionDebugParams.EnableMultithreading;
 	Chaos::PhysicsParallelFor(InputDataArray.Num(), LambdaParallelUpdate, ForceSingleThread);
 }
\ No newline at end of file
