Mover: NavWalkingMode changes that grab the NavigationSystem directly and find NavData based on properties from the NavMoverComponent. Avoids a FindComponentByInterface call so we're more performant and threadsafe.

diff --git a/Engine/Plugins/Experimental/Mover/Source/Mover/Mover.Build.cs b/Engine/Plugins/Experimental/Mover/Source/Mover/Mover.Build.cs
--- a/Engine/Plugins/Experimental/Mover/Source/Mover/Mover.Build.cs
+++ b/Engine/Plugins/Experimental/Mover/Source/Mover/Mover.Build.cs
@@ -19,7 +19,8 @@ public class Mover : ModuleRules
 				"AnimGraphRuntime",
 				"MotionWarping",
 				"Water",
-				"GameplayTags"
+				"GameplayTags",
+				"NavigationSystem"
 			}
 			);
 			
diff --git a/Engine/Plugins/Experimental/Mover/Source/Mover/Private/DefaultMovementSet/Modes/AsyncNavWalkingMode.cpp b/Engine/Plugins/Experimental/Mover/Source/Mover/Private/DefaultMovementSet/Modes/AsyncNavWalkingMode.cpp
--- a/Engine/Plugins/Experimental/Mover/Source/Mover/Private/DefaultMovementSet/Modes/AsyncNavWalkingMode.cpp
+++ b/Engine/Plugins/Experimental/Mover/Source/Mover/Private/DefaultMovementSet/Modes/AsyncNavWalkingMode.cpp
@@ -14,6 +14,7 @@
 #include "MoveLibrary/ModularMovement.h"
 #include "MoveLibrary/MovementUtils.h"
 #include "DrawDebugHelpers.h"
+#include "NavigationSystem.h"
 
 #include UE_INLINE_GENERATED_CPP_BY_NAME(AsyncNavWalkingMode)
 
@@ -234,20 +235,22 @@ void UAsyncNavWalkingMode::SimulationTick_Implementation(const FSimulationTickPa
 		// Find the point on the NavMesh
 		bool bHasNavigationData = false;
 
-		if (bSlideAlongNavMeshEdge && CachedNavLocation.HasNodeRef())
+		if (NavDataInterface.IsValid())
 		{
-			bHasNavigationData = NavDataInterface.Get()->FindMoveAlongSurface(CachedNavLocation, TargetFeetLocation, OUT DestNavLocation);
+			if (bSlideAlongNavMeshEdge && CachedNavLocation.HasNodeRef())
+			{
+				bHasNavigationData = NavDataInterface->FindMoveAlongSurface(CachedNavLocation, TargetFeetLocation, OUT DestNavLocation);
 
-			if (bHasNavigationData)
+				if (bHasNavigationData)
+				{
+					TargetFeetLocation = UMovementUtils::ProjectToGravityFloor(DestNavLocation.Location, UpDirection) + UMovementUtils::GetGravityVerticalComponent(TargetFeetLocation, UpDirection);
+				}
+			}
+			else
 			{
-				TargetFeetLocation = UMovementUtils::ProjectToGravityFloor(DestNavLocation.Location, UpDirection) + UMovementUtils::GetGravityVerticalComponent(TargetFeetLocation, UpDirection);
+				bHasNavigationData = FindNavFloor(TargetFeetLocation, OUT DestNavLocation, NavDataInterface.Get());
 			}
 		}
-		else
-		{
-			bHasNavigationData = FindNavFloor(TargetFeetLocation, OUT DestNavLocation, NavDataInterface.Get());
-		}
-
 
 		if (!bHasNavigationData)
 		{
@@ -382,24 +385,17 @@ void UAsyncNavWalkingMode::Activate()
 
 const INavigationDataInterface* UAsyncNavWalkingMode::GetNavData() const
 {
-	const UWorld* World = GetWorld();
+	ANavigationData* NavData = nullptr;
 
-	const AActor* MoverComponentOwner;
-	if (const UMoverComponent* MoverComponent = GetMoverComponent())
+	if (const UWorld* World = GetWorld())
 	{
-		MoverComponentOwner = MoverComponent->GetOwner();
-		if (World == nullptr || World->GetNavigationSystem() == nullptr || MoverComponentOwner == nullptr)
+		const UNavigationSystemV1* NavSys = Cast<UNavigationSystemV1>(World->GetNavigationSystem());
+		if (NavSys && NavMoverComponent.IsValid())
 		{
-			return nullptr;
+			const FNavAgentProperties& AgentProps = NavMoverComponent->GetNavAgentPropertiesRef();
+			NavData = NavSys->GetNavDataForProps(AgentProps, NavMoverComponent->GetNavLocation());
 		}
 	}
-	else
-	{
-		return nullptr;
-	}
-
-
-	INavigationDataInterface* NavData = FNavigationSystem::GetNavDataForActor(*MoverComponentOwner);
 
 	return NavData;
 }
diff --git a/Engine/Plugins/Experimental/Mover/Source/Mover/Private/DefaultMovementSet/Modes/NavWalkingMode.cpp b/Engine/Plugins/Experimental/Mover/Source/Mover/Private/DefaultMovementSet/Modes/NavWalkingMode.cpp
--- a/Engine/Plugins/Experimental/Mover/Source/Mover/Private/DefaultMovementSet/Modes/NavWalkingMode.cpp
+++ b/Engine/Plugins/Experimental/Mover/Source/Mover/Private/DefaultMovementSet/Modes/NavWalkingMode.cpp
@@ -3,6 +3,7 @@
 
 #include "DefaultMovementSet/Modes/NavWalkingMode.h"
 #include "MoverComponent.h"
+#include "NavigationSystem.h"
 #include "AI/NavigationSystemBase.h"
 #include "AI/Navigation/NavigationDataInterface.h"
 #include "Components/ShapeComponent.h"
@@ -228,19 +229,22 @@ void UNavWalkingMode::SimulationTick_Implementation(const FSimulationTickParams&
 		// Find the point on the NavMesh
 		bool bHasNavigationData = false;
 
-		if (bSlideAlongNavMeshEdge && CachedNavLocation.HasNodeRef())
+		if (NavDataInterface.IsValid())
 		{
-			bHasNavigationData = NavDataInterface.Get()->FindMoveAlongSurface(CachedNavLocation, AdjustedDest, OUT DestNavLocation);
+			if (bSlideAlongNavMeshEdge && CachedNavLocation.HasNodeRef())
+			{
+				bHasNavigationData = NavDataInterface->FindMoveAlongSurface(CachedNavLocation, AdjustedDest, OUT DestNavLocation);
 
-			if (bHasNavigationData)
+				if (bHasNavigationData)
+				{
+					AdjustedDest = UMovementUtils::ProjectToGravityFloor(DestNavLocation.Location, UpDirection) + UMovementUtils::GetGravityVerticalComponent(AdjustedDest, UpDirection);
+				}
+			}
+			else
 			{
-				AdjustedDest = UMovementUtils::ProjectToGravityFloor(DestNavLocation.Location, UpDirection) + UMovementUtils::GetGravityVerticalComponent(AdjustedDest, UpDirection);
+				bHasNavigationData = FindNavFloor(AdjustedDest, DestNavLocation, NavDataInterface.Get());
 			}
 		}
-		else
-		{
-			bHasNavigationData = FindNavFloor(AdjustedDest, DestNavLocation, NavDataInterface.Get());
-		}
 
 
 		if (!bHasNavigationData)
@@ -400,25 +404,18 @@ void UNavWalkingMode::Deactivate()
 
 const INavigationDataInterface* UNavWalkingMode::GetNavData() const
 {
-	const UWorld* World = GetWorld();
+	ANavigationData* NavData = nullptr;
 	
-	const AActor* MoverComponentOwner;
-	if (const UMoverComponent* MoverComponent = GetMoverComponent())
+	if (const UWorld* World = GetWorld())
 	{
-		MoverComponentOwner = MoverComponent->GetOwner();
-		if (World == nullptr || World->GetNavigationSystem() == nullptr || MoverComponentOwner == nullptr)
+		const UNavigationSystemV1* NavSys = Cast<UNavigationSystemV1>(World->GetNavigationSystem());
+		if (NavSys && NavMoverComponent)
 		{
-			return nullptr;
+			const FNavAgentProperties& AgentProps = NavMoverComponent->GetNavAgentPropertiesRef();
+			NavData = NavSys->GetNavDataForProps(AgentProps, NavMoverComponent->GetNavLocation());
 		}
 	}
-	else
-	{
-		return nullptr;
-	}
 	
-
-	const INavigationDataInterface* NavData = FNavigationSystem::GetNavDataForActor(*MoverComponentOwner);
-
 	return NavData;
 }
 
