Mover: adding sync state type readout for gameplay debugger tool category

diff --git a/Engine/Plugins/Experimental/Mover/Source/Mover/Private/Debug/GameplayDebuggerCategory_Mover.cpp b/Engine/Plugins/Experimental/Mover/Source/Mover/Private/Debug/GameplayDebuggerCategory_Mover.cpp
--- a/Engine/Plugins/Experimental/Mover/Source/Mover/Private/Debug/GameplayDebuggerCategory_Mover.cpp
+++ b/Engine/Plugins/Experimental/Mover/Source/Mover/Private/Debug/GameplayDebuggerCategory_Mover.cpp
@@ -109,6 +109,7 @@ void FGameplayDebuggerCategory_Mover::CollectData(APlayerController* OwnerPC, AA
 	DataPack.MoveIntent = FVector::ZeroVector;
 	DataPack.ActiveLayeredMoves.Empty();
 	DataPack.ActiveModifiers.Empty();
+	DataPack.SyncStateDataTypes.Empty();
 	DataPack.ModeMap.Empty();
 	DataPack.MoveInputType = 0;
 	DataPack.MoveInput = FVector::ZeroVector;
@@ -155,6 +156,11 @@ void FGameplayDebuggerCategory_Mover::CollectData(APlayerController* OwnerPC, AA
 			DataPack.ActiveModifiers.Add(*It->Get()->ToSimpleString());
 		}
 
+		for (auto It = SyncState.SyncStateCollection.GetDataArray().CreateConstIterator(); It; ++It)
+		{
+			DataPack.SyncStateDataTypes.Add(It->Get()->GetScriptStruct()->GetName());
+		}
+
 		const FMoverInputCmdContext& LastInputCmd = MyMoverComponent->GetLastInputCmd();
 
 		if (const FCharacterDefaultInputs* DefaultInputs = LastInputCmd.InputCollection.FindDataByType<FCharacterDefaultInputs>())
@@ -199,12 +205,13 @@ void FGameplayDebuggerCategory_Mover::DrawData(APlayerController* OwnerPC, FGame
 		);
 	}
 
-	CanvasContext.Printf(TEXT("{yellow}Active Moves: {white}\n%s\n{yellow}Active Modifiers: {white}\n%s\n{yellow}Mode Map: \n{white}%s\n{yellow}Active Transitions: {white}\n%s"),
+	CanvasContext.Printf(TEXT("{yellow}Active Moves: {white}\n%s\n{yellow}Active Modifiers: {white}\n%s\n{yellow}Mode Map: \n{white}%s\n{yellow}Active Transitions: {white}\n%s\n{yellow}SyncStateTypes: {white}%s"),
 		*FString::JoinBy(DataPack.ActiveLayeredMoves, TEXT("\n"), [](FString MoveAsString) { return MoveAsString; }),
 		*FString::JoinBy(DataPack.ActiveModifiers, TEXT("\n"), [](FString ModifierAsString) { return ModifierAsString; }),
 		*FString::JoinBy(DataPack.ModeMap, TEXT("\n"), [](FString ModeMappingAsString) { return ModeMappingAsString; }),
-		*FString::JoinBy(DataPack.ActiveTransitions, TEXT("\n"), [](FString TransitionAsString) { return TransitionAsString; })
-	);
+		*FString::JoinBy(DataPack.ActiveTransitions, TEXT("\n"), [](FString TransitionAsString) { return TransitionAsString; }),
+		*FString::JoinBy(DataPack.SyncStateDataTypes, TEXT(","), [](FString SyncStateTypeAsString) { return SyncStateTypeAsString; })
+		);
 }
 
 
@@ -341,6 +348,7 @@ void FGameplayDebuggerCategory_Mover::FRepData::Serialize(FArchive& Ar)
 	Ar << MoveIntent;
 	Ar << ActiveLayeredMoves;
 	Ar << ActiveModifiers;
+	Ar << SyncStateDataTypes;
 	Ar << ModeMap;
 	Ar << ActiveTransitions;
 	Ar << MoveInputType;
diff --git a/Engine/Plugins/Experimental/Mover/Source/Mover/Private/Debug/GameplayDebuggerCategory_Mover.h b/Engine/Plugins/Experimental/Mover/Source/Mover/Private/Debug/GameplayDebuggerCategory_Mover.h
--- a/Engine/Plugins/Experimental/Mover/Source/Mover/Private/Debug/GameplayDebuggerCategory_Mover.h
+++ b/Engine/Plugins/Experimental/Mover/Source/Mover/Private/Debug/GameplayDebuggerCategory_Mover.h
@@ -37,6 +37,7 @@ protected:
 		TArray<FString> ModeMap;
 		TArray<FString> ActiveTransitions;
 		TArray<FString> ActiveModifiers;
+		TArray<FString> SyncStateDataTypes;
 
 		// Last input
 		int8 MoveInputType;
