Mover: Removing a setting NavAgentProps.NavWalkingSearchHeightScale to default nav values in UpdateNavComponent as it erases set values set by user

diff --git a/Engine/Plugins/Experimental/Mover/Source/Mover/Private/DefaultMovementSet/NavMoverComponent.cpp b/Engine/Plugins/Experimental/Mover/Source/Mover/Private/DefaultMovementSet/NavMoverComponent.cpp
--- a/Engine/Plugins/Experimental/Mover/Source/Mover/Private/DefaultMovementSet/NavMoverComponent.cpp
+++ b/Engine/Plugins/Experimental/Mover/Source/Mover/Private/DefaultMovementSet/NavMoverComponent.cpp
@@ -164,9 +164,6 @@ void UNavMoverComponent::UpdateNavAgent(const UObject& ObjectToUpdateFrom)
 	{
 		return;
 	}
-
-	// initialize properties from navigation system
-	NavAgentProps.NavWalkingSearchHeightScale = FNavigationSystem::GetDefaultSupportedAgent().NavWalkingSearchHeightScale;
 	
 	if (const UCapsuleComponent* CapsuleComponent = Cast<UCapsuleComponent>(&ObjectToUpdateFrom))
 	{
