[NetPhysics] Enable delay physics proxy deletion until out of rewind history.
- Also fix edge case crash with delaying the deletion where we would delete the physics proxy one frame early and then access it during a resim, due to the delete logic running after FrameNumber has been iterated already.
- Also check particle pointer before using it
- And stop log from printing in shipping build

diff --git a/Engine/Plugins/Experimental/Mover/Source/Mover/Private/Backends/MoverNetworkPhysicsLiaisonBase.cpp b/Engine/Plugins/Experimental/Mover/Source/Mover/Private/Backends/MoverNetworkPhysicsLiaisonBase.cpp
--- a/Engine/Plugins/Experimental/Mover/Source/Mover/Private/Backends/MoverNetworkPhysicsLiaisonBase.cpp
+++ b/Engine/Plugins/Experimental/Mover/Source/Mover/Private/Backends/MoverNetworkPhysicsLiaisonBase.cpp
@@ -794,16 +794,18 @@ void UMoverNetworkPhysicsLiaisonComponentBase::PerformPreSimulate_Internal(const
 
 	MoverComp.GetSimBlackboard()->TryGet(CommonBlackboard::LastFloorResult, Output.FloorResult);
 
-	Chaos::FPBDRigidParticleHandle* ParticleHandle = GetControlledParticle_Internal();
-	FVector TargetDeltaPos = OutputSyncState.GetLocation_WorldSpace() - ParticleHandle->GetX();
-	if (TargetDeltaPos.SizeSquared2D() > GPhysicsDrivenMotionDebugParams.TeleportThreshold * GPhysicsDrivenMotionDebugParams.TeleportThreshold)
+	if (Chaos::FPBDRigidParticleHandle* ParticleHandle = GetControlledParticle_Internal())
 	{
-		GetPhysicsSolver()->GetEvolution()->SetParticleTransform(ParticleHandle, OutputSyncState.GetLocation_WorldSpace(), OutputSyncState.GetOrientation_WorldSpace().Quaternion(), true);
-	}
+		FVector TargetDeltaPos = OutputSyncState.GetLocation_WorldSpace() - ParticleHandle->GetX();
+		if (TargetDeltaPos.SizeSquared2D() > GPhysicsDrivenMotionDebugParams.TeleportThreshold * GPhysicsDrivenMotionDebugParams.TeleportThreshold)
+		{
+			GetPhysicsSolver()->GetEvolution()->SetParticleTransform(ParticleHandle, OutputSyncState.GetLocation_WorldSpace(), OutputSyncState.GetOrientation_WorldSpace().Quaternion(), true);
+		}
 
-	ParticleHandle->SetV(OutputSyncState.GetVelocity_WorldSpace());
+		ParticleHandle->SetV(OutputSyncState.GetVelocity_WorldSpace());
 
-	//@todo DanH: Does the base need to concern itself with setting W on the particle?
+		//@todo DanH: Does the base need to concern itself with setting W on the particle?
+	}
 }
 
 Chaos::FPhysicsObject* UMoverNetworkPhysicsLiaisonComponentBase::GetControlledPhysicsObject() const
diff --git a/Engine/Source/Runtime/Experimental/Chaos/Private/PBDRigidsSolver.cpp b/Engine/Source/Runtime/Experimental/Chaos/Private/PBDRigidsSolver.cpp
--- a/Engine/Source/Runtime/Experimental/Chaos/Private/PBDRigidsSolver.cpp
+++ b/Engine/Source/Runtime/Experimental/Chaos/Private/PBDRigidsSolver.cpp
@@ -830,7 +830,7 @@ namespace Chaos
 		return Serializer;
 	}
 
-	bool bResimDelayDestroyObject = false;
+	bool bResimDelayDestroyObject = true;
 	FAutoConsoleVariableRef CVarResimDelayDestroyObject(TEXT("p.Resim.DelayDestroyObject"), bResimDelayDestroyObject, TEXT("Keep physics objects and physics proxy available until after they go out of rewind history. Disabled by default due to not fully implemented."));
 
 	void FPBDRigidsSolver::UnregisterObject(FSingleParticlePhysicsProxy* Proxy)
@@ -1271,7 +1271,7 @@ namespace Chaos
 		for(int32 Idx = PendingDestroyPhysicsProxy.Num() - 1; Idx >= 0; --Idx)
 		{
 			FPendingDestroyInfo& Info = PendingDestroyPhysicsProxy[Idx];
-			if(Info.DestroyOnStep <= GetCurrentFrame() || IsShuttingDown())
+			if(Info.DestroyOnStep <= (GetCurrentFrame() - 1) || IsShuttingDown()) // Note: GetCurrentFrame() - 1 because this happens at the end of the physics frame after frame number has been incremented already
 			{
 				// Callback for Unregister PhysicsObject
 				if (PhysicsObjectUnregistrationWatchers.Num() > 0)
