Mover: fix cases where turn generator options were not filtered properly on some movement modes

diff --git a/Engine/Plugins/Experimental/Mover/Source/Mover/Public/DefaultMovementSet/Modes/AsyncNavWalkingMode.h b/Engine/Plugins/Experimental/Mover/Source/Mover/Public/DefaultMovementSet/Modes/AsyncNavWalkingMode.h
--- a/Engine/Plugins/Experimental/Mover/Source/Mover/Public/DefaultMovementSet/Modes/AsyncNavWalkingMode.h
+++ b/Engine/Plugins/Experimental/Mover/Source/Mover/Public/DefaultMovementSet/Modes/AsyncNavWalkingMode.h
@@ -107,7 +107,7 @@ protected:
 	uint8 bProjectNavMeshOnBothWorldChannels : 1;
 
 	/** Optional modular object for generating rotation towards desired orientation. If not specified, linear interpolation will be used. */
-	UPROPERTY(EditAnywhere, Instanced, Category = Mover, meta = (MustImplement = "/Script/Mover.TurnGeneratorInterface"))
+	UPROPERTY(EditAnywhere, Instanced, Category = Mover, meta = (ObjectMustImplement = "/Script/Mover.TurnGeneratorInterface"))
 	TObjectPtr<UObject> TurnGenerator;
 	
 	/** Get Navigation data for the actor. Returns null if there is no associated nav data. */
diff --git a/Engine/Plugins/Experimental/Mover/Source/Mover/Public/DefaultMovementSet/Modes/AsyncWalkingMode.h b/Engine/Plugins/Experimental/Mover/Source/Mover/Public/DefaultMovementSet/Modes/AsyncWalkingMode.h
--- a/Engine/Plugins/Experimental/Mover/Source/Mover/Public/DefaultMovementSet/Modes/AsyncWalkingMode.h
+++ b/Engine/Plugins/Experimental/Mover/Source/Mover/Public/DefaultMovementSet/Modes/AsyncWalkingMode.h
@@ -44,7 +44,7 @@ protected:
 	EStaticFloorCheckPolicy FloorCheckPolicy = EStaticFloorCheckPolicy::OnDynamicBaseOnly;
 
 	/** Optional modular object for generating rotation towards desired orientation. If not specified, linear interpolation will be used. */
-	UPROPERTY(EditAnywhere, Instanced, Category=Mover, meta=(MustImplement="/Script/Mover.TurnGeneratorInterface"))
+	UPROPERTY(EditAnywhere, Instanced, Category=Mover, meta=(ObjectMustImplement="/Script/Mover.TurnGeneratorInterface"))
 	TObjectPtr<UObject> TurnGenerator;
 
 	virtual void OnRegistered(const FName ModeName) override; 
