Mover: fixing potential deadlock in standalone liaison when uncoordinated movement of a character could occur during application of simulation state.

diff --git a/Engine/Plugins/Experimental/Mover/Source/Mover/Private/Backends/MoverStandaloneLiaison.cpp b/Engine/Plugins/Experimental/Mover/Source/Mover/Private/Backends/MoverStandaloneLiaison.cpp
--- a/Engine/Plugins/Experimental/Mover/Source/Mover/Private/Backends/MoverStandaloneLiaison.cpp
+++ b/Engine/Plugins/Experimental/Mover/Source/Mover/Private/Backends/MoverStandaloneLiaison.cpp
@@ -107,6 +107,12 @@ bool UMoverStandaloneLiaisonComponent::ReadPendingSyncState(OUT FMoverSyncState&
 
 bool UMoverStandaloneLiaisonComponent::WritePendingSyncState(const FMoverSyncState& SyncStateToWrite)
 {
+	if (bIsInApplySimulationState)
+	{
+		UE_LOG(LogMover, Log, TEXT("Attempted to WritePendingSyncState during ApplySimulationState. Skipping. (%s)"), *GetNameSafe(GetOwner()));
+		return false;	// can't write to the sync state while it's being applied. Any changes from this will be overwritten anyway.
+	}
+
 	{
 		FRWScopeLock Lock(StateDataLock, SLT_Write);
 		CachedLastSyncState = SyncStateToWrite;
@@ -370,6 +376,8 @@ void UMoverStandaloneLiaisonComponent::TickMovementSimulation(float DeltaSeconds
 
 void UMoverStandaloneLiaisonComponent::TickApplySimulationState(float DeltaSeconds)
 {
+	TGuardValue<bool> InApplySimulationState(bIsInApplySimulationState, true);
+
 	{
 		SCOPED_NAMED_EVENT(MoverBackendLiaison_Tick_FinalizeFrame, FColor::Green);
 
diff --git a/Engine/Plugins/Experimental/Mover/Source/Mover/Public/Backends/MoverStandaloneLiaison.h b/Engine/Plugins/Experimental/Mover/Source/Mover/Public/Backends/MoverStandaloneLiaison.h
--- a/Engine/Plugins/Experimental/Mover/Source/Mover/Public/Backends/MoverStandaloneLiaison.h
+++ b/Engine/Plugins/Experimental/Mover/Source/Mover/Public/Backends/MoverStandaloneLiaison.h
@@ -183,6 +183,8 @@ protected:
 
 	FRWLock StateDataLock;	// used when reading/writing to our cached state (sync/aux) data
 
+	bool bIsInApplySimulationState = false;	// transient flag indicating ApplySimulationState is active
+
 	FMoverStandaloneProduceInputTickFunction ProduceInputTickFunction;
 	FMoverStandaloneSimulateMovementTickFunction SimulateMovementTickFunction;
 	FMoverStandaloneApplyStateTickFunction ApplyStateTickFunction;
