Mover: improving copy implementation of data collections and sync state to avoid memory reallocation when possible.

Also fixes a handful of struct types that needed a MakeCopy ops identifier due to having native members that aren't uproperties

diff --git a/Engine/Plugins/Experimental/Mover/Source/Mover/Private/LayeredMove.cpp b/Engine/Plugins/Experimental/Mover/Source/Mover/Private/LayeredMove.cpp
--- a/Engine/Plugins/Experimental/Mover/Source/Mover/Private/LayeredMove.cpp
+++ b/Engine/Plugins/Experimental/Mover/Source/Mover/Private/LayeredMove.cpp
@@ -4,6 +4,7 @@
 #include "MoverLog.h"
 #include "MoverComponent.h"
 #include "MoverSimulationTypes.h"
+#include "MoverModule.h"
 
 const float LayeredMove_InvalidTime = -UE_BIG_NUMBER;
 
@@ -162,41 +163,70 @@ void FLayeredMoveGroup::NetSerialize(FArchive& Ar, uint8 MaxNumMovesToSerialize/
 }
 
 
-FLayeredMoveGroup& FLayeredMoveGroup::operator=(const FLayeredMoveGroup& Other)
+static void CopyLayeredMoveArray(TArray<TSharedPtr<FLayeredMoveBase>>& Dest, const TArray<TSharedPtr<FLayeredMoveBase>>& Src)
 {
-	// Perform deep copy of this Group
-	if (this != &Other)
+	bool bCanCopyInPlace = (UE::Mover::DisableDataCopyInPlace == 0);
+
+	if (bCanCopyInPlace)
 	{
-		// Deep copy active moves
-		ActiveLayeredMoves.Empty(Other.ActiveLayeredMoves.Num());
-		for (int i = 0; i < Other.ActiveLayeredMoves.Num(); ++i)
+		if (Dest.Num() != Src.Num())
 		{
-			if (Other.ActiveLayeredMoves[i].IsValid())
-			{
-				FLayeredMoveBase* CopyOfSourcePtr = Other.ActiveLayeredMoves[i]->Clone();
-				ActiveLayeredMoves.Add(TSharedPtr<FLayeredMoveBase>(CopyOfSourcePtr));
-			}
-			else
+			bCanCopyInPlace = false;
+		}
+		else
+		{
+			for (int32 i = 0; i < Dest.Num(); ++i)
 			{
-				UE_LOG(LogMover, Warning, TEXT("FLayeredMoveGroup::operator= trying to copy invalid Other Layered Move"));
+				if ((!Dest[i].IsValid() || !Src[i].IsValid()) ||
+					(Dest[i]->GetScriptStruct() != Src[i]->GetScriptStruct()))
+				{
+					bCanCopyInPlace = false;
+					break;
+				}
 			}
 		}
+	}
 
-		// Deep copy queued moves
-		QueuedLayeredMoves.Empty(Other.QueuedLayeredMoves.Num());
-		for (int i = 0; i < Other.QueuedLayeredMoves.Num(); ++i)
+	if (bCanCopyInPlace)
+	{
+		// Copy all structs in place
+		for (int32 i = 0; i < Dest.Num(); ++i)
 		{
-			if (Other.QueuedLayeredMoves[i].IsValid())
+			UScriptStruct* ScriptStruct = Dest[i]->GetScriptStruct();
+			FLayeredMoveBase* DestMove = Dest[i].Get();
+			FLayeredMoveBase* SrcMove = Src[i].Get();
+
+			ScriptStruct->CopyScriptStruct(DestMove, SrcMove, 1);
+		}
+	}
+	else
+	{
+		// Deep copy active moves
+		Dest.Empty(Src.Num());
+		for (int i = 0; i < Src.Num(); ++i)
+		{
+			if (Src[i].IsValid())
 			{
-				FLayeredMoveBase* CopyOfSourcePtr = Other.QueuedLayeredMoves[i]->Clone();
-				QueuedLayeredMoves.Add(TSharedPtr<FLayeredMoveBase>(CopyOfSourcePtr));
+				FLayeredMoveBase* CopyOfSourcePtr = Src[i]->Clone();
+				Dest.Add(TSharedPtr<FLayeredMoveBase>(CopyOfSourcePtr));
 			}
 			else
 			{
-				UE_LOG(LogMover, Warning, TEXT("FLayeredMoveGroup::operator= trying to copy invalid Other Layered Move"));
+				UE_LOG(LogMover, Warning, TEXT("CopyLayeredMoveArray trying to copy invalid Other Layered Move"));
 			}
 		}
 	}
+}
+
+
+FLayeredMoveGroup& FLayeredMoveGroup::operator=(const FLayeredMoveGroup& Other)
+{
+	// Perform deep copy of this Group
+	if (this != &Other)
+	{
+		CopyLayeredMoveArray(ActiveLayeredMoves, Other.ActiveLayeredMoves);
+		CopyLayeredMoveArray(QueuedLayeredMoves, Other.QueuedLayeredMoves);
+	}
 
 	return *this;
 }
diff --git a/Engine/Plugins/Experimental/Mover/Source/Mover/Private/MovementModifier.cpp b/Engine/Plugins/Experimental/Mover/Source/Mover/Private/MovementModifier.cpp
--- a/Engine/Plugins/Experimental/Mover/Source/Mover/Private/MovementModifier.cpp
+++ b/Engine/Plugins/Experimental/Mover/Source/Mover/Private/MovementModifier.cpp
@@ -2,6 +2,7 @@
 
 #include "MovementModifier.h"
 #include "MoverComponent.h"
+#include "MoverModule.h"
 
 const float MovementModifier_InvalidTime = -UE_BIG_NUMBER;
 
@@ -135,42 +136,73 @@ TArray<TSharedPtr<FMovementModifierBase>> FMovementModifierGroup::GenerateActive
 	return ActiveModifiers;
 }
 
-FMovementModifierGroup& FMovementModifierGroup::operator=(const FMovementModifierGroup& Other)
+static void CopyModifierArray(TArray<TSharedPtr<FMovementModifierBase>>& Dest, const TArray<TSharedPtr<FMovementModifierBase>>& Src)
 {
-	// Perform deep copy of this Group
-	if (this != &Other)
+	// If all of our modifiers types already match, we can copy in place without reallocating memory
+	bool bCanCopyInPlace = (UE::Mover::DisableDataCopyInPlace == 0);
+
+	if (bCanCopyInPlace)
 	{
-		// Deep copy active moves
-		ActiveModifiers.Empty(Other.ActiveModifiers.Num());
-		for (int i = 0; i < Other.ActiveModifiers.Num(); ++i)
+		if (Dest.Num() != Src.Num())
 		{
-			if (Other.ActiveModifiers[i].IsValid())
-			{
-				FMovementModifierBase* CopyOfSourcePtr = Other.ActiveModifiers[i]->Clone();
-				ActiveModifiers.Add(TSharedPtr<FMovementModifierBase>(CopyOfSourcePtr));
-			}
-			else
+			bCanCopyInPlace = false;
+		}
+		else
+		{
+			for (int32 i = 0; i < Dest.Num(); ++i)
 			{
-				UE_LOG(LogMover, Warning, TEXT("FMovementModifierGroup::operator= trying to copy invalid Other Modifier"));
+				if ((!Dest[i].IsValid() || !Src[i].IsValid()) ||
+					(Dest[i]->GetScriptStruct() != Src[i]->GetScriptStruct()))
+				{
+					bCanCopyInPlace = false;
+					break;
+				}
 			}
 		}
+	}
+
+	if (bCanCopyInPlace)
+	{
+		// Copy all structs in place
+		for (int32 i = 0; i < Dest.Num(); ++i)
+		{
+			UScriptStruct* ScriptStruct = Dest[i]->GetScriptStruct();
+			FMovementModifierBase* DestMod = Dest[i].Get();
+			FMovementModifierBase* SrcMod = Src[i].Get();
 
-		// Deep copy queued moves
-		QueuedModifiers.Empty(Other.QueuedModifiers.Num());
-		for (int i = 0; i < Other.QueuedModifiers.Num(); ++i)
+			ScriptStruct->CopyScriptStruct(DestMod, SrcMod, 1);
+		}
+	}
+	else
+	{
+		// Deep copy active modifiers
+		Dest.Empty(Src.Num());
+		for (int i = 0; i < Src.Num(); ++i)
 		{
-			if (Other.QueuedModifiers[i].IsValid())
+			if (Src[i].IsValid())
 			{
-				FMovementModifierBase* CopyOfSourcePtr = Other.QueuedModifiers[i]->Clone();
-				QueuedModifiers.Add(TSharedPtr<FMovementModifierBase>(CopyOfSourcePtr));
+				FMovementModifierBase* CopyOfSourcePtr = Src[i]->Clone();
+				Dest.Add(TSharedPtr<FMovementModifierBase>(CopyOfSourcePtr));
 			}
 			else
 			{
-				UE_LOG(LogMover, Warning, TEXT("FMovementModifierGroup::operator= trying to copy invalid Other Modifier"));
+				UE_LOG(LogMover, Warning, TEXT("CopyModifierArray trying to copy invalid Other Modifier"));
 			}
 		}
 	}
 
+}
+
+
+FMovementModifierGroup& FMovementModifierGroup::operator=(const FMovementModifierGroup& Other)
+{
+	// Perform deep copy of this Group
+	if (this != &Other)
+	{
+		CopyModifierArray(ActiveModifiers, Other.ActiveModifiers);
+		CopyModifierArray(QueuedModifiers, Other.QueuedModifiers);
+	}
+
 	return *this;
 }
 
diff --git a/Engine/Plugins/Experimental/Mover/Source/Mover/Private/MoverModule.cpp b/Engine/Plugins/Experimental/Mover/Source/Mover/Private/MoverModule.cpp
--- a/Engine/Plugins/Experimental/Mover/Source/Mover/Private/MoverModule.cpp
+++ b/Engine/Plugins/Experimental/Mover/Source/Mover/Private/MoverModule.cpp
@@ -16,6 +16,18 @@
 
 #define LOCTEXT_NAMESPACE "FMoverModule"
 
+
+namespace UE::Mover
+{
+	int32 DisableDataCopyInPlace = 0;
+	static FAutoConsoleVariableRef CVarDisableDataCopyInPlace(
+		TEXT("mover.debug.DisableDataCopyInPlace"),
+		DisableDataCopyInPlace,
+		TEXT("Whether to allow Mover data collections with identical contained struct types to be copied in place, avoiding reallocating memory"),
+		ECVF_Default);
+}
+
+
 void FMoverModule::StartupModule()
 {
 	// This code will execute after your module is loaded into memory; the exact timing is specified in the .uplugin file per-module
diff --git a/Engine/Plugins/Experimental/Mover/Source/Mover/Private/MoverTypes.cpp b/Engine/Plugins/Experimental/Mover/Source/Mover/Private/MoverTypes.cpp
--- a/Engine/Plugins/Experimental/Mover/Source/Mover/Private/MoverTypes.cpp
+++ b/Engine/Plugins/Experimental/Mover/Source/Mover/Private/MoverTypes.cpp
@@ -4,7 +4,9 @@
 #include "MoverTypes.h"
 #include "Blueprint/BlueprintExceptionInfo.h"
 #include "MoverLog.h"
+#include "MoverModule.h"
 #include "UObject/ObjectKey.h"
+#include "HAL/IConsoleManager.h"
 
 #define LOCTEXT_NAMESPACE "MoverData"
 
@@ -152,18 +154,56 @@ FMoverDataCollection& FMoverDataCollection::operator=(const FMoverDataCollection
 	// Perform deep copy of this Group
 	if (this != &Other)
 	{
-		// Deep copy active data blocks
-		DataArray.Empty(Other.DataArray.Num());
-		for (int i = 0; i < Other.DataArray.Num(); ++i)
+		// If all of our data collection types already match, we can copy in place without reallocating memory
+		bool bCanCopyInPlace = (UE::Mover::DisableDataCopyInPlace == 0);
+		
+		if (bCanCopyInPlace)
 		{
-			if (Other.DataArray[i].IsValid())
+			if (DataArray.Num() != Other.DataArray.Num())
 			{
-				FMoverDataStructBase* CopyOfSourcePtr = Other.DataArray[i]->Clone();
-				DataArray.Add(TSharedPtr<FMoverDataStructBase>(CopyOfSourcePtr));
+				bCanCopyInPlace = false;
 			}
 			else
 			{
-				UE_LOG(LogMover, Warning, TEXT("FMoverDataCollection::operator= trying to copy invalid Other DataArray element"));
+				for (int32 i = 0; i < DataArray.Num(); ++i)
+				{
+					if ((!DataArray[i].IsValid() || !Other.DataArray[i].IsValid()) ||
+						(DataArray[i]->GetScriptStruct() != Other.DataArray[i]->GetScriptStruct()))
+					{
+						bCanCopyInPlace = false;
+						break;
+					}
+				}
+			}
+		}
+
+		if (bCanCopyInPlace)
+		{
+			// Copy all structs in place
+			for (int32 i = 0; i < DataArray.Num(); ++i)
+			{
+				UScriptStruct* ScriptStruct = DataArray[i]->GetScriptStruct();
+				FMoverDataStructBase* Dest = DataArray[i].Get();
+				FMoverDataStructBase* Src = Other.DataArray[i].Get();
+
+				ScriptStruct->CopyScriptStruct(Dest, Src, 1);
+			}
+		}
+		else
+		{
+			// Deep copy active data blocks
+			DataArray.Empty(Other.DataArray.Num());
+			for (int i = 0; i < Other.DataArray.Num(); ++i)
+			{
+				if (Other.DataArray[i].IsValid())
+				{
+					FMoverDataStructBase* CopyOfSourcePtr = Other.DataArray[i]->Clone();
+					DataArray.Add(TSharedPtr<FMoverDataStructBase>(CopyOfSourcePtr));
+				}
+				else
+				{
+					UE_LOG(LogMover, Warning, TEXT("FMoverDataCollection::operator= trying to copy invalid Other DataArray element"));
+				}
 			}
 		}
 	}
diff --git a/Engine/Plugins/Experimental/Mover/Source/Mover/Public/DefaultMovementSet/LayeredMoves/AnimRootMotionLayeredMove.h b/Engine/Plugins/Experimental/Mover/Source/Mover/Public/DefaultMovementSet/LayeredMoves/AnimRootMotionLayeredMove.h
--- a/Engine/Plugins/Experimental/Mover/Source/Mover/Public/DefaultMovementSet/LayeredMoves/AnimRootMotionLayeredMove.h
+++ b/Engine/Plugins/Experimental/Mover/Source/Mover/Public/DefaultMovementSet/LayeredMoves/AnimRootMotionLayeredMove.h
@@ -45,3 +45,12 @@ struct MOVER_API FLayeredMove_AnimRootMotion : public FLayeredMoveBase
 
 	virtual void AddReferencedObjects(class FReferenceCollector& Collector) override;
 };
+
+template<>
+struct TStructOpsTypeTraits< FLayeredMove_AnimRootMotion > : public TStructOpsTypeTraitsBase2< FLayeredMove_AnimRootMotion >
+{
+	enum
+	{
+		WithCopy = true
+	};
+};
\ No newline at end of file
diff --git a/Engine/Plugins/Experimental/Mover/Source/Mover/Public/DefaultMovementSet/LayeredMoves/LaunchMove.h b/Engine/Plugins/Experimental/Mover/Source/Mover/Public/DefaultMovementSet/LayeredMoves/LaunchMove.h
--- a/Engine/Plugins/Experimental/Mover/Source/Mover/Public/DefaultMovementSet/LayeredMoves/LaunchMove.h
+++ b/Engine/Plugins/Experimental/Mover/Source/Mover/Public/DefaultMovementSet/LayeredMoves/LaunchMove.h
@@ -36,3 +36,13 @@ struct MOVER_API FLayeredMove_Launch : public FLayeredMoveBase
 
 	virtual void AddReferencedObjects(class FReferenceCollector& Collector) override;
 };
+
+
+template<>
+struct TStructOpsTypeTraits< FLayeredMove_Launch > : public TStructOpsTypeTraitsBase2< FLayeredMove_Launch >
+{
+	enum
+	{
+		WithCopy = true
+	};
+};
\ No newline at end of file
diff --git a/Engine/Plugins/Experimental/Mover/Source/Mover/Public/DefaultMovementSet/LayeredMoves/RootMotionAttributeLayeredMove.h b/Engine/Plugins/Experimental/Mover/Source/Mover/Public/DefaultMovementSet/LayeredMoves/RootMotionAttributeLayeredMove.h
--- a/Engine/Plugins/Experimental/Mover/Source/Mover/Public/DefaultMovementSet/LayeredMoves/RootMotionAttributeLayeredMove.h
+++ b/Engine/Plugins/Experimental/Mover/Source/Mover/Public/DefaultMovementSet/LayeredMoves/RootMotionAttributeLayeredMove.h
@@ -41,3 +41,12 @@ protected:
 
 	virtual void AddReferencedObjects(class FReferenceCollector& Collector) override;
 };
+
+template<>
+struct TStructOpsTypeTraits< FLayeredMove_RootMotionAttribute > : public TStructOpsTypeTraitsBase2< FLayeredMove_RootMotionAttribute >
+{
+	enum
+	{
+		WithCopy = true
+	};
+};
\ No newline at end of file
diff --git a/Engine/Plugins/Experimental/Mover/Source/Mover/Public/DefaultMovementSet/MovementModifiers/StanceModifier.h b/Engine/Plugins/Experimental/Mover/Source/Mover/Public/DefaultMovementSet/MovementModifiers/StanceModifier.h
--- a/Engine/Plugins/Experimental/Mover/Source/Mover/Public/DefaultMovementSet/MovementModifiers/StanceModifier.h
+++ b/Engine/Plugins/Experimental/Mover/Source/Mover/Public/DefaultMovementSet/MovementModifiers/StanceModifier.h
@@ -75,3 +75,12 @@ protected:
 	// Reverts any movement settings like acceleration or max speed changes
 	void RevertMovementSettings(UMoverComponent* MoverComp);
 };
+
+template<>
+struct TStructOpsTypeTraits< FStanceModifier > : public TStructOpsTypeTraitsBase2< FStanceModifier >
+{
+	enum
+	{
+		WithCopy = true
+	};
+};
\ No newline at end of file
diff --git a/Engine/Plugins/Experimental/Mover/Source/Mover/Public/MoverModule.h b/Engine/Plugins/Experimental/Mover/Source/Mover/Public/MoverModule.h
--- a/Engine/Plugins/Experimental/Mover/Source/Mover/Public/MoverModule.h
+++ b/Engine/Plugins/Experimental/Mover/Source/Mover/Public/MoverModule.h
@@ -8,6 +8,11 @@
 struct IConsoleCommand;
 class UWorld;
 
+namespace UE::Mover
+{
+	extern int32 DisableDataCopyInPlace;
+}
+
 class FMoverModule : public IModuleInterface
 {
 public:
diff --git a/Engine/Plugins/Experimental/Mover/Source/Mover/Public/MoverTypes.h b/Engine/Plugins/Experimental/Mover/Source/Mover/Public/MoverTypes.h
--- a/Engine/Plugins/Experimental/Mover/Source/Mover/Public/MoverTypes.h
+++ b/Engine/Plugins/Experimental/Mover/Source/Mover/Public/MoverTypes.h
@@ -328,7 +328,7 @@ struct TStructOpsTypeTraits<FMoverDataCollection> : public TStructOpsTypeTraitsB
 {
 	enum
 	{
-		WithCopy = true,		// Necessary so that TSharedPtr<FMoverDataStructBase> Data is copied around
+		WithCopy = true,		// Necessary so that DataArray is copied around
 		WithNetSerializer = true,
 		WithIdenticalViaEquality = true,
 		WithAddStructReferencedObjects = true,
diff --git a/Engine/Plugins/Experimental/MoverExamples/Source/MoverExamples/Public/CharacterVariants/Ziplining/ZipliningMode.h b/Engine/Plugins/Experimental/MoverExamples/Source/MoverExamples/Public/CharacterVariants/Ziplining/ZipliningMode.h
--- a/Engine/Plugins/Experimental/MoverExamples/Source/MoverExamples/Public/CharacterVariants/Ziplining/ZipliningMode.h
+++ b/Engine/Plugins/Experimental/MoverExamples/Source/MoverExamples/Public/CharacterVariants/Ziplining/ZipliningMode.h
@@ -50,3 +50,12 @@ struct FZipliningState : public FMoverDataStructBase
 	virtual bool ShouldReconcile(const FMoverDataStructBase& AuthorityState) const override;
 	virtual void Interpolate(const FMoverDataStructBase& From, const FMoverDataStructBase& To, float Pct) override;
 };
+
+template<>
+struct TStructOpsTypeTraits< FZipliningState > : public TStructOpsTypeTraitsBase2< FZipliningState >
+{
+	enum
+	{
+		WithCopy = true				// Needed for copying non-uproperty members
+	};
+};
\ No newline at end of file
diff --git a/Engine/Plugins/Experimental/MoverExamples/Source/MoverExamples/Public/MovementBases/FollowPathMode.h b/Engine/Plugins/Experimental/MoverExamples/Source/MoverExamples/Public/MovementBases/FollowPathMode.h
--- a/Engine/Plugins/Experimental/MoverExamples/Source/MoverExamples/Public/MovementBases/FollowPathMode.h
+++ b/Engine/Plugins/Experimental/MoverExamples/Source/MoverExamples/Public/MovementBases/FollowPathMode.h
@@ -157,3 +157,12 @@ struct FFollowPathState : public FMoverDataStructBase
 	}
 
 };
+
+template<>
+struct TStructOpsTypeTraits<FFollowPathState> : public TStructOpsTypeTraitsBase2<FFollowPathState>
+{
+	enum
+	{
+		WithCopy = true		// Necessary so that non-uproperty members get copied
+	};
+};
\ No newline at end of file
diff --git a/Engine/Plugins/Experimental/MoverExamples/Source/MoverExamples/Public/MovementBases/FollowSplineMode.h b/Engine/Plugins/Experimental/MoverExamples/Source/MoverExamples/Public/MovementBases/FollowSplineMode.h
--- a/Engine/Plugins/Experimental/MoverExamples/Source/MoverExamples/Public/MovementBases/FollowSplineMode.h
+++ b/Engine/Plugins/Experimental/MoverExamples/Source/MoverExamples/Public/MovementBases/FollowSplineMode.h
@@ -192,3 +192,12 @@ struct FFollowSplineState : public FMoverDataStructBase
 		}
 	}
 };
+
+template<>
+struct TStructOpsTypeTraits<FFollowSplineState> : public TStructOpsTypeTraitsBase2<FFollowSplineState>
+{
+	enum
+	{
+		WithCopy = true		// Necessary so that non-uproperty members get copied
+	};
+};
