Mover: misc memory-related improvements, avoiding reallocations by reusing intermediate structs between frames, as their types rarely change

diff --git a/Engine/Plugins/Experimental/Mover/Source/Mover/Private/Backends/MoverStandaloneLiaison.cpp b/Engine/Plugins/Experimental/Mover/Source/Mover/Private/Backends/MoverStandaloneLiaison.cpp
--- a/Engine/Plugins/Experimental/Mover/Source/Mover/Private/Backends/MoverStandaloneLiaison.cpp
+++ b/Engine/Plugins/Experimental/Mover/Source/Mover/Private/Backends/MoverStandaloneLiaison.cpp
@@ -340,6 +340,8 @@ void UMoverStandaloneLiaisonComponent::TickInputProduction(float DeltaSeconds)
 
 void UMoverStandaloneLiaisonComponent::TickMovementSimulation(float DeltaSeconds)
 {
+	SCOPED_NAMED_EVENT(StandaloneMoverBackendLiaison_Tick_MovementSimulation, FColor::Blue);
+
 	if (!GetEnableProduceInput())
 	{
 		// If input production is turned off, we need to update current time ourselves
@@ -351,27 +353,25 @@ void UMoverStandaloneLiaisonComponent::TickMovementSimulation(float DeltaSeconds
 	TimeStep.BaseSimTimeMs = CurrentSimTimeMs;
 	TimeStep.StepMs = FMath::Floor(DeltaSeconds * 1000.f);
 
-	FMoverTickStartData StartData;
-	FMoverTickEndData EndData;
-
-	StartData.InputCmd = LastProducedInputCmd;
+	WorkingStartData.InputCmd = LastProducedInputCmd;
 
 	{
 		FRWScopeLock Lock(StateDataLock, SLT_ReadOnly);
-		StartData.SyncState = CachedLastSyncState;
-		StartData.AuxState = CachedLastAuxState;
+		WorkingStartData.SyncState = CachedLastSyncState;
+		WorkingStartData.AuxState = CachedLastAuxState;
 	}
 
 	{
-		SCOPED_NAMED_EVENT(StandaloneMoverBackendLiaison_Tick_SimulationTick, FColor::Blue);
-		MoverComp->SimulationTick(TimeStep, StartData, OUT EndData);
+		SCOPED_NAMED_EVENT(StandaloneMoverBackendLiaison_Tick_SimulationOnly, FColor::Blue);
+		MoverComp->SimulationTick(TimeStep, WorkingStartData, OUT WorkingEndData);
 	}
 
 	{
 		FRWScopeLock Lock(StateDataLock, SLT_Write);
-		CachedLastSyncState = EndData.SyncState;
-		CachedLastAuxState = EndData.AuxState;
+		CachedLastSyncState = WorkingEndData.SyncState;
+		CachedLastAuxState = WorkingEndData.AuxState;
 	}
+
 }
 
 void UMoverStandaloneLiaisonComponent::TickApplySimulationState(float DeltaSeconds)
@@ -379,7 +379,7 @@ void UMoverStandaloneLiaisonComponent::TickApplySimulationState(float DeltaSecon
 	TGuardValue<bool> InApplySimulationState(bIsInApplySimulationState, true);
 
 	{
-		SCOPED_NAMED_EVENT(MoverBackendLiaison_Tick_FinalizeFrame, FColor::Green);
+		SCOPED_NAMED_EVENT(StandaloneMoverBackendLiaison_Tick_FinalizeFrame, FColor::Green);
 
 		FRWScopeLock Lock(StateDataLock, SLT_ReadOnly);
 		MoverComp->FinalizeFrame(&CachedLastSyncState, &CachedLastAuxState);
diff --git a/Engine/Plugins/Experimental/Mover/Source/Mover/Private/MoveLibrary/MovementUtilsTypes.cpp b/Engine/Plugins/Experimental/Mover/Source/Mover/Private/MoveLibrary/MovementUtilsTypes.cpp
--- a/Engine/Plugins/Experimental/Mover/Source/Mover/Private/MoveLibrary/MovementUtilsTypes.cpp
+++ b/Engine/Plugins/Experimental/Mover/Source/Mover/Private/MoveLibrary/MovementUtilsTypes.cpp
@@ -6,7 +6,7 @@
 
 
 
-FMovingComponentSet::FMovingComponentSet(USceneComponent* InUpdatedComponent)
+void FMovingComponentSet::SetFrom(USceneComponent* InUpdatedComponent)
 {
 	UpdatedComponent = InUpdatedComponent;
 
@@ -19,7 +19,7 @@ FMovingComponentSet::FMovingComponentSet(USceneComponent* InUpdatedComponent)
 	}
 }
 
-FMovingComponentSet::FMovingComponentSet(UMoverComponent* InMoverComponent)
+void FMovingComponentSet::SetFrom(UMoverComponent* InMoverComponent)
 {
 	MoverComponent = InMoverComponent;
 
diff --git a/Engine/Plugins/Experimental/Mover/Source/Mover/Private/MovementModeStateMachine.cpp b/Engine/Plugins/Experimental/Mover/Source/Mover/Private/MovementModeStateMachine.cpp
--- a/Engine/Plugins/Experimental/Mover/Source/Mover/Private/MovementModeStateMachine.cpp
+++ b/Engine/Plugins/Experimental/Mover/Source/Mover/Private/MovementModeStateMachine.cpp
@@ -159,15 +159,14 @@ void UMovementModeStateMachine::OnSimulationTick(USceneComponent* UpdatedCompone
 	//GEngine->AddOnScreenDebugMessage(-1, -0.1f, FColor::White, FString::Printf(TEXT("Mode Tick: %s  Queued: %s"), *CurrentModeName.ToString(), *NextModeName.ToString()));
 
 	FMoverTimeStep SubTimeStep = TimeStep;
-	FMoverTickStartData SubstepStartData = StartState;
-	FMoverDefaultSyncState* SubstepStartSyncState = SubstepStartData.SyncState.SyncStateCollection.FindMutableDataByType<FMoverDefaultSyncState>();
+	WorkingSubstepStartData = StartState;
 
 	UMoverComponent* MoverComp = CastChecked<UMoverComponent>(GetOuter());
 	check(MoverComp->MovementMixer);
 
 	if (!QueuedModeTransition->IsSet())
 	{
-		QueueNextMode(SubstepStartData.SyncState.MovementMode);
+		QueueNextMode(WorkingSubstepStartData.SyncState.MovementMode);
 	}
 
 	AdvanceToNextMode();
@@ -177,7 +176,6 @@ void UMovementModeStateMachine::OnSimulationTick(USceneComponent* UpdatedCompone
 	int32 NumConsecutiveFullRefundedSubsteps = 0;
 	while (SubTimeStep.BaseSimTimeMs < EndingSimTimeMs)
 	{
-		SubstepStartSyncState = SubstepStartData.SyncState.SyncStateCollection.FindMutableDataByType<FMoverDefaultSyncState>();
 		FMoverDefaultSyncState* OutputSyncState = &OutputState.SyncState.SyncStateCollection.FindOrAddMutableDataByType<FMoverDefaultSyncState>();
 		OutputState.SyncState.MovementMode = CurrentModeName;
 
@@ -186,15 +184,15 @@ void UMovementModeStateMachine::OnSimulationTick(USceneComponent* UpdatedCompone
 		SubTimeStep.StepMs = EndingSimTimeMs - SubTimeStep.BaseSimTimeMs;		// TODO: convert this to an overridable function that can support MaxStepTime, MaxIterations, etc.
 
 		// Transfer any queued moves into the starting state. They'll be started during the move generation.
-		FlushQueuedMovesToGroup(SubstepStartData.SyncState.LayeredMoves);
-		OutputState.SyncState.LayeredMoves = SubstepStartData.SyncState.LayeredMoves;
+		FlushQueuedMovesToGroup(WorkingSubstepStartData.SyncState.LayeredMoves);
+		OutputState.SyncState.LayeredMoves = WorkingSubstepStartData.SyncState.LayeredMoves;
 
-		FlushQueuedModifiersToGroup(SubstepStartData.SyncState.MovementModifiers);
-		OutputState.SyncState.MovementModifiers = SubstepStartData.SyncState.MovementModifiers;
+		FlushQueuedModifiersToGroup(WorkingSubstepStartData.SyncState.MovementModifiers);
+		OutputState.SyncState.MovementModifiers = WorkingSubstepStartData.SyncState.MovementModifiers;
 		
 		FApplyMovementEffectParams EffectParams;
 		EffectParams.MoverComp = MoverComp;
-		EffectParams.StartState = &SubstepStartData;
+		EffectParams.StartState = &WorkingSubstepStartData;
 		EffectParams.TimeStep = &SubTimeStep;
 		EffectParams.UpdatedComponent = UpdatedComponent;
 		EffectParams.UpdatedPrimitive = UpdatedPrimitive;
@@ -208,7 +206,7 @@ void UMovementModeStateMachine::OnSimulationTick(USceneComponent* UpdatedCompone
 			{
 				if (SyncDataIt->Get())
 				{
-					SubstepStartData.SyncState.SyncStateCollection.AddOrOverwriteData(TSharedPtr<FMoverDataStructBase>(SyncDataIt->Get()->Clone()));
+					WorkingSubstepStartData.SyncState.SyncStateCollection.AddOrOverwriteData(TSharedPtr<FMoverDataStructBase>(SyncDataIt->Get()->Clone()));
 				}
 			}
 
@@ -216,13 +214,13 @@ void UMovementModeStateMachine::OnSimulationTick(USceneComponent* UpdatedCompone
 			{
 				bModeSetFromInstantEffect = true;
 				SetModeImmediately(OutputState.SyncState.MovementMode);
-				SubstepStartData.SyncState.MovementMode = CurrentModeName;
+				WorkingSubstepStartData.SyncState.MovementMode = CurrentModeName;
 			}
 		}
 
 		FMovementModifierGroup& CurrentModifiers = OutputState.SyncState.MovementModifiers;
 		FlushModifierCancellationsToGroup(CurrentModifiers);
-		TArray<TSharedPtr<FMovementModifierBase>> ActiveModifiers = CurrentModifiers.GenerateActiveModifiers(MoverComp, SubTimeStep, SubstepStartData.SyncState, SubstepStartData.AuxState);
+		TArray<TSharedPtr<FMovementModifierBase>> ActiveModifiers = CurrentModifiers.GenerateActiveModifiers(MoverComp, SubTimeStep, WorkingSubstepStartData.SyncState, WorkingSubstepStartData.AuxState);
 
 		for (TSharedPtr<FMovementModifierBase> Modifier : ActiveModifiers)
 		{
@@ -245,7 +243,7 @@ void UMovementModeStateMachine::OnSimulationTick(USceneComponent* UpdatedCompone
 		for (TSharedPtr<FLayeredMoveBase>& ActiveMove : ActiveMoves)
 		{
 			FProposedMove MoveStep;
-			if (ActiveMove->GenerateMove(SubstepStartData, SubTimeStep, MoverComp, SimBlackboard, MoveStep))
+			if (ActiveMove->GenerateMove(WorkingSubstepStartData, SubTimeStep, MoverComp, SimBlackboard, MoveStep))
 			{
 				// If this active move is already past it's first tick we don't need to set the preferred mode again
 				if (ActiveMove->StartSimTimeMs < SubTimeStep.BaseSimTimeMs)
@@ -275,7 +273,7 @@ void UMovementModeStateMachine::OnSimulationTick(USceneComponent* UpdatedCompone
 				!(bHasLayeredMoveContributions && CombinedLayeredMove.MixMode == EMoveMixMode::OverrideAll))
 			{
 				QUICK_SCOPE_CYCLE_COUNTER(STAT_GenerateMoveFromMode);
-				CurrentMode->GenerateMove(SubstepStartData, SubTimeStep, OUT CombinedMove);
+				CurrentMode->GenerateMove(WorkingSubstepStartData, SubTimeStep, OUT CombinedMove);
 
 				bHasModeMoveContribution = true;
 			}
@@ -300,16 +298,15 @@ void UMovementModeStateMachine::OnSimulationTick(USceneComponent* UpdatedCompone
 			}
 			CurrentLayeredMoves.ResetResidualVelocity();
 
-			MoverComp->ProcessGeneratedMovement.ExecuteIfBound(SubstepStartData, SubTimeStep, OUT CombinedMove);
+			MoverComp->ProcessGeneratedMovement.ExecuteIfBound(WorkingSubstepStartData, SubTimeStep, OUT CombinedMove);
 			
 			// Execute the combined proposed move
 			{
-				FSimulationTickParams SimTickParams;
-				SimTickParams.StartState = SubstepStartData;
-				SimTickParams.MovingComps = FMovingComponentSet(MoverComp);
-				SimTickParams.SimBlackboard = SimBlackboard;
-				SimTickParams.TimeStep = SubTimeStep;
-				SimTickParams.ProposedMove = CombinedMove;
+				WorkingSimTickParams.StartState = WorkingSubstepStartData;
+				WorkingSimTickParams.MovingComps.SetFrom(MoverComp);
+				WorkingSimTickParams.SimBlackboard = SimBlackboard;
+				WorkingSimTickParams.TimeStep = SubTimeStep;
+				WorkingSimTickParams.ProposedMove = CombinedMove;
 
 				// Check for any transitions, first those registered with the current movement mode, then global ones that could occur from any mode
 				FTransitionEvalResult EvalResult = FTransitionEvalResult::NoTransition;
@@ -319,7 +316,7 @@ void UMovementModeStateMachine::OnSimulationTick(USceneComponent* UpdatedCompone
 				{
 					if (IsValid(Transition) && ((SubStepCount == 0) || !Transition->bFirstSubStepOnly))
 					{
-						EvalResult = Transition->Evaluate(SimTickParams);
+						EvalResult = Transition->Evaluate(WorkingSimTickParams);
 
 						if (!EvalResult.NextMode.IsNone())
 						{
@@ -338,7 +335,7 @@ void UMovementModeStateMachine::OnSimulationTick(USceneComponent* UpdatedCompone
 					{
 						if (IsValid(Transition))
 						{
-							EvalResult = Transition->Evaluate(SimTickParams);
+							EvalResult = Transition->Evaluate(WorkingSimTickParams);
 
 							if (!EvalResult.NextMode.IsNone())
 							{
@@ -355,12 +352,12 @@ void UMovementModeStateMachine::OnSimulationTick(USceneComponent* UpdatedCompone
 				if (TransitionToTrigger && !EvalResult.NextMode.IsNone())
 				{
 					OutputState.MovementEndState.NextModeName = EvalResult.NextMode;
-					OutputState.MovementEndState.RemainingMs = SimTickParams.TimeStep.StepMs; 	// Pass all remaining time to next mode
-					TransitionToTrigger->Trigger(SimTickParams);
+					OutputState.MovementEndState.RemainingMs = WorkingSimTickParams.TimeStep.StepMs; 	// Pass all remaining time to next mode
+					TransitionToTrigger->Trigger(WorkingSimTickParams);
 				}
 				else
 				{
-					CurrentMode->SimulationTick(SimTickParams, OutputState);
+					CurrentMode->SimulationTick(WorkingSimTickParams, OutputState);
 				}
 
 				OutputState.MovementEndState.RemainingMs = FMath::Clamp(OutputState.MovementEndState.RemainingMs, 0.0f, SubTimeStep.StepMs);
@@ -400,15 +397,15 @@ void UMovementModeStateMachine::OnSimulationTick(USceneComponent* UpdatedCompone
 		SubTimeStep.BaseSimTimeMs += (SubTimeStep.StepMs - RemainingMs);
 		SubTimeStep.StepMs = EndingSimTimeMs - SubTimeStep.BaseSimTimeMs;
 
-		SubstepStartData.SyncState = OutputState.SyncState;
-		SubstepStartData.AuxState  = OutputState.AuxState;
+		WorkingSubstepStartData.SyncState = OutputState.SyncState;
+		WorkingSubstepStartData.AuxState  = OutputState.AuxState;
 
 		++SubStepCount;
 	}
 
 	FApplyMovementEffectParams EffectParams;
 	EffectParams.MoverComp = MoverComp;
-	EffectParams.StartState = &SubstepStartData;
+	EffectParams.StartState = &WorkingSubstepStartData;
 	EffectParams.TimeStep = &SubTimeStep;
 	EffectParams.UpdatedComponent = UpdatedComponent;
 	EffectParams.UpdatedPrimitive = UpdatedPrimitive;
diff --git a/Engine/Plugins/Experimental/Mover/Source/Mover/Public/Backends/MoverStandaloneLiaison.h b/Engine/Plugins/Experimental/Mover/Source/Mover/Public/Backends/MoverStandaloneLiaison.h
--- a/Engine/Plugins/Experimental/Mover/Source/Mover/Public/Backends/MoverStandaloneLiaison.h
+++ b/Engine/Plugins/Experimental/Mover/Source/Mover/Public/Backends/MoverStandaloneLiaison.h
@@ -190,6 +190,11 @@ protected:
 	FMoverStandaloneApplyStateTickFunction ApplyStateTickFunction;
 
 	friend struct FMoverStandaloneProduceInputTickFunction;
+private:
+	// Internal-use-only tick data structs, for efficiency since they typically have the same contents from frame to frame
+	FMoverTickStartData WorkingStartData;
+	FMoverTickEndData WorkingEndData;
+
 	friend struct FMoverStandaloneSimulateMovementTickFunction;
 	friend struct FMoverStandaloneApplyStateTickFunction;
 };
diff --git a/Engine/Plugins/Experimental/Mover/Source/Mover/Public/MoveLibrary/MovementUtilsTypes.h b/Engine/Plugins/Experimental/Mover/Source/Mover/Public/MoveLibrary/MovementUtilsTypes.h
--- a/Engine/Plugins/Experimental/Mover/Source/Mover/Public/MoveLibrary/MovementUtilsTypes.h
+++ b/Engine/Plugins/Experimental/Mover/Source/Mover/Public/MoveLibrary/MovementUtilsTypes.h
@@ -73,8 +73,11 @@ struct MOVER_API FMovingComponentSet
 	GENERATED_BODY()
 
 	FMovingComponentSet() {}
-	FMovingComponentSet(USceneComponent* InUpdatedComponent);
-	FMovingComponentSet(UMoverComponent* InMoverComponent);
+	FMovingComponentSet(USceneComponent* InUpdatedComponent) { SetFrom(InUpdatedComponent); }
+	FMovingComponentSet(UMoverComponent* InMoverComponent) { SetFrom(InMoverComponent); }
+
+	void SetFrom(USceneComponent* InUpdatedComponent);
+	void SetFrom(UMoverComponent* InMoverComponent);
 
 	UPROPERTY(EditAnywhere, BlueprintReadWrite, Category = Mover)
 	TWeakObjectPtr<USceneComponent> UpdatedComponent;
diff --git a/Engine/Plugins/Experimental/Mover/Source/Mover/Public/MovementModeStateMachine.h b/Engine/Plugins/Experimental/Mover/Source/Mover/Public/MovementModeStateMachine.h
--- a/Engine/Plugins/Experimental/Mover/Source/Mover/Public/MovementModeStateMachine.h
+++ b/Engine/Plugins/Experimental/Mover/Source/Mover/Public/MovementModeStateMachine.h
@@ -84,6 +84,10 @@ protected:
  	/** Modifiers that are to be canceled at the start of the next sim subtick */
  	TArray<FMovementModifierHandle> ModifiersToCancel;
  	
+	// Internal-use-only tick data structs, for efficiency since they typically have the same contents from frame to frame
+	FMoverTickStartData WorkingSubstepStartData;
+	FSimulationTickParams WorkingSimTickParams;
+
 private:
 	void ConstructDefaultModes();
 	void AdvanceToNextMode();
