Mover: shifting blackboard lifetime to be bookended by MoverComponent's InitializeComponent / UninitializeComponent

#rnx

diff --git a/Engine/Plugins/Experimental/Mover/Source/Mover/Private/MoveLibrary/MoverBlackboard.cpp b/Engine/Plugins/Experimental/Mover/Source/Mover/Private/MoveLibrary/MoverBlackboard.cpp
--- a/Engine/Plugins/Experimental/Mover/Source/Mover/Private/MoveLibrary/MoverBlackboard.cpp
+++ b/Engine/Plugins/Experimental/Mover/Source/Mover/Private/MoveLibrary/MoverBlackboard.cpp
@@ -23,3 +23,8 @@ void UMoverBlackboard::Invalidate(EInvalidationReason Reason)
 	}
 }
 
+void UMoverBlackboard::BeginDestroy()
+{
+	InvalidateAll();
+	Super::BeginDestroy();
+}
\ No newline at end of file
diff --git a/Engine/Plugins/Experimental/Mover/Source/Mover/Private/MoverComponent.cpp b/Engine/Plugins/Experimental/Mover/Source/Mover/Private/MoverComponent.cpp
--- a/Engine/Plugins/Experimental/Mover/Source/Mover/Private/MoverComponent.cpp
+++ b/Engine/Plugins/Experimental/Mover/Source/Mover/Private/MoverComponent.cpp
@@ -86,6 +86,13 @@ void UMoverComponent::InitializeComponent()
 
 	if (MyWorld && MyWorld->IsGameWorld())
 	{
+		if (SimBlackboard)
+		{
+			SimBlackboard->InvalidateAll();
+		}
+
+		SimBlackboard = NewObject<UMoverBlackboard>(this, TEXT("MoverBlackboard"), RF_Transient);
+
 		FindDefaultUpdatedComponent();
 
 		// Instantiate out sister backend component that will actually talk to the system driving the simulation
@@ -122,6 +129,13 @@ void UMoverComponent::UninitializeComponent()
 		BackendLiaisonComp = nullptr;
 	}
 
+	if (SimBlackboard)
+	{
+		SimBlackboard->InvalidateAll();
+	}
+
+	SimBlackboard = nullptr;
+
 	Super::UninitializeComponent();
 }
 
@@ -784,8 +798,6 @@ bool UMoverComponent::InitMoverSimulation()
 {
 	check(UpdatedComponent);
 
-	SimBlackboard = NewObject<UMoverBlackboard>(this, TEXT("MoverBlackboard"), RF_Transient);
-
 	// Set up FSM and initial movement states
 	ModeFSM = NewObject<UMovementModeStateMachine>(this, TEXT("MoverStateMachine"), RF_Transient);
 	ModeFSM->ClearAllMovementModes();
diff --git a/Engine/Plugins/Experimental/Mover/Source/Mover/Public/MoveLibrary/MoverBlackboard.h b/Engine/Plugins/Experimental/Mover/Source/Mover/Public/MoveLibrary/MoverBlackboard.h
--- a/Engine/Plugins/Experimental/Mover/Source/Mover/Public/MoveLibrary/MoverBlackboard.h
+++ b/Engine/Plugins/Experimental/Mover/Source/Mover/Public/MoveLibrary/MoverBlackboard.h
@@ -130,6 +130,9 @@ public:
 	/** Invalidate all objects */
 	void InvalidateAll() { Invalidate(EInvalidationReason::FullReset); }
 
+	// UObject interface
+	virtual void BeginDestroy() override;
+	// End UObject interface
 
 private:
 	TMap<FName, TUniquePtr<BlackboardObject>> ObjectsByName;
diff --git a/Engine/Plugins/Experimental/Mover/Source/Mover/Public/MoverComponent.h b/Engine/Plugins/Experimental/Mover/Source/Mover/Public/MoverComponent.h
--- a/Engine/Plugins/Experimental/Mover/Source/Mover/Public/MoverComponent.h
+++ b/Engine/Plugins/Experimental/Mover/Source/Mover/Public/MoverComponent.h
@@ -605,7 +605,7 @@ protected:
 
 	/** Used to store cached data & computations between decoupled systems, that can be referenced by name */
 	UPROPERTY(Transient)
-	TObjectPtr<UMoverBlackboard> SimBlackboard;
+	TObjectPtr<UMoverBlackboard> SimBlackboard = nullptr;
 
 private:
 	/** Collection of settings objects that are shared between movement modes. This list is automatically managed based on the @MovementModes contents. */
