Mover: adding BP exposure for trajectory sample member variables

diff --git a/Engine/Plugins/Experimental/Mover/Source/Mover/Public/MoveLibrary/MovementUtils.h b/Engine/Plugins/Experimental/Mover/Source/Mover/Public/MoveLibrary/MovementUtils.h
--- a/Engine/Plugins/Experimental/Mover/Source/Mover/Public/MoveLibrary/MovementUtils.h
+++ b/Engine/Plugins/Experimental/Mover/Source/Mover/Public/MoveLibrary/MovementUtils.h
@@ -27,14 +27,25 @@ struct MOVER_API FTrajectorySampleInfo
 {
 	GENERATED_BODY()
 
-	FTransform Transform;					// Position and orientation (world space)
+	// Position and orientation (world space)
+	UPROPERTY(BlueprintReadWrite, EditAnywhere, Category = Mover)
+	FTransform Transform;
+
+	// Velocity at the time of this sample (world space, units/sec)
+	UPROPERTY(BlueprintReadWrite, EditAnywhere, Category = Mover)
+	FVector LinearVelocity;
 
-	FVector LinearVelocity;					// Velocity at the time of this sample (world space, units/sec)
-	FVector InstantaneousAcceleration;		// Acceleration at the time of this sample (world space, units/sec^2)
+	// Acceleration at the time of this sample (world space, units/sec^2)
+	UPROPERTY(BlueprintReadWrite, EditAnywhere, Category = Mover)
+	FVector InstantaneousAcceleration;
 
-	FRotator AngularVelocity;				// Rotational velocity (world space, degrees/sec)
+	// Rotational velocity (world space, degrees/sec)
+	UPROPERTY(BlueprintReadWrite, EditAnywhere, Category = Mover)
+	FRotator AngularVelocity;
 
-	float SimTimeMs;						// Time stamp of this sample, in server simulation time
+	// Time stamp of this sample, in server simulation time
+	UPROPERTY(BlueprintReadWrite, EditAnywhere, Category = Mover)
+	float SimTimeMs;
 };
 
 // Input parameters for compute velocity function
