Mover: adding internal rw lock to Mover blackboard for thread safety

#rnx

diff --git a/Engine/Plugins/Experimental/Mover/Source/Mover/Private/MoveLibrary/MoverBlackboard.cpp b/Engine/Plugins/Experimental/Mover/Source/Mover/Private/MoveLibrary/MoverBlackboard.cpp
--- a/Engine/Plugins/Experimental/Mover/Source/Mover/Private/MoveLibrary/MoverBlackboard.cpp
+++ b/Engine/Plugins/Experimental/Mover/Source/Mover/Private/MoveLibrary/MoverBlackboard.cpp
@@ -7,6 +7,7 @@
 
 void UMoverBlackboard::Invalidate(FName ObjName)
 {
+	FRWScopeLock Lock(ObjectsMapLock, SLT_Write);
 	ObjectsByName.Remove(ObjName);
 }
 
@@ -16,8 +17,11 @@ void UMoverBlackboard::Invalidate(EInvalidationReason Reason)
 	{
 		default:
 		case EInvalidationReason::FullReset:
+		{
+			FRWScopeLock Lock(ObjectsMapLock, SLT_Write);
 			ObjectsByName.Empty();
-			break;
+		}
+		break;
 
 		// TODO: Support other reasons
 	}
diff --git a/Engine/Plugins/Experimental/Mover/Source/Mover/Public/MoveLibrary/MoverBlackboard.h b/Engine/Plugins/Experimental/Mover/Source/Mover/Public/MoveLibrary/MoverBlackboard.h
--- a/Engine/Plugins/Experimental/Mover/Source/Mover/Public/MoveLibrary/MoverBlackboard.h
+++ b/Engine/Plugins/Experimental/Mover/Source/Mover/Public/MoveLibrary/MoverBlackboard.h
@@ -78,6 +78,8 @@ public:
  	template<typename T>
 	bool TryGet(FName ObjName, T& OutFoundValue) const
 	{
+		FRWScopeLock Lock(ObjectsMapLock, SLT_ReadOnly);
+
 		if (const TUniquePtr<BlackboardObject>* ExistingObject = ObjectsByName.Find(ObjName))
 		{
 			OutFoundValue = ExistingObject->Get()->Get<T>();
@@ -109,8 +111,9 @@ public:
 	*/
 
 	/** Returns true/false to indicate if an object is stored with that name */
-	bool Contains(FName ObjName)
+	bool Contains(FName ObjName) const
 	{
+		FRWScopeLock Lock(ObjectsMapLock, SLT_ReadOnly);
 		return ObjectsByName.Contains(ObjName);
 	}
 
@@ -118,6 +121,7 @@ public:
 	template<typename T>
 	void Set(FName ObjName, T Obj)
 	{
+		FRWScopeLock Lock(ObjectsMapLock, SLT_Write);
 		ObjectsByName.Emplace(ObjName, MakeUnique<BlackboardObject>(Obj));
 	}
 
@@ -135,5 +139,6 @@ public:
 	// End UObject interface
 
 private:
+	mutable FRWLock ObjectsMapLock;		// used internally when reading/writing to ObjectsByName map
 	TMap<FName, TUniquePtr<BlackboardObject>> ObjectsByName;
 };
\ No newline at end of file
