Use static FName for Mover DiagnosticContext

diff --git a/Engine/Plugins/Experimental/Mover/Source/Mover/Private/Backends/MoverStandaloneLiaison.cpp b/Engine/Plugins/Experimental/Mover/Source/Mover/Private/Backends/MoverStandaloneLiaison.cpp
--- a/Engine/Plugins/Experimental/Mover/Source/Mover/Private/Backends/MoverStandaloneLiaison.cpp
+++ b/Engine/Plugins/Experimental/Mover/Source/Mover/Private/Backends/MoverStandaloneLiaison.cpp
@@ -424,7 +424,8 @@ FName FMoverStandaloneProduceInputTickFunction::DiagnosticContext(bool bDetailed
 		return FName(*FString::Printf(TEXT("MoverStandaloneProduceInputTickFunction/%s"), *GetFullNameSafe(Target.Get())));
 	}
 
-	return FName(TEXT("MoverStandaloneProduceInputTickFunction"));
+	static const FName MoverStandaloneProduceInputTickFunctionFName(TEXT("MoverStandaloneProduceInputTickFunction"));
+	return MoverStandaloneProduceInputTickFunctionFName;
 }
 
 // FMoverStandaloneSimulateMovementTickFunction ---------------------
@@ -450,7 +451,8 @@ FName FMoverStandaloneSimulateMovementTickFunction::DiagnosticContext(bool bDeta
 		return FName(*FString::Printf(TEXT("MoverStandaloneSimulateMovementTickFunction/%s"), *GetFullNameSafe(Target.Get())));
 	}
 
-	return FName(TEXT("MoverStandaloneSimulateMovementTickFunction"));
+	static const FName MoverStandaloneSimulateMovementTickFunctionFName(TEXT("MoverStandaloneSimulateMovementTickFunction"));
+	return MoverStandaloneSimulateMovementTickFunctionFName;
 }
 
 
@@ -477,5 +479,6 @@ FName FMoverStandaloneApplyStateTickFunction::DiagnosticContext(bool bDetailed)
 		return FName(*FString::Printf(TEXT("MoverStandaloneApplyStateTickFunction/%s"), *GetFullNameSafe(Target.Get())));
 	}
 
-	return FName(TEXT("MoverStandaloneApplyStateTickFunction"));
+	static const FName MoverStandaloneApplyStateTickFunctionFName(TEXT("MoverStandaloneApplyStateTickFunction"));
+	return MoverStandaloneApplyStateTickFunctionFName;
 }
\ No newline at end of file
