Adapted UAsyncNavWalkingMode::FindBestNavMeshLocation to handle HitResults with no associated PrimitiveComponent but with a valid BodyInstanceOwner.

diff --git a/Engine/Plugins/Experimental/Mover/Source/Mover/Private/DefaultMovementSet/Modes/AsyncNavWalkingMode.cpp b/Engine/Plugins/Experimental/Mover/Source/Mover/Private/DefaultMovementSet/Modes/AsyncNavWalkingMode.cpp
--- a/Engine/Plugins/Experimental/Mover/Source/Mover/Private/DefaultMovementSet/Modes/AsyncNavWalkingMode.cpp
+++ b/Engine/Plugins/Experimental/Mover/Source/Mover/Private/DefaultMovementSet/Modes/AsyncNavWalkingMode.cpp
@@ -8,6 +8,7 @@
 #include "Components/ShapeComponent.h"
 #include "DefaultMovementSet/NavMoverComponent.h"
 #include "DefaultMovementSet/Settings/CommonLegacyMovementSettings.h"
+#include "PhysicsEngine/PhysicsBodyInstanceOwnerInterface.h"
 #include "MoveLibrary/AsyncMovementUtils.h"
 #include "MoveLibrary/FloorQueryUtils.h"
 #include "MoveLibrary/GroundMovementUtils.h"
@@ -440,10 +441,19 @@ void UAsyncNavWalkingMode::FindBestNavMeshLocation(const FVector& TraceStart, co
 		FORCEINLINE bool operator()(const FHitResult& TestHit) const
 		{
 			UPrimitiveComponent* PrimComp = TestHit.GetComponent();
-			const bool bBlockOnWorldStatic = PrimComp && (PrimComp->GetCollisionResponseToChannel(ECC_WorldStatic) == ECR_Block);
-			const bool bBlockOnWorldDynamic = PrimComp && (PrimComp->GetCollisionResponseToChannel(ECC_WorldDynamic) == ECR_Block);
-
-			return !bBlockOnWorldStatic && (!bBlockOnWorldDynamic || bCheckOnlyWorldStatic);
+			// Prefer using primitive component if valid
+			if (IPhysicsBodyInstanceOwner* BodyInstanceOwner = !PrimComp ? IPhysicsBodyInstanceOwner::GetPhysicsBodyInstandeOwnerFromHitResult(TestHit) : nullptr)
+			{
+				const bool bBlockOnWorldStatic = (BodyInstanceOwner->GetCollisionResponseToChannel(ECC_WorldStatic) == ECR_Block);
+				const bool bBlockOnWorldDynamic = (BodyInstanceOwner->GetCollisionResponseToChannel(ECC_WorldDynamic) == ECR_Block);
+				return !bBlockOnWorldStatic && (!bBlockOnWorldDynamic || bCheckOnlyWorldStatic);
+			}
+			else
+			{
+				const bool bBlockOnWorldStatic = PrimComp && (PrimComp->GetCollisionResponseToChannel(ECC_WorldStatic) == ECR_Block);
+				const bool bBlockOnWorldDynamic = PrimComp && (PrimComp->GetCollisionResponseToChannel(ECC_WorldDynamic) == ECR_Block);
+				return !bBlockOnWorldStatic && (!bBlockOnWorldDynamic || bCheckOnlyWorldStatic);
+			}
 		}
 
 		bool bCheckOnlyWorldStatic;
