Mover: added warnings to detect excessive speed coming from root motion attributes. Configure the speed threshold using mover.debug.RootMotionAttributesExcessiveSpeedSq

diff --git a/Engine/Plugins/Experimental/Mover/Source/Mover/Private/DefaultMovementSet/LayeredMoves/RootMotionAttributeLayeredMove.cpp b/Engine/Plugins/Experimental/Mover/Source/Mover/Private/DefaultMovementSet/LayeredMoves/RootMotionAttributeLayeredMove.cpp
--- a/Engine/Plugins/Experimental/Mover/Source/Mover/Private/DefaultMovementSet/LayeredMoves/RootMotionAttributeLayeredMove.cpp
+++ b/Engine/Plugins/Experimental/Mover/Source/Mover/Private/DefaultMovementSet/LayeredMoves/RootMotionAttributeLayeredMove.cpp
@@ -25,6 +25,12 @@ FAutoConsoleVariable CVarDisableRootMotionAttrContributions(
 	TEXT("If enabled, contributions from root motion attributes will be ignored in favor of other Mover influences"),
 	ECVF_Cheat);
 
+static float ExcessiveLinearVelocitySquaredThreshold = 400.f * 400.f;
+FAutoConsoleVariableRef CVarExcessiveLinearSpeedSquaredThreshold(
+	TEXT("mover.debug.RootMotionAttributesExcessiveSpeedSq"),
+	ExcessiveLinearVelocitySquaredThreshold,
+	TEXT("If > 0, a warning will be logged when a root motion attribute's squared speed exceeds this threshold\n"));
+
 #endif	// !UE_BUILD_SHIPPING
 
 
@@ -150,9 +156,19 @@ bool FLayeredMove_RootMotionAttribute::GenerateMove(const FMoverTickStartData& S
 		OutProposedMove.AngularVelocity = WorldSpaceRootMotion.GetRotation().Rotator() * (1.f / DeltaSeconds);
 
 #if !UE_BUILD_SHIPPING
-		UE_CLOG(CVarLogRootMotionAttrSteps->GetBool(), LogMover, Log, TEXT("RootMotionAttr. SimF %i (dt %.3f) => LocalT: %s (WST: %s)  Vel: %.3f"),
+		UE_CLOG(CVarLogRootMotionAttrSteps->GetBool(), LogMover, Log, TEXT("RootMotionAttr. SimF %i (dt %.3f) => LocalT: %s (WST: %s)  XY Speed: %.6f Z: %.6f   AngV: %s"),
 			TimeStep.ServerFrame, DeltaSeconds,
-			*LocalRootMotion.GetTranslation().ToString(), *WorldSpaceRootMotion.GetTranslation().ToString(), OutProposedMove.LinearVelocity.Length());
+			*LocalRootMotion.GetTranslation().ToString(), *WorldSpaceRootMotion.GetTranslation().ToString(),
+			OutProposedMove.LinearVelocity.Size2D(), OutProposedMove.LinearVelocity.Z, *OutProposedMove.AngularVelocity.ToCompactString());
+
+		if (ExcessiveLinearVelocitySquaredThreshold > 0.f && ExcessiveLinearVelocitySquaredThreshold < OutProposedMove.LinearVelocity.SquaredLength())
+		{
+			UE_LOG(LogMover, Warning, TEXT("RootMotionAttr on %s has excessive speed.  LocalTrans: %s (DT: %.5f) -> XY Vel: %.6f  Z Vel: %.6f"),
+				*GetNameSafe(MoverComp->GetOwner()),
+				*LocalRootMotion.GetTranslation().ToString(), DeltaSeconds, OutProposedMove.LinearVelocity.Size2D(), OutProposedMove.LinearVelocity.Z);
+		}
+
+
 #endif // !UE_BUILD_SHIPPING
 
 		return true;
