Chaos Visual Debugger: Fix to sync data and local sim data from the mover plugin not making it to CVD when recording

#rnx

diff --git a/Engine/Plugins/Experimental/Mover/Source/Mover/Private/ChaosVisualDebugger/MoverCVDRuntimeTrace.cpp b/Engine/Plugins/Experimental/Mover/Source/Mover/Private/ChaosVisualDebugger/MoverCVDRuntimeTrace.cpp
--- a/Engine/Plugins/Experimental/Mover/Source/Mover/Private/ChaosVisualDebugger/MoverCVDRuntimeTrace.cpp
+++ b/Engine/Plugins/Experimental/Mover/Source/Mover/Private/ChaosVisualDebugger/MoverCVDRuntimeTrace.cpp
@@ -38,6 +38,11 @@ public:
 	{
 		return *this;
 	}
+
+	virtual FArchive& operator<<(struct FWeakObjectPtr& Value)
+	{
+		return *this;
+	}
 };
 class FSkipObjectRefsMemoryReader : public FMemoryReader
 {
@@ -51,6 +56,11 @@ public:
 	{
 		return *this;
 	}
+
+	virtual FArchive& operator<<(struct FWeakObjectPtr& Value)
+	{
+		return *this;
+	}
 };
 
 void FMoverCVDRuntimeTrace::UnwrapSimData(const FMoverCVDSimDataWrapper& InSimDataWrapper, TSharedPtr<FMoverInputCmdContext>& OutInputCmd, TSharedPtr<FMoverSyncState>& OutSyncState, TSharedPtr<FMoverDataCollection>& OutLocalSimData)
@@ -63,7 +73,7 @@ void FMoverCVDRuntimeTrace::UnwrapSimData(const FMoverCVDSimDataWrapper& InSimDa
 		FMemoryReader ArReader(InSimDataWrapper.InputCmdBytes, true);
 		OutInputCmd = MakeShared<FMoverInputCmdContext>();
 		FMoverInputCmdContext::StaticStruct()->SerializeBin(ArReader, OutInputCmd.Get());
-		// Input cmd's Collection of Custom Input Structs
+		// Input cmd's Collection of custom structs
 		FSkipObjectRefsMemoryReader ArInputCollectionReader(InSimDataWrapper.InputMoverDataCollectionBytes, true);
 		OutInputCmd->InputCollection.SerializeDebugData(ArInputCollectionReader);
 	}
@@ -74,7 +84,13 @@ void FMoverCVDRuntimeTrace::UnwrapSimData(const FMoverCVDSimDataWrapper& InSimDa
 		FMemoryReader ArReader(InSimDataWrapper.SyncStateBytes, true);
 		OutSyncState = MakeShared<FMoverSyncState>();
 		FMoverSyncState::StaticStruct()->SerializeBin(ArReader, OutSyncState.Get());
+		// Sync State's Collection of custom structs
+		FSkipObjectRefsMemoryReader ArInputCollectionReader(InSimDataWrapper.SyncStateDataCollectionBytes, true);
+		OutInputCmd->InputCollection.SerializeDebugData(ArInputCollectionReader);
+	}
 
+	{
+		FSkipObjectRefsMemoryReader ArReader(InSimDataWrapper.LocalSimDataBytes, true);
 		OutLocalSimData = MakeShared<FMoverDataCollection>();
 		OutLocalSimData->SerializeDebugData(ArReader);
 	}
@@ -93,8 +109,7 @@ void FMoverCVDRuntimeTrace::WrapSimData(uint32 SolverID, uint32 ParticleID, cons
 		// This is not version friendly, we need to instead use SerializeTagParticles. Slower and Sergio is working on a faster version, but it's not available yet.
 		InInputCmd.StaticStruct()->SerializeBin(ArWriter, &const_cast<FMoverInputCmdContext&>(InInputCmd));
 	}
-
-	// Input cmd's Collection of Custom Input Structs
+	// Input cmd's Collection of custom structs
 	{
 		FSkipObjectRefsMemoryWriter ArWriter(OutSimDataWrapper.InputMoverDataCollectionBytes, true);
 		const_cast<FMoverInputCmdContext&>(InInputCmd).InputCollection.SerializeDebugData(ArWriter);
@@ -105,8 +120,15 @@ void FMoverCVDRuntimeTrace::WrapSimData(uint32 SolverID, uint32 ParticleID, cons
 		FMemoryWriter ArWriter(OutSimDataWrapper.SyncStateBytes, true);
 		// This is not version friendly, we need to instead use SerializeTagParticles. Slower and Sergio is working on a faster version, but it's not available yet.
 		InSyncState.StaticStruct()->SerializeBin(ArWriter, &const_cast<FMoverSyncState&>(InSyncState));
+	}
+	{
+		FSkipObjectRefsMemoryWriter ArWriter(OutSimDataWrapper.SyncStateDataCollectionBytes, true);
+		const_cast<FMoverSyncState&>(InSyncState).SyncStateCollection.SerializeDebugData(ArWriter);
+	}
 
-		// Local Sim Data Collection
+	// Local sim data (catch all other structs we want to record)
+	{
+		FSkipObjectRefsMemoryWriter ArWriter(OutSimDataWrapper.LocalSimDataBytes, true);
 		if (LocalSimData)
 		{
 			(const_cast<FMoverDataCollection*>(LocalSimData))->SerializeDebugData(ArWriter);
diff --git a/Engine/Plugins/Experimental/Mover/Source/MoverCVDData/Private/MoverCVDDataWrappers.cpp b/Engine/Plugins/Experimental/Mover/Source/MoverCVDData/Private/MoverCVDDataWrappers.cpp
--- a/Engine/Plugins/Experimental/Mover/Source/MoverCVDData/Private/MoverCVDDataWrappers.cpp
+++ b/Engine/Plugins/Experimental/Mover/Source/MoverCVDData/Private/MoverCVDDataWrappers.cpp
@@ -14,8 +14,10 @@ bool FMoverCVDSimDataWrapper::Serialize(FArchive& Ar)
 	Ar << SolverID;
 	Ar << ParticleID;
 	Ar << SyncStateBytes;
+	Ar << SyncStateDataCollectionBytes;
 	Ar << InputCmdBytes;
 	Ar << InputMoverDataCollectionBytes;
+	Ar << LocalSimDataBytes;
 
 	return !Ar.IsError();
 }
diff --git a/Engine/Plugins/Experimental/Mover/Source/MoverCVDData/Public/MoverCVDDataWrappers.h b/Engine/Plugins/Experimental/Mover/Source/MoverCVDData/Public/MoverCVDDataWrappers.h
--- a/Engine/Plugins/Experimental/Mover/Source/MoverCVDData/Public/MoverCVDDataWrappers.h
+++ b/Engine/Plugins/Experimental/Mover/Source/MoverCVDData/Public/MoverCVDDataWrappers.h
@@ -22,8 +22,10 @@ struct FMoverCVDSimDataWrapper : public FChaosVDWrapperDataBase
 	int32 ParticleID = INDEX_NONE;
 
 	TArray<uint8> SyncStateBytes;
+	TArray<uint8> SyncStateDataCollectionBytes;
 	TArray<uint8> InputCmdBytes;
 	TArray<uint8> InputMoverDataCollectionBytes;
+	TArray<uint8> LocalSimDataBytes;
 
 	MOVERCVDDATA_API bool Serialize(FArchive& Ar);
 };
