Mover: increasing the default threshold for reporting excessive movement speed

#rnx

diff --git a/Engine/Plugins/Experimental/Mover/Source/Mover/Private/DefaultMovementSet/LayeredMoves/RootMotionAttributeLayeredMove.cpp b/Engine/Plugins/Experimental/Mover/Source/Mover/Private/DefaultMovementSet/LayeredMoves/RootMotionAttributeLayeredMove.cpp
--- a/Engine/Plugins/Experimental/Mover/Source/Mover/Private/DefaultMovementSet/LayeredMoves/RootMotionAttributeLayeredMove.cpp
+++ b/Engine/Plugins/Experimental/Mover/Source/Mover/Private/DefaultMovementSet/LayeredMoves/RootMotionAttributeLayeredMove.cpp
@@ -25,7 +25,7 @@ FAutoConsoleVariable CVarDisableRootMotionAttrContributions(
 	TEXT("If enabled, contributions from root motion attributes will be ignored in favor of other Mover influences"),
 	ECVF_Cheat);
 
-static float ExcessiveLinearVelocitySquaredThreshold = 400.f * 400.f;
+static float ExcessiveLinearVelocitySquaredThreshold = 2000.f * 2000.f;
 FAutoConsoleVariableRef CVarExcessiveLinearSpeedSquaredThreshold(
 	TEXT("mover.debug.RootMotionAttributesExcessiveSpeedSq"),
 	ExcessiveLinearVelocitySquaredThreshold,
