Raise the priority of the mover simulate and apply ticks so they happen as soon as their prerequisites finish

diff --git a/Engine/Plugins/Experimental/Mover/Source/Mover/Private/Backends/MoverStandaloneLiaison.cpp b/Engine/Plugins/Experimental/Mover/Source/Mover/Private/Backends/MoverStandaloneLiaison.cpp
--- a/Engine/Plugins/Experimental/Mover/Source/Mover/Private/Backends/MoverStandaloneLiaison.cpp
+++ b/Engine/Plugins/Experimental/Mover/Source/Mover/Private/Backends/MoverStandaloneLiaison.cpp
@@ -69,11 +69,13 @@ UMoverStandaloneLiaisonComponent::UMoverStandaloneLiaisonComponent(const FObject
 	SimulateMovementTickFunction.SetTickFunctionEnable(true);
 	SimulateMovementTickFunction.TickGroup = TG_PrePhysics;
 	SimulateMovementTickFunction.bRunOnAnyThread = bUseAsyncMovementSimulationTick && MoverStandaloneBackendCVars::RunMovementSimOnAnyThread;
+	SimulateMovementTickFunction.bHighPriority = true;
 
 	ApplyStateTickFunction.bCanEverTick = true;
 	ApplyStateTickFunction.bStartWithTickEnabled = true;
 	ApplyStateTickFunction.SetTickFunctionEnable(true);
 	ApplyStateTickFunction.TickGroup = TG_PrePhysics;
+	ApplyStateTickFunction.bHighPriority = true;	
 
 	bWantsInitializeComponent = true;
 	bAutoActivate = true;
