Mover: improving thread safety with locked access around action queues.

diff --git a/Engine/Plugins/Experimental/Mover/Source/Mover/Private/MovementModeStateMachine.cpp b/Engine/Plugins/Experimental/Mover/Source/Mover/Private/MovementModeStateMachine.cpp
--- a/Engine/Plugins/Experimental/Mover/Source/Mover/Private/MovementModeStateMachine.cpp
+++ b/Engine/Plugins/Experimental/Mover/Source/Mover/Private/MovementModeStateMachine.cpp
@@ -428,8 +428,21 @@ void UMovementModeStateMachine::OnSimulationPreRollback(const FMoverSyncState* I
 void UMovementModeStateMachine::OnSimulationRollback(const FMoverSyncState* SyncState, const FMoverAuxStateContext* AuxState)
 {
 	ClearQueuedMode();
-	QueuedLayeredMoves.Empty();
-	QueuedInstantEffects.Empty();
+
+	{
+		FRWScopeLock QueueLock(LayeredMoveQueueLock, SLT_Write);
+		QueuedLayeredMoves.Empty();
+	}
+
+	{
+		FRWScopeLock QueueLock(InstantEffectsQueueLock, SLT_Write);
+		QueuedInstantEffects.Empty();
+	}
+
+	{
+		FRWScopeLock QueueLock(ModifiersQueueLock, SLT_Write);
+		QueuedMovementModifiers.Empty();
+	}
 }
 
 
@@ -455,11 +468,13 @@ const UBaseMovementMode* UMovementModeStateMachine::FindMovementMode(FName ModeN
 
 void UMovementModeStateMachine::QueueLayeredMove(TSharedPtr<FLayeredMoveBase> Move)
 {
+	FRWScopeLock QueueLock(LayeredMoveQueueLock, SLT_Write);
 	QueuedLayeredMoves.Add(Move);
 }
 
 void UMovementModeStateMachine::QueueInstantMovementEffect(TSharedPtr<FInstantMovementEffect> Effect)
 {
+	FRWScopeLock QueueLock(InstantEffectsQueueLock, SLT_Write);
 	QueuedInstantEffects.Add(Effect);
 }
 
@@ -467,6 +482,8 @@ FMovementModifierHandle UMovementModeStateMachine::QueueMovementModifier(TShared
 {
 	if (ensure(Modifier.IsValid()))
 	{
+		FRWScopeLock QueueLock(ModifiersQueueLock, SLT_Write);
+
 		QueuedMovementModifiers.Add(Modifier);
 		Modifier->GenerateHandle();
 
@@ -478,25 +495,31 @@ FMovementModifierHandle UMovementModeStateMachine::QueueMovementModifier(TShared
 
 void UMovementModeStateMachine::CancelModifierFromHandle(FMovementModifierHandle ModifierHandle)
 {
-	QueuedMovementModifiers.RemoveAll([ModifierHandle, this]
-	(const TSharedPtr<FMovementModifierBase>& Modifier)
 	{
-		if (Modifier.IsValid())
+		FRWScopeLock QueueLock(ModifiersQueueLock, SLT_Write);
+		QueuedMovementModifiers.RemoveAll([ModifierHandle, this]
+		(const TSharedPtr<FMovementModifierBase>& Modifier)
 		{
-			if (Modifier->GetHandle() == ModifierHandle)
+			if (Modifier.IsValid())
 			{
-				return true;
+				if (Modifier->GetHandle() == ModifierHandle)
+				{
+					return true;
+				}
+			}
+			else
+			{
+				return true;	
 			}
-		}
-		else
-		{
-			return true;	
-		}
 
-		return false;
-	});
+			return false;
+		});
+	}
 	
-	ModifiersToCancel.Add(ModifierHandle);
+	{
+		FRWScopeLock QueueLock(ModifierCancelQueueLock, SLT_Write);
+		ModifiersToCancel.Add(ModifierHandle);
+	}
 }
 
 
@@ -542,6 +565,8 @@ void UMovementModeStateMachine::AdvanceToNextMode()
 
 void UMovementModeStateMachine::FlushQueuedMovesToGroup(FLayeredMoveGroup& Group)
 {
+	FRWScopeLock QueueLock(LayeredMoveQueueLock, SLT_Write);
+
 	if (!QueuedLayeredMoves.IsEmpty())
 	{
 		for (TSharedPtr<FLayeredMoveBase>& QueuedMove : QueuedLayeredMoves)
@@ -555,6 +580,8 @@ void UMovementModeStateMachine::FlushQueuedMovesToGroup(FLayeredMoveGroup& Group
 
 void UMovementModeStateMachine::FlushQueuedModifiersToGroup(FMovementModifierGroup& ModifierGroup)
 {
+	FRWScopeLock QueueLock(ModifiersQueueLock, SLT_Write);
+
 	if (!QueuedMovementModifiers.IsEmpty())
 	{
 		for (TSharedPtr<FMovementModifierBase>& QueuedModifier : QueuedMovementModifiers)
@@ -568,6 +595,8 @@ void UMovementModeStateMachine::FlushQueuedModifiersToGroup(FMovementModifierGro
 
 void UMovementModeStateMachine::FlushModifierCancellationsToGroup(FMovementModifierGroup& ActiveModifierGroup)
 {
+	FRWScopeLock QueueLock(ModifierCancelQueueLock, SLT_Write);
+	
 	for (FMovementModifierHandle HandleToCancel : ModifiersToCancel)
 	{
 		ActiveModifierGroup.CancelModifierFromHandle(HandleToCancel);
@@ -578,7 +607,10 @@ void UMovementModeStateMachine::FlushModifierCancellationsToGroup(FMovementModif
 
 void UMovementModeStateMachine::RollbackModifiers(const FMoverSyncState* InvalidSyncState, const FMoverSyncState* SyncState, const FMoverAuxStateContext* InvalidAuxState, const FMoverAuxStateContext* AuxState)
 {
-	QueuedMovementModifiers.Empty();
+	{
+		FRWScopeLock QueueLock(ModifiersQueueLock, SLT_Write);
+		QueuedMovementModifiers.Empty();
+	}
 	
 	if (UMoverComponent* MoverComp = Cast<UMoverComponent>(GetOuter()))
 	{
@@ -629,7 +661,9 @@ void UMovementModeStateMachine::RollbackModifiers(const FMoverSyncState* Invalid
 bool UMovementModeStateMachine::ApplyInstantEffects(FApplyMovementEffectParams& ApplyEffectParams, FMoverSyncState& OutputState)
 {
 	bool bInstantMovementEffectApplied = false;
-	
+
+	FRWScopeLock QueueLock(InstantEffectsQueueLock, SLT_Write);
+
 	if (!QueuedInstantEffects.IsEmpty())
 	{
 		for (TSharedPtr<FInstantMovementEffect>& QueuedEffect : QueuedInstantEffects)
diff --git a/Engine/Plugins/Experimental/Mover/Source/Mover/Public/MovementModeStateMachine.h b/Engine/Plugins/Experimental/Mover/Source/Mover/Public/MovementModeStateMachine.h
--- a/Engine/Plugins/Experimental/Mover/Source/Mover/Public/MovementModeStateMachine.h
+++ b/Engine/Plugins/Experimental/Mover/Source/Mover/Public/MovementModeStateMachine.h
@@ -72,16 +72,16 @@ protected:
 	FName DefaultModeName = NAME_None;
 	FName CurrentModeName = NAME_None;
 
-	/** Moves that are queued to be added to the simulation at the start of the next sim subtick */
+	/** Moves that are queued to be added to the simulation at the start of the next sim subtick. Access covered by lock. */
 	TArray<TSharedPtr<FLayeredMoveBase>> QueuedLayeredMoves;
 
- 	/** Effects that are queued to be applied to the simulation at the start of the next sim subtick or at the end of this tick */
+ 	/** Effects that are queued to be applied to the simulation at the start of the next sim subtick or at the end of this tick.  Access covered by lock. */
  	TArray<TSharedPtr<FInstantMovementEffect>> QueuedInstantEffects;
 
- 	/** Modifiers that are queued to be added to the simulation at the start of the next sim subtick */
+ 	/** Modifiers that are queued to be added to the simulation at the start of the next sim subtick. Access covered by lock. */
  	TArray<TSharedPtr<FMovementModifierBase>> QueuedMovementModifiers;
 
- 	/** Modifiers that are to be canceled at the start of the next sim subtick */
+ 	/** Modifiers that are to be canceled at the start of the next sim subtick.  Access covered by lock. */
  	TArray<FMovementModifierHandle> ModifiersToCancel;
  	
 	// Internal-use-only tick data structs, for efficiency since they typically have the same contents from frame to frame
@@ -89,6 +89,12 @@ protected:
 	FSimulationTickParams WorkingSimTickParams;
 
 private:
+	// Locks for thread safety on queueing mechanisms
+	mutable FRWLock LayeredMoveQueueLock;
+	mutable FRWLock InstantEffectsQueueLock;
+	mutable FRWLock ModifiersQueueLock;
+	mutable FRWLock ModifierCancelQueueLock;
+
 	void ConstructDefaultModes();
 	void AdvanceToNextMode();
 	void FlushQueuedMovesToGroup(FLayeredMoveGroup& Group);
