Mover: partial reversion of a change that would nullify the blackboard pointer earlier than necessary. Let's let the uproperty management handle it.

This is working around a rare case where an async movement sim tick could still be in-progress while something on the GT uninitialized the pawn.

diff --git a/Engine/Plugins/Experimental/Mover/Source/Mover/Private/MoverComponent.cpp b/Engine/Plugins/Experimental/Mover/Source/Mover/Private/MoverComponent.cpp
--- a/Engine/Plugins/Experimental/Mover/Source/Mover/Private/MoverComponent.cpp
+++ b/Engine/Plugins/Experimental/Mover/Source/Mover/Private/MoverComponent.cpp
@@ -134,8 +134,6 @@ void UMoverComponent::UninitializeComponent()
 		SimBlackboard->InvalidateAll();
 	}
 
-	SimBlackboard = nullptr;
-
 	Super::UninitializeComponent();
 }
 
