Optimize Mover copies to happen per index instead of searching the array twice

diff --git a/Engine/Plugins/Experimental/Mover/Source/Mover/Private/LayeredMove.cpp b/Engine/Plugins/Experimental/Mover/Source/Mover/Private/LayeredMove.cpp
--- a/Engine/Plugins/Experimental/Mover/Source/Mover/Private/LayeredMove.cpp
+++ b/Engine/Plugins/Experimental/Mover/Source/Mover/Private/LayeredMove.cpp
@@ -165,41 +165,38 @@ void FLayeredMoveGroup::NetSerialize(FArchive& Ar, uint8 MaxNumMovesToSerialize/
 
 static void CopyLayeredMoveArray(TArray<TSharedPtr<FLayeredMoveBase>>& Dest, const TArray<TSharedPtr<FLayeredMoveBase>>& Src)
 {
-	bool bCanCopyInPlace = (UE::Mover::DisableDataCopyInPlace == 0);
-
+	bool bCanCopyInPlace = (UE::Mover::DisableDataCopyInPlace == 0 && Dest.Num() == Src.Num());
 	if (bCanCopyInPlace)
 	{
-		if (Dest.Num() != Src.Num())
-		{
-			bCanCopyInPlace = false;
-		}
-		else
+		// If copy in place is enabled and the arrays are the same size, copy by index
+		for (int32 i = 0; i < Dest.Num(); ++i)
 		{
-			for (int32 i = 0; i < Dest.Num(); ++i)
+			if (FLayeredMoveBase* SrcData = Src[i].Get())
 			{
-				if ((!Dest[i].IsValid() || !Src[i].IsValid()) ||
-					(Dest[i]->GetScriptStruct() != Src[i]->GetScriptStruct()))
+				FLayeredMoveBase* DestData = Dest[i].Get();
+				UScriptStruct* SourceStruct = SrcData->GetScriptStruct();
+
+				if (DestData && SourceStruct == DestData->GetScriptStruct())
 				{
-					bCanCopyInPlace = false;
-					break;
+					// Same type so copy in place
+					SourceStruct->CopyScriptStruct(DestData, SrcData, 1);
+				}
+				else
+				{
+					// Different type so replace the shared ptr with a clone
+					Dest[i] = TSharedPtr<FLayeredMoveBase>(SrcData->Clone());
 				}
 			}
+			else
+			{
+				// Found invalid source, fall back to full copy
+				bCanCopyInPlace = false;
+				break;
+			}
 		}
 	}
-
-	if (bCanCopyInPlace)
-	{
-		// Copy all structs in place
-		for (int32 i = 0; i < Dest.Num(); ++i)
-		{
-			UScriptStruct* ScriptStruct = Dest[i]->GetScriptStruct();
-			FLayeredMoveBase* DestMove = Dest[i].Get();
-			FLayeredMoveBase* SrcMove = Src[i].Get();
-
-			ScriptStruct->CopyScriptStruct(DestMove, SrcMove, 1);
-		}
-	}
-	else
+	
+	if (!bCanCopyInPlace)
 	{
 		// Deep copy active moves
 		Dest.Empty(Src.Num());
diff --git a/Engine/Plugins/Experimental/Mover/Source/Mover/Private/MovementModifier.cpp b/Engine/Plugins/Experimental/Mover/Source/Mover/Private/MovementModifier.cpp
--- a/Engine/Plugins/Experimental/Mover/Source/Mover/Private/MovementModifier.cpp
+++ b/Engine/Plugins/Experimental/Mover/Source/Mover/Private/MovementModifier.cpp
@@ -138,42 +138,38 @@ TArray<TSharedPtr<FMovementModifierBase>> FMovementModifierGroup::GenerateActive
 
 static void CopyModifierArray(TArray<TSharedPtr<FMovementModifierBase>>& Dest, const TArray<TSharedPtr<FMovementModifierBase>>& Src)
 {
-	// If all of our modifiers types already match, we can copy in place without reallocating memory
-	bool bCanCopyInPlace = (UE::Mover::DisableDataCopyInPlace == 0);
-
+	bool bCanCopyInPlace = (UE::Mover::DisableDataCopyInPlace == 0 && Dest.Num() == Src.Num());
 	if (bCanCopyInPlace)
 	{
-		if (Dest.Num() != Src.Num())
-		{
-			bCanCopyInPlace = false;
-		}
-		else
+		// If copy in place is enabled and the arrays are the same size, copy by index
+		for (int32 i = 0; i < Dest.Num(); ++i)
 		{
-			for (int32 i = 0; i < Dest.Num(); ++i)
+			if (FMovementModifierBase* SrcData = Src[i].Get())
 			{
-				if ((!Dest[i].IsValid() || !Src[i].IsValid()) ||
-					(Dest[i]->GetScriptStruct() != Src[i]->GetScriptStruct()))
+				FMovementModifierBase* DestData = Dest[i].Get();
+				UScriptStruct* SourceStruct = SrcData->GetScriptStruct();
+
+				if (DestData && SourceStruct == DestData->GetScriptStruct())
 				{
-					bCanCopyInPlace = false;
-					break;
+					// Same type so copy in place
+					SourceStruct->CopyScriptStruct(DestData, SrcData, 1);
 				}
+				else
+				{
+					// Different type so replace the shared ptr with a clone
+					Dest[i] = TSharedPtr<FMovementModifierBase>(SrcData->Clone());
+				}
+			}
+			else
+			{
+				// Found invalid source, fall back to full copy
+				bCanCopyInPlace = false;
+				break;
 			}
 		}
 	}
 
-	if (bCanCopyInPlace)
-	{
-		// Copy all structs in place
-		for (int32 i = 0; i < Dest.Num(); ++i)
-		{
-			UScriptStruct* ScriptStruct = Dest[i]->GetScriptStruct();
-			FMovementModifierBase* DestMod = Dest[i].Get();
-			FMovementModifierBase* SrcMod = Src[i].Get();
-
-			ScriptStruct->CopyScriptStruct(DestMod, SrcMod, 1);
-		}
-	}
-	else
+	if (!bCanCopyInPlace)
 	{
 		// Deep copy active modifiers
 		Dest.Empty(Src.Num());
diff --git a/Engine/Plugins/Experimental/Mover/Source/Mover/Private/MoverTypes.cpp b/Engine/Plugins/Experimental/Mover/Source/Mover/Private/MoverTypes.cpp
--- a/Engine/Plugins/Experimental/Mover/Source/Mover/Private/MoverTypes.cpp
+++ b/Engine/Plugins/Experimental/Mover/Source/Mover/Private/MoverTypes.cpp
@@ -154,42 +154,38 @@ FMoverDataCollection& FMoverDataCollection::operator=(const FMoverDataCollection
 	// Perform deep copy of this Group
 	if (this != &Other)
 	{
-		// If all of our data collection types already match, we can copy in place without reallocating memory
-		bool bCanCopyInPlace = (UE::Mover::DisableDataCopyInPlace == 0);
-		
+		bool bCanCopyInPlace = (UE::Mover::DisableDataCopyInPlace == 0 && DataArray.Num() == Other.DataArray.Num());
 		if (bCanCopyInPlace)
 		{
-			if (DataArray.Num() != Other.DataArray.Num())
-			{
-				bCanCopyInPlace = false;
-			}
-			else
+			// If copy in place is enabled and the arrays are the same size, copy by index
+			for (int32 i = 0; i < DataArray.Num(); ++i)
 			{
-				for (int32 i = 0; i < DataArray.Num(); ++i)
+				if (FMoverDataStructBase* SrcData = Other.DataArray[i].Get())
 				{
-					if ((!DataArray[i].IsValid() || !Other.DataArray[i].IsValid()) ||
-						(DataArray[i]->GetScriptStruct() != Other.DataArray[i]->GetScriptStruct()))
+					FMoverDataStructBase* DestData = DataArray[i].Get();
+					UScriptStruct* SourceStruct = SrcData->GetScriptStruct();
+
+					if (DestData && SourceStruct == DestData->GetScriptStruct())
 					{
-						bCanCopyInPlace = false;
-						break;
+						// Same type so copy in place
+						SourceStruct->CopyScriptStruct(DestData, SrcData, 1);
+					}
+					else
+					{
+						// Different type so replace the shared ptr with a clone
+						DataArray[i] = TSharedPtr<FMoverDataStructBase>(SrcData->Clone());
 					}
 				}
+				else
+				{
+					// Found invalid source, fall back to full copy
+					bCanCopyInPlace = false;
+					break;
+				}
 			}
 		}
-
-		if (bCanCopyInPlace)
-		{
-			// Copy all structs in place
-			for (int32 i = 0; i < DataArray.Num(); ++i)
-			{
-				UScriptStruct* ScriptStruct = DataArray[i]->GetScriptStruct();
-				FMoverDataStructBase* Dest = DataArray[i].Get();
-				FMoverDataStructBase* Src = Other.DataArray[i].Get();
-
-				ScriptStruct->CopyScriptStruct(Dest, Src, 1);
-			}
-		}
-		else
+		
+		if (!bCanCopyInPlace)
 		{
 			// Deep copy active data blocks
 			DataArray.Empty(Other.DataArray.Num());
