Mover: Adding AdditiveVelocity MixMode support to MultiJump Layered Move. Also making small adjustments to the extended Mover pawn to make sure additive velocity is respected by turning off extra jump input from CharacterMoverComponent Jump

#virtualized

diff --git a/Engine/Plugins/Experimental/Mover/Source/Mover/Private/DefaultMovementSet/LayeredMoves/MultiJumpLayeredMove.cpp b/Engine/Plugins/Experimental/Mover/Source/Mover/Private/DefaultMovementSet/LayeredMoves/MultiJumpLayeredMove.cpp
--- a/Engine/Plugins/Experimental/Mover/Source/Mover/Private/DefaultMovementSet/LayeredMoves/MultiJumpLayeredMove.cpp
+++ b/Engine/Plugins/Experimental/Mover/Source/Mover/Private/DefaultMovementSet/LayeredMoves/MultiJumpLayeredMove.cpp
@@ -125,17 +125,27 @@ bool FLayeredMove_MultiJump::PerformJump(const FMoverDefaultSyncState* SyncState
 
 	const FVector ImpulseVelocity = UpDir * UpwardsSpeed;
 
-	// Jump impulse overrides vertical velocity while maintaining the rest
-	if (MixMode == EMoveMixMode::OverrideVelocity)
+	switch (MixMode)
 	{
-		const FVector PriorVelocityWS = SyncState->GetVelocity_WorldSpace();
-		const FVector StartingNonUpwardsVelocity = PriorVelocityWS - PriorVelocityWS.ProjectOnToNormal(UpDir);
+	case EMoveMixMode::AdditiveVelocity:
+		{
+			OutProposedMove.LinearVelocity = ImpulseVelocity;
+			break;
+		}
+		
+	case EMoveMixMode::OverrideAll:
+	case EMoveMixMode::OverrideVelocity:
+		{
+			// Jump impulse overrides vertical velocity while maintaining the rest
+			const FVector PriorVelocityWS = SyncState->GetVelocity_WorldSpace();
+			const FVector StartingNonUpwardsVelocity = PriorVelocityWS - PriorVelocityWS.ProjectOnToNormal(UpDir);
 
-		OutProposedMove.LinearVelocity = StartingNonUpwardsVelocity + ImpulseVelocity;
-	}
-	else
-	{
-		ensureMsgf(false, TEXT("Multi-Jump layered move only supports Override Velocity mix mode and was queued with a different mix mode. Layered move will do nothing."));
+			OutProposedMove.LinearVelocity = StartingNonUpwardsVelocity + ImpulseVelocity;
+			break;
+		}
+		
+	default:
+		ensureMsgf(false, TEXT("Multi-Jump layered move has an invalid MixMode and will do nothing."));
 		return false;
 	}
 
