Mover: fix bugs where an actor can violate the min/max distance it's supposed to float above walkable floors:
- on the frame where it transitions from falling to ground-based movement
- on the first frame of falling if close enough to the ground, moving towards it slowly but not enough to hit it
- added option for Flying modes (off by default to prevent perf impact in general case)

diff --git a/Engine/Plugins/Experimental/Mover/Source/Mover/Private/DefaultMovementSet/Modes/AsyncFallingMode.cpp b/Engine/Plugins/Experimental/Mover/Source/Mover/Private/DefaultMovementSet/Modes/AsyncFallingMode.cpp
--- a/Engine/Plugins/Experimental/Mover/Source/Mover/Private/DefaultMovementSet/Modes/AsyncFallingMode.cpp
+++ b/Engine/Plugins/Experimental/Mover/Source/Mover/Private/DefaultMovementSet/Modes/AsyncFallingMode.cpp
@@ -6,6 +6,7 @@
 #include "MoveLibrary/AirMovementUtils.h"
 #include "MoveLibrary/AsyncMovementUtils.h"
 #include "MoveLibrary/BasedMovementUtils.h"
+#include "MoveLibrary/GroundMovementUtils.h"
 #include "MoveLibrary/FloorQueryUtils.h"
 #include "MoveLibrary/MovementUtils.h"
 #include "DefaultMovementSet/Settings/CommonLegacyMovementSettings.h"
@@ -226,6 +227,8 @@ void UAsyncFallingMode::SimulationTick_Implementation(const FSimulationTickParam
 		if (UAirMovementUtils::IsValidLandingSpot(Params.MovingComps, LocationInProgress,
 				SweepHit, CommonLegacySettings->FloorSweepDistance, CommonLegacySettings->MaxWalkSlopeCosine, OUT LandingFloor))
 		{
+			// Adjust height to float slightly above walkable floor
+			LocationInProgress = UGroundMovementUtils::TestMoveToAdjustToFloor(Params.MovingComps, LocationInProgress, RotationInProgress, CommonLegacySettings->MaxWalkSlopeCosine, IN OUT LandingFloor, IN OUT MoveRecord);
 			CaptureFinalState(LocationInProgress, RotationInProgress.Rotator(), LandingFloor, DeltaSeconds, DeltaSeconds * PctTimeApplied, OutputSyncState, OutputState, IN OUT MoveRecord);
 			return;
 		}
@@ -253,8 +256,8 @@ void UAsyncFallingMode::SimulationTick_Implementation(const FSimulationTickParam
 
 		if (LandingFloor.IsWalkableFloor())
 		{
-			CaptureFinalState(LocationInProgress, RotationInProgress.Rotator(), LandingFloor, DeltaSeconds, DeltaSeconds * PctTimeApplied, OutputSyncState, OutputState, MoveRecord);
-			return;
+			// Adjust height to float slightly above walkable floor
+			LocationInProgress = UGroundMovementUtils::TestMoveToAdjustToFloor(Params.MovingComps, LocationInProgress, RotationInProgress, CommonLegacySettings->MaxWalkSlopeCosine, IN OUT LandingFloor, IN OUT MoveRecord);
 		}
 	}
 	else
diff --git a/Engine/Plugins/Experimental/Mover/Source/Mover/Private/DefaultMovementSet/Modes/AsyncFlyingMode.cpp b/Engine/Plugins/Experimental/Mover/Source/Mover/Private/DefaultMovementSet/Modes/AsyncFlyingMode.cpp
--- a/Engine/Plugins/Experimental/Mover/Source/Mover/Private/DefaultMovementSet/Modes/AsyncFlyingMode.cpp
+++ b/Engine/Plugins/Experimental/Mover/Source/Mover/Private/DefaultMovementSet/Modes/AsyncFlyingMode.cpp
@@ -2,6 +2,8 @@
 
 #include "DefaultMovementSet/Modes/AsyncFlyingMode.h"
 #include "MoveLibrary/AirMovementUtils.h"
+#include "MoveLibrary/FloorQueryUtils.h"
+#include "MoveLibrary/GroundMovementUtils.h"
 #include "MoveLibrary/MovementUtils.h"
 #include "MoveLibrary/AsyncMovementUtils.h"
 #include "MoverComponent.h"
@@ -133,6 +135,18 @@ void UAsyncFlyingMode::SimulationTick_Implementation(const FSimulationTickParams
 		}
 	}
 
+	if (bRespectDistanceOverWalkableSurfaces)
+	{
+		// If we are very close to a walkable surface, make sure we maintain a small gap over it
+		FFloorCheckResult FloorUnderActor;
+		UFloorQueryUtils::FindFloor(Params.MovingComps, CommonLegacySettings->FloorSweepDistance, CommonLegacySettings->MaxWalkSlopeCosine, LocationInProgress, FloorUnderActor);
+
+		if (FloorUnderActor.IsWalkableFloor())
+		{
+			LocationInProgress = UGroundMovementUtils::TestMoveToAdjustToFloor(Params.MovingComps, LocationInProgress, RotationInProgress, CommonLegacySettings->MaxWalkSlopeCosine, IN OUT FloorUnderActor, MoveRecord);
+		}
+	}
+
 	OutputSyncState.SetTransforms_WorldSpace(
 		LocationInProgress,
 		RotationInProgress.Rotator(),
diff --git a/Engine/Plugins/Experimental/Mover/Source/Mover/Private/DefaultMovementSet/Modes/FallingMode.cpp b/Engine/Plugins/Experimental/Mover/Source/Mover/Private/DefaultMovementSet/Modes/FallingMode.cpp
--- a/Engine/Plugins/Experimental/Mover/Source/Mover/Private/DefaultMovementSet/Modes/FallingMode.cpp
+++ b/Engine/Plugins/Experimental/Mover/Source/Mover/Private/DefaultMovementSet/Modes/FallingMode.cpp
@@ -6,6 +6,7 @@
 #include "MoveLibrary/MovementUtils.h"
 #include "MoveLibrary/BasedMovementUtils.h"
 #include "MoveLibrary/FloorQueryUtils.h"
+#include "MoveLibrary/GroundMovementUtils.h"
 #include "DefaultMovementSet/Settings/CommonLegacyMovementSettings.h"
 #include "MoveLibrary/AirMovementUtils.h"
 
@@ -176,6 +177,24 @@ void UFallingMode::SimulationTick_Implementation(const FSimulationTickParams& Pa
 	
 	const FVector StartingFallingVelocity = StartingSyncState->GetVelocity_WorldSpace();
 
+	const FVector UpDirection = MoverComponent->GetUpDirection();
+
+	// If we're just starting to fall downwards, check for a floor first
+	if (FMath::IsNearlyZero(StartingFallingVelocity.Dot(UpDirection)) && !ProposedMove.LinearVelocity.IsNearlyZero() &&
+		ProposedMove.LinearVelocity.Dot(MoverComponent->GetUpDirection()) < 0.0)
+	{
+		// If we are very close to a walkable floor, make sure we're maintaining the correct distance from it
+		FFloorCheckResult FloorUnderActor;
+		UFloorQueryUtils::FindFloor(Params.MovingComps, CommonLegacySettings->FloorSweepDistance, CommonLegacySettings->MaxWalkSlopeCosine, UpdatedComponent->GetComponentLocation(), OUT FloorUnderActor);
+
+		if (FloorUnderActor.IsWalkableFloor())
+		{
+			UGroundMovementUtils::TryMoveToAdjustHeightAboveFloor(MoverComponent, FloorUnderActor, CommonLegacySettings->MaxWalkSlopeCosine, MoveRecord);
+			CaptureFinalState(UpdatedComponent, *StartingSyncState, FloorUnderActor, DeltaSeconds, DeltaSeconds * PctTimeApplied, OutputSyncState, OutputState, MoveRecord);
+			return;
+		}
+	}
+
 	//FVector MoveDelta = 0.5f * (PriorFallingVelocity + ProposedMove.LinearVelocity) * DeltaSeconds; 	// TODO: revive midpoint integration
 	FVector MoveDelta = ProposedMove.LinearVelocity * DeltaSeconds;
 
@@ -198,6 +217,7 @@ void UFallingMode::SimulationTick_Implementation(const FSimulationTickParams& Pa
 		if (UAirMovementUtils::IsValidLandingSpot(Params.MovingComps, UpdatedComponent->GetComponentLocation(),
 			Hit, CommonLegacySettings->FloorSweepDistance, CommonLegacySettings->MaxWalkSlopeCosine, OUT LandingFloor))
 		{
+			UGroundMovementUtils::TryMoveToAdjustHeightAboveFloor(MoverComponent, LandingFloor, CommonLegacySettings->MaxWalkSlopeCosine, MoveRecord); // make sure we maintain a small gap over walking surfaces
 			CaptureFinalState(UpdatedComponent, *StartingSyncState, LandingFloor, DeltaSeconds, DeltaSeconds * PctTimeApplied, OutputSyncState, OutputState, MoveRecord);
 			return;
 		}
@@ -217,6 +237,7 @@ void UFallingMode::SimulationTick_Implementation(const FSimulationTickParams& Pa
 
 		if (LandingFloor.IsWalkableFloor())
 		{
+			UGroundMovementUtils::TryMoveToAdjustHeightAboveFloor(MoverComponent, LandingFloor, CommonLegacySettings->MaxWalkSlopeCosine, MoveRecord); // make sure we maintain a small gap over walking surfaces
 			CaptureFinalState(UpdatedComponent, *StartingSyncState, LandingFloor, DeltaSeconds, DeltaSeconds * PctTimeApplied, OutputSyncState, OutputState, MoveRecord);
 			return;
 		}
diff --git a/Engine/Plugins/Experimental/Mover/Source/Mover/Private/DefaultMovementSet/Modes/FlyingMode.cpp b/Engine/Plugins/Experimental/Mover/Source/Mover/Private/DefaultMovementSet/Modes/FlyingMode.cpp
--- a/Engine/Plugins/Experimental/Mover/Source/Mover/Private/DefaultMovementSet/Modes/FlyingMode.cpp
+++ b/Engine/Plugins/Experimental/Mover/Source/Mover/Private/DefaultMovementSet/Modes/FlyingMode.cpp
@@ -2,6 +2,8 @@
 
 #include "DefaultMovementSet/Modes/FlyingMode.h"
 #include "MoveLibrary/AirMovementUtils.h"
+#include "MoveLibrary/FloorQueryUtils.h"
+#include "MoveLibrary/GroundMovementUtils.h"
 #include "MoveLibrary/MovementUtils.h"
 #include "MoverComponent.h"
 #include "DefaultMovementSet/Settings/CommonLegacyMovementSettings.h"
@@ -68,7 +70,7 @@ void UFlyingMode::GenerateMove_Implementation(const FMoverTickStartData& StartSt
 
 void UFlyingMode::SimulationTick_Implementation(const FSimulationTickParams& Params, FMoverTickEndData& OutputState)
 {
-	const UMoverComponent* MoverComp = GetMoverComponent();
+	UMoverComponent* MoverComp = GetMoverComponent();
 	const FMoverTickStartData& StartState = Params.StartState;
 	USceneComponent* UpdatedComponent = Params.MovingComps.UpdatedComponent.Get();
 	FProposedMove ProposedMove = Params.ProposedMove;
@@ -115,11 +117,23 @@ void UFlyingMode::SimulationTick_Implementation(const FSimulationTickParams& Par
 
 	if (Hit.IsValidBlockingHit())
 	{
-		UMoverComponent* MoverComponent = GetMoverComponent();
 		FMoverOnImpactParams ImpactParams(DefaultModeNames::Flying, Hit, MoveDelta);
-		MoverComponent->HandleImpact(ImpactParams);
+		MoverComp->HandleImpact(ImpactParams);
+
 		// Try to slide the remaining distance along the surface.
-		UMovementUtils::TryMoveToSlideAlongSurface(FMovingComponentSet(MoverComponent), MoveDelta, 1.f - Hit.Time, OrientQuat, Hit.Normal, Hit, true, MoveRecord);
+		UMovementUtils::TryMoveToSlideAlongSurface(FMovingComponentSet(MoverComp), MoveDelta, 1.f - Hit.Time, OrientQuat, Hit.Normal, Hit, true, MoveRecord);
+	}
+
+	if (bRespectDistanceOverWalkableSurfaces)
+	{
+		// If we are very close to a walkable surface, make sure we maintain a small gap over it
+		FFloorCheckResult FloorUnderActor;
+		UFloorQueryUtils::FindFloor(Params.MovingComps, CommonLegacySettings->FloorSweepDistance, CommonLegacySettings->MaxWalkSlopeCosine, UpdatedComponent->GetComponentLocation(), OUT FloorUnderActor);
+
+		if (FloorUnderActor.IsWalkableFloor())
+		{
+			UGroundMovementUtils::TryMoveToAdjustHeightAboveFloor(MoverComp, FloorUnderActor, CommonLegacySettings->MaxWalkSlopeCosine, MoveRecord);
+		}
 	}
 
 	CaptureFinalState(UpdatedComponent, MoveRecord, *StartingSyncState, OutputSyncState, DeltaSeconds);
diff --git a/Engine/Plugins/Experimental/Mover/Source/Mover/Private/MoveLibrary/FloorQueryUtils.cpp b/Engine/Plugins/Experimental/Mover/Source/Mover/Private/MoveLibrary/FloorQueryUtils.cpp
--- a/Engine/Plugins/Experimental/Mover/Source/Mover/Private/MoveLibrary/FloorQueryUtils.cpp
+++ b/Engine/Plugins/Experimental/Mover/Source/Mover/Private/MoveLibrary/FloorQueryUtils.cpp
@@ -9,8 +9,10 @@
 
 namespace UE::FloorQueryUtility
 {
-	const float MIN_FLOOR_DIST = 1.9f;
-	const float MAX_FLOOR_DIST = 2.4f;
+
+	const float MIN_FLOOR_DIST = 1.9f;	// Smallest distance we want our primitive floating above walkable floors while in ground-based movement.
+	const float MAX_FLOOR_DIST = 2.4f;	// Largest distance we want our primitive floating above walkable floors while in ground-based movement.
+
 	const float SWEEP_EDGE_REJECT_DISTANCE = 0.15f;
 }
 
diff --git a/Engine/Plugins/Experimental/Mover/Source/Mover/Public/DefaultMovementSet/Modes/AsyncFlyingMode.h b/Engine/Plugins/Experimental/Mover/Source/Mover/Public/DefaultMovementSet/Modes/AsyncFlyingMode.h
--- a/Engine/Plugins/Experimental/Mover/Source/Mover/Public/DefaultMovementSet/Modes/AsyncFlyingMode.h
+++ b/Engine/Plugins/Experimental/Mover/Source/Mover/Public/DefaultMovementSet/Modes/AsyncFlyingMode.h
@@ -23,6 +23,13 @@ class UAsyncFlyingMode : public UBaseMovementMode
 
 
 public:
+	/**
+	 * If true, the actor will 'float' above any walkable surfaces to maintain the same height as ground-based modes.
+	 * This can prevent pops when transitioning to ground-based movement, at the cost of performing floor checks while flying.
+	 */
+	UPROPERTY(Category = Mover, EditAnywhere, BlueprintReadWrite)
+	bool bRespectDistanceOverWalkableSurfaces = false;
+
 	UE_API virtual void OnRegistered(const FName ModeName) override;
 	UE_API virtual void OnUnregistered() override;
 	UE_API virtual void GenerateMove_Implementation(const FMoverTickStartData& StartState, const FMoverTimeStep& TimeStep, FProposedMove& OutProposedMove) const override;
diff --git a/Engine/Plugins/Experimental/Mover/Source/Mover/Public/DefaultMovementSet/Modes/FlyingMode.h b/Engine/Plugins/Experimental/Mover/Source/Mover/Public/DefaultMovementSet/Modes/FlyingMode.h
--- a/Engine/Plugins/Experimental/Mover/Source/Mover/Public/DefaultMovementSet/Modes/FlyingMode.h
+++ b/Engine/Plugins/Experimental/Mover/Source/Mover/Public/DefaultMovementSet/Modes/FlyingMode.h
@@ -22,6 +22,13 @@ class UFlyingMode : public UBaseMovementMode
 
 
 public:
+	/**
+	 * If true, the actor will 'float' above any walkable surfaces to maintain the same height as ground-based modes. 
+	 * This can prevent pops when transitioning to ground-based movement, at the cost of performing floor checks while flying.
+	 */
+	UPROPERTY(Category = Mover, EditAnywhere, BlueprintReadWrite)
+	bool bRespectDistanceOverWalkableSurfaces = false;
+
 	UE_API virtual void GenerateMove_Implementation(const FMoverTickStartData& StartState, const FMoverTimeStep& TimeStep, FProposedMove& OutProposedMove) const override;
 
 	UE_API virtual void SimulationTick_Implementation(const FSimulationTickParams& Params, FMoverTickEndData& OutputState) override;
