Mover: fixing static analysis warning about using a pointer before a null check

#rnx

diff --git a/Engine/Plugins/Experimental/Mover/Source/Mover/Private/DefaultMovementSet/Modes/FallingMode.cpp b/Engine/Plugins/Experimental/Mover/Source/Mover/Private/DefaultMovementSet/Modes/FallingMode.cpp
--- a/Engine/Plugins/Experimental/Mover/Source/Mover/Private/DefaultMovementSet/Modes/FallingMode.cpp
+++ b/Engine/Plugins/Experimental/Mover/Source/Mover/Private/DefaultMovementSet/Modes/FallingMode.cpp
@@ -147,6 +147,11 @@ void UFallingMode::SimulationTick_Implementation(const FSimulationTickParams& Pa
 	USceneComponent* UpdatedComponent = Params.MovingComps.UpdatedComponent.Get();
 	FProposedMove ProposedMove = Params.ProposedMove;
 
+	if (!UpdatedComponent)
+	{
+		return;
+	}
+
 	const FCharacterDefaultInputs* CharacterInputs = StartState.InputCmd.InputCollection.FindDataByType<FCharacterDefaultInputs>();
 	const FMoverDefaultSyncState* StartingSyncState = StartState.SyncState.SyncStateCollection.FindDataByType<FMoverDefaultSyncState>();
 	check(StartingSyncState);
diff --git a/Engine/Plugins/Experimental/Mover/Source/Mover/Private/DefaultMovementSet/Modes/FlyingMode.cpp b/Engine/Plugins/Experimental/Mover/Source/Mover/Private/DefaultMovementSet/Modes/FlyingMode.cpp
--- a/Engine/Plugins/Experimental/Mover/Source/Mover/Private/DefaultMovementSet/Modes/FlyingMode.cpp
+++ b/Engine/Plugins/Experimental/Mover/Source/Mover/Private/DefaultMovementSet/Modes/FlyingMode.cpp
@@ -75,6 +75,11 @@ void UFlyingMode::SimulationTick_Implementation(const FSimulationTickParams& Par
 	USceneComponent* UpdatedComponent = Params.MovingComps.UpdatedComponent.Get();
 	FProposedMove ProposedMove = Params.ProposedMove;
 
+	if (!UpdatedComponent)
+	{
+		return;
+	}
+
 	const FCharacterDefaultInputs* CharacterInputs = StartState.InputCmd.InputCollection.FindDataByType<FCharacterDefaultInputs>();
 	const FMoverDefaultSyncState* StartingSyncState = StartState.SyncState.SyncStateCollection.FindDataByType<FMoverDefaultSyncState>();
 	check(StartingSyncState);
