Mover: changing init flow so that the state machine is available before BeginPlay.
This allows us to remove the InitMoverSimulation call from the backends and eliminates the need for delays in queuing movement changes.

diff --git a/Engine/Plugins/Experimental/ChaosMover/Source/ChaosMover/Private/Backends/ChaosMoverBackend.cpp b/Engine/Plugins/Experimental/ChaosMover/Source/ChaosMover/Private/Backends/ChaosMoverBackend.cpp
--- a/Engine/Plugins/Experimental/ChaosMover/Source/ChaosMover/Private/Backends/ChaosMoverBackend.cpp
+++ b/Engine/Plugins/Experimental/ChaosMover/Source/ChaosMover/Private/Backends/ChaosMoverBackend.cpp
@@ -246,9 +246,6 @@ void UChaosMoverBackendComponent::BeginPlay()
 
 	if (UWorld* World = GetWorld(); World && World->IsGameWorld())
 	{
-		// This is needed for now to create the mover component blackboard
-		GetMoverComponent().InitMoverSimulation();
-
 		CreatePhysics();
 		InitSimulation();
 
diff --git a/Engine/Plugins/Experimental/Mover/Source/Mover/Private/Backends/MoverNetworkPhysicsLiaisonBase.cpp b/Engine/Plugins/Experimental/Mover/Source/Mover/Private/Backends/MoverNetworkPhysicsLiaisonBase.cpp
--- a/Engine/Plugins/Experimental/Mover/Source/Mover/Private/Backends/MoverNetworkPhysicsLiaisonBase.cpp
+++ b/Engine/Plugins/Experimental/Mover/Source/Mover/Private/Backends/MoverNetworkPhysicsLiaisonBase.cpp
@@ -309,7 +309,6 @@ void UMoverNetworkPhysicsLiaisonComponentBase::InitializeComponent()
 	if (UWorld* World = GetWorld(); World && World->IsGameWorld())
 	{
 		UMoverComponent& MoverComp = GetMoverComponent();
-		MoverComp.InitMoverSimulation();
 		MoverComp.ModeFSM->SetModeImmediately(MoverComp.StartingMovementMode);
 	}
 
diff --git a/Engine/Plugins/Experimental/Mover/Source/Mover/Private/Backends/MoverNetworkPredictionLiaison.cpp b/Engine/Plugins/Experimental/Mover/Source/Mover/Private/Backends/MoverNetworkPredictionLiaison.cpp
--- a/Engine/Plugins/Experimental/Mover/Source/Mover/Private/Backends/MoverNetworkPredictionLiaison.cpp
+++ b/Engine/Plugins/Experimental/Mover/Source/Mover/Private/Backends/MoverNetworkPredictionLiaison.cpp
@@ -259,8 +259,6 @@ void UMoverNetworkPredictionLiaisonComponent::InitializeNetworkPredictionProxy()
 
 	if (ensureAlwaysMsgf(MoverComp, TEXT("UMoverNetworkPredictionLiaisonComponent on actor %s failed to find associated Mover component. This actor's movement will not be simulated. Verify its setup."), *GetNameSafe(GetOwner())))
 	{
-		MoverComp->InitMoverSimulation();
-
 		NetworkPredictionProxy.Init<FMoverActorModelDef>(GetWorld(), GetReplicationProxies(), this, this);
 	}
 }
diff --git a/Engine/Plugins/Experimental/Mover/Source/Mover/Private/Backends/MoverStandaloneLiaison.cpp b/Engine/Plugins/Experimental/Mover/Source/Mover/Private/Backends/MoverStandaloneLiaison.cpp
--- a/Engine/Plugins/Experimental/Mover/Source/Mover/Private/Backends/MoverStandaloneLiaison.cpp
+++ b/Engine/Plugins/Experimental/Mover/Source/Mover/Private/Backends/MoverStandaloneLiaison.cpp
@@ -137,8 +137,6 @@ void UMoverStandaloneLiaisonComponent::BeginPlay()
 		{
 			MoverComp = FoundMoverComp;
 
-			MoverComp->InitMoverSimulation();
-
 			FRWScopeLock Lock(StateDataLock, SLT_Write);
 			MoverComp->InitializeSimulationState(OUT &CachedLastSyncState, OUT &CachedLastAuxState);
 		}
diff --git a/Engine/Plugins/Experimental/Mover/Source/Mover/Private/MoverComponent.cpp b/Engine/Plugins/Experimental/Mover/Source/Mover/Private/MoverComponent.cpp
--- a/Engine/Plugins/Experimental/Mover/Source/Mover/Private/MoverComponent.cpp
+++ b/Engine/Plugins/Experimental/Mover/Source/Mover/Private/MoverComponent.cpp
@@ -112,6 +112,42 @@ void UMoverComponent::InitializeComponent()
 		{
 			UE_LOG(LogMover, Error, TEXT("No backend class set on %s. Mover actor will not function."), *GetNameSafe(GetOwner()));
 		}
+
+		// Set up FSM and initial movement states
+		ModeFSM = NewObject<UMovementModeStateMachine>(this, TEXT("MoverStateMachine"), RF_Transient);
+		ModeFSM->ClearAllMovementModes();
+		ModeFSM->ClearAllGlobalTransitions();
+
+		bool bHasMatchingStartingState = false;
+
+		for (const TPair<FName, TObjectPtr<UBaseMovementMode>>& Element : MovementModes)
+		{
+			if (Element.Value.Get() == nullptr)
+			{
+				UE_LOG(LogMover, Warning, TEXT("Invalid Movement Mode type '%s' detected on %s. Mover actor will not function correctly."),
+					*Element.Key.ToString(), *GetNameSafe(GetOwner()));
+				continue;
+			}
+
+			ModeFSM->RegisterMovementMode(Element.Key, Element.Value);
+
+			bHasMatchingStartingState |= (StartingMovementMode == Element.Key);
+		}
+
+		for (TObjectPtr<UBaseMovementModeTransition>& Transition : Transitions)
+		{
+			ModeFSM->RegisterGlobalTransition(Transition);
+		}
+
+		UE_CLOG(!bHasMatchingStartingState, LogMover, Warning, TEXT("Invalid StartingMovementMode '%s' specified on %s. Mover actor will not function."),
+			*StartingMovementMode.ToString(), *GetNameSafe(GetOwner()));
+
+		if (bHasMatchingStartingState && StartingMovementMode != NAME_None)
+		{
+			ModeFSM->SetDefaultMode(StartingMovementMode);
+			ModeFSM->QueueNextMode(StartingMovementMode);
+		}
+
 	}
 
 	// Gather initial state to fulfill queries
@@ -126,14 +162,20 @@ void UMoverComponent::UninitializeComponent()
 	if (UActorComponent* LiaisonAsComp = Cast<UActorComponent>(BackendLiaisonComp.GetObject()))
 	{
 		LiaisonAsComp->DestroyComponent();
-		BackendLiaisonComp = nullptr;
 	}
+	BackendLiaisonComp = nullptr;
 
 	if (SimBlackboard)
 	{
 		SimBlackboard->InvalidateAll();
 	}
 
+	if (ModeFSM)
+	{
+		ModeFSM->ClearAllMovementModes();
+		ModeFSM->ClearAllGlobalTransitions();
+	}
+
 	Super::UninitializeComponent();
 }
 
@@ -255,13 +297,6 @@ void UMoverComponent::BeginPlay()
 	}
 }
 
-void UMoverComponent::EndPlay(const EEndPlayReason::Type EndPlayReason)
-{
-	ModeFSM->ClearAllMovementModes();
-	
-	Super::EndPlay(EndPlayReason);
-}
-
 void UMoverComponent::BindProcessGeneratedMovement(FMover_ProcessGeneratedMovement ProcessGeneratedMovementEvent)
 {
 	ProcessGeneratedMovement = ProcessGeneratedMovementEvent;
@@ -810,51 +845,6 @@ void UMoverComponent::SetFrameStateFromContext(const FMoverSyncState* SyncState,
 }
 
 
-bool UMoverComponent::InitMoverSimulation()
-{
-	check(UpdatedComponent);
-
-	// Set up FSM and initial movement states
-	ModeFSM = NewObject<UMovementModeStateMachine>(this, TEXT("MoverStateMachine"), RF_Transient);
-	ModeFSM->ClearAllMovementModes();
-	ModeFSM->ClearAllGlobalTransitions();
-
-	bool bHasMatchingStartingState = false;
-
-	for (const TPair<FName, TObjectPtr<UBaseMovementMode>>& Element : MovementModes)
-	{
-		if (Element.Value.Get() == nullptr)
-		{
-			UE_LOG(LogMover, Warning, TEXT("Invalid Movement Mode type '%s' detected on %s. Mover actor will not function correctly."),
-				*Element.Key.ToString(), *GetNameSafe(GetOwner()));
-			continue;
-		}
-
-		ModeFSM->RegisterMovementMode(Element.Key, Element.Value);
-
-		bHasMatchingStartingState |= (StartingMovementMode == Element.Key);
-	}
-
-	for (TObjectPtr<UBaseMovementModeTransition>& Transition : Transitions)
-	{
-		ModeFSM->RegisterGlobalTransition(Transition);
-	}
-
-	if (!bHasMatchingStartingState)
-	{
-		UE_LOG(LogMover, Warning, TEXT("Invalid StartingMovementMode '%s' specified on %s. Mover actor will not function."),
-			*StartingMovementMode.ToString(), *GetNameSafe(GetOwner()));
-	}
-
-	if (bHasMatchingStartingState && StartingMovementMode != NAME_None)
-	{
-		ModeFSM->SetDefaultMode(StartingMovementMode);
-		ModeFSM->QueueNextMode(StartingMovementMode);
-	}
-
-	return bHasMatchingStartingState;
-}
-
 void UMoverComponent::CreateDefaultInputAndState(FMoverInputCmdContext& OutInputCmd, FMoverSyncState& OutSyncState, FMoverAuxStateContext& OutAuxState) const
 {
 	OutInputCmd = FMoverInputCmdContext();
diff --git a/Engine/Plugins/Experimental/Mover/Source/Mover/Public/MoverComponent.h b/Engine/Plugins/Experimental/Mover/Source/Mover/Public/MoverComponent.h
--- a/Engine/Plugins/Experimental/Mover/Source/Mover/Public/MoverComponent.h
+++ b/Engine/Plugins/Experimental/Mover/Source/Mover/Public/MoverComponent.h
@@ -73,7 +73,6 @@ public:
 	MOVER_API virtual void RegisterComponentTickFunctions(bool bRegister) override;
 	MOVER_API virtual void PostLoad() override;
 	MOVER_API virtual void BeginPlay() override;
-	MOVER_API virtual void EndPlay(const EEndPlayReason::Type EndPlayReason) override;
 	
 	// Broadcast before each simulation tick.
 	// Note - Guaranteed to run on the game thread (even in async simulation).
@@ -520,8 +519,6 @@ protected:
 	MOVER_API void UpdateCachedFrameState(const FMoverSyncState* SyncState, const FMoverAuxStateContext* AuxState);
 
 public:
-	MOVER_API bool InitMoverSimulation();	// Intended to be called by the backend only
-
 	MOVER_API virtual void CreateDefaultInputAndState(FMoverInputCmdContext& OutInputCmd, FMoverSyncState& OutSyncState, FMoverAuxStateContext& OutAuxState) const;
 
 	/** Handle a blocking impact.*/
