Mover: Making NavWalking mode set floor results in the Mover Blackboard when projecting to level geometry is enabled.

diff --git a/Engine/Plugins/Experimental/Mover/Source/Mover/Private/DefaultMovementSet/Modes/AsyncNavWalkingMode.cpp b/Engine/Plugins/Experimental/Mover/Source/Mover/Private/DefaultMovementSet/Modes/AsyncNavWalkingMode.cpp
--- a/Engine/Plugins/Experimental/Mover/Source/Mover/Private/DefaultMovementSet/Modes/AsyncNavWalkingMode.cpp
+++ b/Engine/Plugins/Experimental/Mover/Source/Mover/Private/DefaultMovementSet/Modes/AsyncNavWalkingMode.cpp
@@ -300,6 +300,13 @@ void UAsyncNavWalkingMode::SimulationTick_Implementation(const FSimulationTickPa
 			const float DownOffset = TotalCapsuleHeight * FMath::Max(0.f, NavMeshProjectionHeightScaleDown);
 			NewFeetLocation = ProjectLocationFromNavMesh(DeltaSeconds, StartingFeetLocation, NewFeetLocation, UpOffset, DownOffset);
 		}
+		else
+		{
+			if (UMoverBlackboard* SimBlackboard = MoverComp->GetSimBlackboard_Mutable())
+			{
+				SimBlackboard->Invalidate(CommonBlackboard::LastFloorResult);
+			}
+		}
 
 		FVector AdjustedDelta = NewFeetLocation - StartingFeetLocation;
 
@@ -374,7 +381,7 @@ void UAsyncNavWalkingMode::Activate()
 	{
 		if (UMoverBlackboard* SimBlackboard = MoverComp->GetSimBlackboard_Mutable())
 		{
-			SimBlackboard->Invalidate(NavWalkingBlackBoard::ProjectedNavMeshHitResult);
+			SimBlackboard->Invalidate(CommonBlackboard::LastFloorResult);
 		}
 	}
 
@@ -442,13 +449,6 @@ void UAsyncNavWalkingMode::FindBestNavMeshLocation(const FVector& TraceStart, co
 		{
 			UPrimitiveComponent* PrimComp = TestHit.GetComponent();
 			// Prefer using primitive component if valid
-			if (IPhysicsBodyInstanceOwner* BodyInstanceOwner = !PrimComp ? IPhysicsBodyInstanceOwner::GetPhysicsBodyInstandeOwnerFromHitResult(TestHit) : nullptr)
-			{
-				const bool bBlockOnWorldStatic = (BodyInstanceOwner->GetCollisionResponseToChannel(ECC_WorldStatic) == ECR_Block);
-				const bool bBlockOnWorldDynamic = (BodyInstanceOwner->GetCollisionResponseToChannel(ECC_WorldDynamic) == ECR_Block);
-				return !bBlockOnWorldStatic && (!bBlockOnWorldDynamic || bCheckOnlyWorldStatic);
-			}
-			else
 			{
 				const bool bBlockOnWorldStatic = PrimComp && (PrimComp->GetCollisionResponseToChannel(ECC_WorldStatic) == ECR_Block);
 				const bool bBlockOnWorldDynamic = PrimComp && (PrimComp->GetCollisionResponseToChannel(ECC_WorldDynamic) == ECR_Block);
@@ -487,9 +487,10 @@ FVector UAsyncNavWalkingMode::ProjectLocationFromNavMesh(float DeltaSeconds, con
 	const FVector TraceStart = TargetNavLocation + UpOffset * UpDirection;
 	const FVector TraceEnd = TargetNavLocation + DownOffset * -UpDirection;
 
-	FHitResult CachedProjectedNavMeshHitResult;
+	FFloorCheckResult CachedFloorCheckResult;
 	UMoverBlackboard* SimBlackboard = MoverComp->GetSimBlackboard_Mutable();
-	SimBlackboard->TryGet(NavWalkingBlackBoard::ProjectedNavMeshHitResult, CachedProjectedNavMeshHitResult);
+	SimBlackboard->TryGet(CommonBlackboard::LastFloorResult, CachedFloorCheckResult);
+	FHitResult CachedProjectedNavMeshHitResult = CachedFloorCheckResult.HitResult;
 
 	// We can skip this trace if we are checking at the same location as the last trace (ie, we haven't moved).
 	const bool bCachedLocationStillValid = (CachedProjectedNavMeshHitResult.bBlockingHit &&
@@ -510,12 +511,21 @@ FVector UAsyncNavWalkingMode::ProjectLocationFromNavMesh(float DeltaSeconds, con
 			if (HitResult.bStartPenetrating || !HitResult.bBlockingHit)
 			{
 				CachedProjectedNavMeshHitResult.Reset();
-				SimBlackboard->Invalidate(NavWalkingBlackBoard::ProjectedNavMeshHitResult);
+				SimBlackboard->Invalidate(CommonBlackboard::LastFloorResult);
 			}
 			else
 			{
 				CachedProjectedNavMeshHitResult = HitResult;
-				SimBlackboard->Set(NavWalkingBlackBoard::ProjectedNavMeshHitResult, CachedProjectedNavMeshHitResult);
+
+				FFloorCheckResult FloorCheckResult;
+				FloorCheckResult.bBlockingHit = HitResult.bBlockingHit;
+				FloorCheckResult.bLineTrace = true;
+				FloorCheckResult.bWalkableFloor = true;
+				FloorCheckResult.LineDist = FMath::Abs((CurrentFeetLocation - CachedProjectedNavMeshHitResult.ImpactPoint).Dot(UpDirection));
+				FloorCheckResult.FloorDist = FloorCheckResult.LineDist; // This is usually set from a sweep trace but it doesn't really hurt setting it. 
+				FloorCheckResult.HitResult = CachedProjectedNavMeshHitResult;
+
+				SimBlackboard->Set(CommonBlackboard::LastFloorResult, FloorCheckResult);
 			}
 		}
 		else
diff --git a/Engine/Plugins/Experimental/Mover/Source/Mover/Private/DefaultMovementSet/Modes/NavWalkingMode.cpp b/Engine/Plugins/Experimental/Mover/Source/Mover/Private/DefaultMovementSet/Modes/NavWalkingMode.cpp
--- a/Engine/Plugins/Experimental/Mover/Source/Mover/Private/DefaultMovementSet/Modes/NavWalkingMode.cpp
+++ b/Engine/Plugins/Experimental/Mover/Source/Mover/Private/DefaultMovementSet/Modes/NavWalkingMode.cpp
@@ -294,6 +294,13 @@ void UNavWalkingMode::SimulationTick_Implementation(const FSimulationTickParams&
 			const float DownOffset = TotalCapsuleHeight * FMath::Max(0.f, NavMeshProjectionHeightScaleDown);
 			NewLocation = ProjectLocationFromNavMesh(DeltaSeconds, StartingLocation, NewLocation, UpOffset, DownOffset);
 		}
+		else
+		{
+			if (UMoverBlackboard* SimBlackboard = MoverComp->GetSimBlackboard_Mutable())
+			{
+				SimBlackboard->Invalidate(CommonBlackboard::LastFloorResult);
+			}	
+		}
 
 		FVector AdjustedDelta = NewLocation - StartingLocation;
 		
@@ -364,7 +371,7 @@ void UNavWalkingMode::SetCollisionForNavWalking(bool bEnable)
 				{
 					if (UMoverBlackboard* SimBlackboard = MoverComponent->GetSimBlackboard_Mutable())
 					{
-						SimBlackboard->Invalidate(NavWalkingBlackBoard::ProjectedNavMeshHitResult);
+						SimBlackboard->Invalidate(CommonBlackboard::LastFloorResult);
 					}
 
 					// Stagger timed updates so many different characters spawned at the same time don't update on the same frame.
@@ -495,9 +502,10 @@ FVector UNavWalkingMode::ProjectLocationFromNavMesh(float DeltaSeconds, const FV
 	const FVector TraceStart = TargetNavLocation + UpOffset * UpDirection;
 	const FVector TraceEnd   = TargetNavLocation + DownOffset * -UpDirection;
 
-	FHitResult CachedProjectedNavMeshHitResult;
+	FFloorCheckResult CachedFloorCheckResult;
 	UMoverBlackboard* SimBlackboard = MoverComp->GetSimBlackboard_Mutable();
-	SimBlackboard->TryGet(NavWalkingBlackBoard::ProjectedNavMeshHitResult, CachedProjectedNavMeshHitResult);
+	SimBlackboard->TryGet(CommonBlackboard::LastFloorResult, CachedFloorCheckResult);
+	FHitResult CachedProjectedNavMeshHitResult = CachedFloorCheckResult.HitResult;
 	
 	// We can skip this trace if we are checking at the same location as the last trace (ie, we haven't moved).
 	const bool bCachedLocationStillValid = (CachedProjectedNavMeshHitResult.bBlockingHit &&
@@ -518,12 +526,21 @@ FVector UNavWalkingMode::ProjectLocationFromNavMesh(float DeltaSeconds, const FV
 			if (HitResult.bStartPenetrating || !HitResult.bBlockingHit)
 			{
 				CachedProjectedNavMeshHitResult.Reset();
-				SimBlackboard->Invalidate(NavWalkingBlackBoard::ProjectedNavMeshHitResult);
+				SimBlackboard->Invalidate(CommonBlackboard::LastFloorResult);
 			}
 			else
 			{
 				CachedProjectedNavMeshHitResult = HitResult;
-				SimBlackboard->Set(NavWalkingBlackBoard::ProjectedNavMeshHitResult, CachedProjectedNavMeshHitResult);
+				
+				FFloorCheckResult FloorCheckResult;
+				FloorCheckResult.bBlockingHit = HitResult.bBlockingHit;
+				FloorCheckResult.bLineTrace = true;
+				FloorCheckResult.bWalkableFloor = true;
+				FloorCheckResult.LineDist = FMath::Abs((CurrentFeetLocation - CachedProjectedNavMeshHitResult.ImpactPoint).Dot(UpDirection));
+				FloorCheckResult.FloorDist = FloorCheckResult.LineDist; // This is usually set from a sweep trace but it doesn't really hurt setting it. 
+				FloorCheckResult.HitResult = HitResult;
+
+				SimBlackboard->Set(CommonBlackboard::LastFloorResult, FloorCheckResult);
 			}
 		}
 		else
diff --git a/Engine/Plugins/Experimental/Mover/Source/Mover/Public/DefaultMovementSet/Modes/NavWalkingMode.h b/Engine/Plugins/Experimental/Mover/Source/Mover/Public/DefaultMovementSet/Modes/NavWalkingMode.h
--- a/Engine/Plugins/Experimental/Mover/Source/Mover/Public/DefaultMovementSet/Modes/NavWalkingMode.h
+++ b/Engine/Plugins/Experimental/Mover/Source/Mover/Public/DefaultMovementSet/Modes/NavWalkingMode.h
@@ -15,10 +15,6 @@ class UNavMoverComponent;
 class INavigationDataInterface;
 class UCommonLegacyMovementSettings;
 
-namespace NavWalkingBlackBoard
-{
-	const FName ProjectedNavMeshHitResult = TEXT("ProjectedNavMeshHitResult");
-}
 
 /** Options for what to do when we find ourselves off the nav mesh */
 UENUM(BlueprintType)
