Backout
Mover: Making NavWalking mode set floor results in the Mover Blackboard when projecting to level geometry is enabled.

diff --git a/Engine/Plugins/Experimental/Mover/Source/Mover/Private/DefaultMovementSet/Modes/AsyncNavWalkingMode.cpp b/Engine/Plugins/Experimental/Mover/Source/Mover/Private/DefaultMovementSet/Modes/AsyncNavWalkingMode.cpp
--- a/Engine/Plugins/Experimental/Mover/Source/Mover/Private/DefaultMovementSet/Modes/AsyncNavWalkingMode.cpp
+++ b/Engine/Plugins/Experimental/Mover/Source/Mover/Private/DefaultMovementSet/Modes/AsyncNavWalkingMode.cpp
@@ -300,13 +300,6 @@ void UAsyncNavWalkingMode::SimulationTick_Implementation(const FSimulationTickPa
 			const float DownOffset = TotalCapsuleHeight * FMath::Max(0.f, NavMeshProjectionHeightScaleDown);
 			NewFeetLocation = ProjectLocationFromNavMesh(DeltaSeconds, StartingFeetLocation, NewFeetLocation, UpOffset, DownOffset);
 		}
-		else
-		{
-			if (UMoverBlackboard* SimBlackboard = MoverComp->GetSimBlackboard_Mutable())
-			{
-				SimBlackboard->Invalidate(CommonBlackboard::LastFloorResult);
-			}
-		}
 
 		FVector AdjustedDelta = NewFeetLocation - StartingFeetLocation;
 
@@ -381,7 +374,7 @@ void UAsyncNavWalkingMode::Activate()
 	{
 		if (UMoverBlackboard* SimBlackboard = MoverComp->GetSimBlackboard_Mutable())
 		{
-			SimBlackboard->Invalidate(CommonBlackboard::LastFloorResult);
+			SimBlackboard->Invalidate(NavWalkingBlackBoard::ProjectedNavMeshHitResult);
 		}
 	}
 
@@ -449,6 +442,13 @@ void UAsyncNavWalkingMode::FindBestNavMeshLocation(const FVector& TraceStart, co
 		{
 			UPrimitiveComponent* PrimComp = TestHit.GetComponent();
 			// Prefer using primitive component if valid
+			if (IPhysicsBodyInstanceOwner* BodyInstanceOwner = !PrimComp ? IPhysicsBodyInstanceOwner::GetPhysicsBodyInstandeOwnerFromHitResult(TestHit) : nullptr)
+			{
+				const bool bBlockOnWorldStatic = (BodyInstanceOwner->GetCollisionResponseToChannel(ECC_WorldStatic) == ECR_Block);
+				const bool bBlockOnWorldDynamic = (BodyInstanceOwner->GetCollisionResponseToChannel(ECC_WorldDynamic) == ECR_Block);
+				return !bBlockOnWorldStatic && (!bBlockOnWorldDynamic || bCheckOnlyWorldStatic);
+			}
+			else
 			{
 				const bool bBlockOnWorldStatic = PrimComp && (PrimComp->GetCollisionResponseToChannel(ECC_WorldStatic) == ECR_Block);
 				const bool bBlockOnWorldDynamic = PrimComp && (PrimComp->GetCollisionResponseToChannel(ECC_WorldDynamic) == ECR_Block);
@@ -487,10 +487,9 @@ FVector UAsyncNavWalkingMode::ProjectLocationFromNavMesh(float DeltaSeconds, con
 	const FVector TraceStart = TargetNavLocation + UpOffset * UpDirection;
 	const FVector TraceEnd = TargetNavLocation + DownOffset * -UpDirection;
 
-	FFloorCheckResult CachedFloorCheckResult;
+	FHitResult CachedProjectedNavMeshHitResult;
 	UMoverBlackboard* SimBlackboard = MoverComp->GetSimBlackboard_Mutable();
-	SimBlackboard->TryGet(CommonBlackboard::LastFloorResult, CachedFloorCheckResult);
-	FHitResult CachedProjectedNavMeshHitResult = CachedFloorCheckResult.HitResult;
+	SimBlackboard->TryGet(NavWalkingBlackBoard::ProjectedNavMeshHitResult, CachedProjectedNavMeshHitResult);
 
 	// We can skip this trace if we are checking at the same location as the last trace (ie, we haven't moved).
 	const bool bCachedLocationStillValid = (CachedProjectedNavMeshHitResult.bBlockingHit &&
@@ -511,21 +510,12 @@ FVector UAsyncNavWalkingMode::ProjectLocationFromNavMesh(float DeltaSeconds, con
 			if (HitResult.bStartPenetrating || !HitResult.bBlockingHit)
 			{
 				CachedProjectedNavMeshHitResult.Reset();
-				SimBlackboard->Invalidate(CommonBlackboard::LastFloorResult);
+				SimBlackboard->Invalidate(NavWalkingBlackBoard::ProjectedNavMeshHitResult);
 			}
 			else
 			{
 				CachedProjectedNavMeshHitResult = HitResult;
-
-				FFloorCheckResult FloorCheckResult;
-				FloorCheckResult.bBlockingHit = HitResult.bBlockingHit;
-				FloorCheckResult.bLineTrace = true;
-				FloorCheckResult.bWalkableFloor = true;
-				FloorCheckResult.LineDist = FMath::Abs((CurrentFeetLocation - CachedProjectedNavMeshHitResult.ImpactPoint).Dot(UpDirection));
-				FloorCheckResult.FloorDist = FloorCheckResult.LineDist; // This is usually set from a sweep trace but it doesn't really hurt setting it. 
-				FloorCheckResult.HitResult = CachedProjectedNavMeshHitResult;
-
-				SimBlackboard->Set(CommonBlackboard::LastFloorResult, FloorCheckResult);
+				SimBlackboard->Set(NavWalkingBlackBoard::ProjectedNavMeshHitResult, CachedProjectedNavMeshHitResult);
 			}
 		}
 		else
diff --git a/Engine/Plugins/Experimental/Mover/Source/Mover/Private/DefaultMovementSet/Modes/NavWalkingMode.cpp b/Engine/Plugins/Experimental/Mover/Source/Mover/Private/DefaultMovementSet/Modes/NavWalkingMode.cpp
--- a/Engine/Plugins/Experimental/Mover/Source/Mover/Private/DefaultMovementSet/Modes/NavWalkingMode.cpp
+++ b/Engine/Plugins/Experimental/Mover/Source/Mover/Private/DefaultMovementSet/Modes/NavWalkingMode.cpp
@@ -294,13 +294,6 @@ void UNavWalkingMode::SimulationTick_Implementation(const FSimulationTickParams&
 			const float DownOffset = TotalCapsuleHeight * FMath::Max(0.f, NavMeshProjectionHeightScaleDown);
 			NewLocation = ProjectLocationFromNavMesh(DeltaSeconds, StartingLocation, NewLocation, UpOffset, DownOffset);
 		}
-		else
-		{
-			if (UMoverBlackboard* SimBlackboard = MoverComp->GetSimBlackboard_Mutable())
-			{
-				SimBlackboard->Invalidate(CommonBlackboard::LastFloorResult);
-			}	
-		}
 
 		FVector AdjustedDelta = NewLocation - StartingLocation;
 		
@@ -371,7 +364,7 @@ void UNavWalkingMode::SetCollisionForNavWalking(bool bEnable)
 				{
 					if (UMoverBlackboard* SimBlackboard = MoverComponent->GetSimBlackboard_Mutable())
 					{
-						SimBlackboard->Invalidate(CommonBlackboard::LastFloorResult);
+						SimBlackboard->Invalidate(NavWalkingBlackBoard::ProjectedNavMeshHitResult);
 					}
 
 					// Stagger timed updates so many different characters spawned at the same time don't update on the same frame.
@@ -502,10 +495,9 @@ FVector UNavWalkingMode::ProjectLocationFromNavMesh(float DeltaSeconds, const FV
 	const FVector TraceStart = TargetNavLocation + UpOffset * UpDirection;
 	const FVector TraceEnd   = TargetNavLocation + DownOffset * -UpDirection;
 
-	FFloorCheckResult CachedFloorCheckResult;
+	FHitResult CachedProjectedNavMeshHitResult;
 	UMoverBlackboard* SimBlackboard = MoverComp->GetSimBlackboard_Mutable();
-	SimBlackboard->TryGet(CommonBlackboard::LastFloorResult, CachedFloorCheckResult);
-	FHitResult CachedProjectedNavMeshHitResult = CachedFloorCheckResult.HitResult;
+	SimBlackboard->TryGet(NavWalkingBlackBoard::ProjectedNavMeshHitResult, CachedProjectedNavMeshHitResult);
 	
 	// We can skip this trace if we are checking at the same location as the last trace (ie, we haven't moved).
 	const bool bCachedLocationStillValid = (CachedProjectedNavMeshHitResult.bBlockingHit &&
@@ -526,21 +518,12 @@ FVector UNavWalkingMode::ProjectLocationFromNavMesh(float DeltaSeconds, const FV
 			if (HitResult.bStartPenetrating || !HitResult.bBlockingHit)
 			{
 				CachedProjectedNavMeshHitResult.Reset();
-				SimBlackboard->Invalidate(CommonBlackboard::LastFloorResult);
+				SimBlackboard->Invalidate(NavWalkingBlackBoard::ProjectedNavMeshHitResult);
 			}
 			else
 			{
 				CachedProjectedNavMeshHitResult = HitResult;
-				
-				FFloorCheckResult FloorCheckResult;
-				FloorCheckResult.bBlockingHit = HitResult.bBlockingHit;
-				FloorCheckResult.bLineTrace = true;
-				FloorCheckResult.bWalkableFloor = true;
-				FloorCheckResult.LineDist = FMath::Abs((CurrentFeetLocation - CachedProjectedNavMeshHitResult.ImpactPoint).Dot(UpDirection));
-				FloorCheckResult.FloorDist = FloorCheckResult.LineDist; // This is usually set from a sweep trace but it doesn't really hurt setting it. 
-				FloorCheckResult.HitResult = HitResult;
-
-				SimBlackboard->Set(CommonBlackboard::LastFloorResult, FloorCheckResult);
+				SimBlackboard->Set(NavWalkingBlackBoard::ProjectedNavMeshHitResult, CachedProjectedNavMeshHitResult);
 			}
 		}
 		else
diff --git a/Engine/Plugins/Experimental/Mover/Source/Mover/Public/DefaultMovementSet/Modes/NavWalkingMode.h b/Engine/Plugins/Experimental/Mover/Source/Mover/Public/DefaultMovementSet/Modes/NavWalkingMode.h
--- a/Engine/Plugins/Experimental/Mover/Source/Mover/Public/DefaultMovementSet/Modes/NavWalkingMode.h
+++ b/Engine/Plugins/Experimental/Mover/Source/Mover/Public/DefaultMovementSet/Modes/NavWalkingMode.h
@@ -15,6 +15,10 @@ class UNavMoverComponent;
 class INavigationDataInterface;
 class UCommonLegacyMovementSettings;
 
+namespace NavWalkingBlackBoard
+{
+	const FName ProjectedNavMeshHitResult = TEXT("ProjectedNavMeshHitResult");
+}
 
 /** Options for what to do when we find ourselves off the nav mesh */
 UENUM(BlueprintType)
