Mover: fix for local->world root motion conversions with non-yaw rotations involved

diff --git a/Engine/Plugins/Experimental/Mover/Source/Mover/Private/MoverComponent.cpp b/Engine/Plugins/Experimental/Mover/Source/Mover/Private/MoverComponent.cpp
--- a/Engine/Plugins/Experimental/Mover/Source/Mover/Private/MoverComponent.cpp
+++ b/Engine/Plugins/Experimental/Mover/Source/Mover/Private/MoverComponent.cpp
@@ -1658,12 +1658,9 @@ static FTransform ConvertLocalRootMotionToAltWorldSpace(const FTransform& LocalR
 	const FTransform NewComponentToWorld = LocalRootMotionTransform * AltComponentWorldTransform;
 	const FTransform NewActorTransform = RelativeCompToActor * NewComponentToWorld;
 
-	const FVector DeltaWorldTranslation = NewActorTransform.GetTranslation() - AltWorldspaceTransform.GetTranslation();
-
-	const FQuat NewWorldRotation = AltComponentWorldTransform.GetRotation() * LocalRootMotionTransform.GetRotation();
-	const FQuat DeltaWorldRotation = NewWorldRotation * AltComponentWorldTransform.GetRotation().Inverse();
-
-	return FTransform(DeltaWorldRotation, DeltaWorldTranslation);
+	FTransform ActorDeltaTransform = NewActorTransform.GetRelativeTransform(AltWorldspaceTransform);
+	
+	return FTransform(ActorDeltaTransform.GetRotation(), NewActorTransform.GetTranslation() - AltWorldspaceTransform.GetTranslation());
 }
 
 FTransform UMoverComponent::ConvertLocalRootMotionToWorld(const FTransform& LocalRootMotionTransform, float DeltaSeconds, const FTransform* AlternateActorToWorld, const FMotionWarpingUpdateContext* OptionalWarpingContext) const
