Mover: respect setting for remaining vertical in all default movement modes' simulation functions. This allows other root motion contributors to operate without needing to know how it will be used

diff --git a/Engine/Plugins/Experimental/Mover/Source/Mover/Private/DefaultMovementSet/Modes/AsyncFallingMode.cpp b/Engine/Plugins/Experimental/Mover/Source/Mover/Private/DefaultMovementSet/Modes/AsyncFallingMode.cpp
--- a/Engine/Plugins/Experimental/Mover/Source/Mover/Private/DefaultMovementSet/Modes/AsyncFallingMode.cpp
+++ b/Engine/Plugins/Experimental/Mover/Source/Mover/Private/DefaultMovementSet/Modes/AsyncFallingMode.cpp
@@ -196,7 +196,11 @@ void UAsyncFallingMode::SimulationTick_Implementation(const FSimulationTickParam
 	const FVector TargetLocation = StartLocation + (ProposedMove.LinearVelocity * DeltaSeconds);
 
 	const FQuat StartRotation = StartingOrient.Quaternion();
-	const FQuat TargetRotation = TargetOrient.Quaternion();
+	FQuat TargetRotation = TargetOrient.Quaternion();
+	if (CommonLegacySettings->bShouldRemainVertical)
+	{
+		TargetRotation = FRotationMatrix::MakeFromZX(MoverComponent->GetUpDirection(), TargetRotation.GetForwardVector()).ToQuat();
+	}
 
 
 	FHitResult SweepHit(1.f);
diff --git a/Engine/Plugins/Experimental/Mover/Source/Mover/Private/DefaultMovementSet/Modes/AsyncFlyingMode.cpp b/Engine/Plugins/Experimental/Mover/Source/Mover/Private/DefaultMovementSet/Modes/AsyncFlyingMode.cpp
--- a/Engine/Plugins/Experimental/Mover/Source/Mover/Private/DefaultMovementSet/Modes/AsyncFlyingMode.cpp
+++ b/Engine/Plugins/Experimental/Mover/Source/Mover/Private/DefaultMovementSet/Modes/AsyncFlyingMode.cpp
@@ -111,7 +111,11 @@ void UAsyncFlyingMode::SimulationTick_Implementation(const FSimulationTickParams
 	const FVector TargetLocation = StartLocation + (ProposedMove.LinearVelocity * DeltaSeconds);
 
 	const FQuat StartRotation = StartingOrient.Quaternion();
-	const FQuat TargetRotation = TargetOrient.Quaternion();
+	FQuat TargetRotation = TargetOrient.Quaternion();
+	if (CommonLegacySettings->bShouldRemainVertical)
+	{
+		TargetRotation = FRotationMatrix::MakeFromZX(MoverComp->GetUpDirection(), TargetRotation.GetForwardVector()).ToQuat();
+	}
 
 	FHitResult SweepHit(1.f);
 	FMovementRecord SweepRecord;
diff --git a/Engine/Plugins/Experimental/Mover/Source/Mover/Private/DefaultMovementSet/Modes/AsyncNavWalkingMode.cpp b/Engine/Plugins/Experimental/Mover/Source/Mover/Private/DefaultMovementSet/Modes/AsyncNavWalkingMode.cpp
--- a/Engine/Plugins/Experimental/Mover/Source/Mover/Private/DefaultMovementSet/Modes/AsyncNavWalkingMode.cpp
+++ b/Engine/Plugins/Experimental/Mover/Source/Mover/Private/DefaultMovementSet/Modes/AsyncNavWalkingMode.cpp
@@ -168,8 +168,12 @@ void UAsyncNavWalkingMode::SimulationTick_Implementation(const FSimulationTickPa
 	const FVector StartLocation = StartingSyncState->GetLocation_WorldSpace();
 
 	const FQuat StartRotation = StartingOrient.Quaternion();
-	const FQuat TargetRotation = TargetOrient.Quaternion();
 
+	FQuat TargetRotation = TargetOrient.Quaternion();
+	if (CommonLegacySettings->bShouldRemainVertical)
+	{
+		TargetRotation = FRotationMatrix::MakeFromZX(UpDirection, TargetRotation.GetForwardVector()).ToQuat();
+	}
 
 	FVector LocationInProgress = StartLocation;
 	FQuat RotationInProgress = StartRotation;
diff --git a/Engine/Plugins/Experimental/Mover/Source/Mover/Private/DefaultMovementSet/Modes/AsyncWalkingMode.cpp b/Engine/Plugins/Experimental/Mover/Source/Mover/Private/DefaultMovementSet/Modes/AsyncWalkingMode.cpp
--- a/Engine/Plugins/Experimental/Mover/Source/Mover/Private/DefaultMovementSet/Modes/AsyncWalkingMode.cpp
+++ b/Engine/Plugins/Experimental/Mover/Source/Mover/Private/DefaultMovementSet/Modes/AsyncWalkingMode.cpp
@@ -170,7 +170,12 @@ void UAsyncWalkingMode::SimulationTick_Implementation(const FSimulationTickParam
 	}
 
 	const FQuat StartRotation = StartingOrient.Quaternion();
-	const FQuat TargetRotation = TargetOrient.Quaternion();
+
+	FQuat TargetRotation = TargetOrient.Quaternion();
+	if (CommonLegacySettings->bShouldRemainVertical)
+	{
+		TargetRotation = FRotationMatrix::MakeFromZX(UpDirection, TargetRotation.GetForwardVector()).ToQuat();
+	}
 
 	FVector LocationInProgress = StartLocation;
 	FQuat   RotationInProgress = StartRotation;
diff --git a/Engine/Plugins/Experimental/Mover/Source/Mover/Private/DefaultMovementSet/Modes/FallingMode.cpp b/Engine/Plugins/Experimental/Mover/Source/Mover/Private/DefaultMovementSet/Modes/FallingMode.cpp
--- a/Engine/Plugins/Experimental/Mover/Source/Mover/Private/DefaultMovementSet/Modes/FallingMode.cpp
+++ b/Engine/Plugins/Experimental/Mover/Source/Mover/Private/DefaultMovementSet/Modes/FallingMode.cpp
@@ -204,7 +204,12 @@ void UFallingMode::SimulationTick_Implementation(const FSimulationTickParams& Pa
 	FVector MoveDelta = ProposedMove.LinearVelocity * DeltaSeconds;
 
 	FHitResult Hit(1.f);
-	const FQuat TargetOrientQuat = TargetOrient.Quaternion();
+	FQuat TargetOrientQuat = TargetOrient.Quaternion();
+
+	if (CommonLegacySettings->bShouldRemainVertical)
+	{
+		TargetOrientQuat = FRotationMatrix::MakeFromZX(UpDirection, TargetOrientQuat.GetForwardVector()).ToQuat();
+	}
 
 	UMovementUtils::TrySafeMoveUpdatedComponent(Params.MovingComps, MoveDelta, TargetOrientQuat, true, Hit, ETeleportType::None, MoveRecord);
 
diff --git a/Engine/Plugins/Experimental/Mover/Source/Mover/Private/DefaultMovementSet/Modes/FlyingMode.cpp b/Engine/Plugins/Experimental/Mover/Source/Mover/Private/DefaultMovementSet/Modes/FlyingMode.cpp
--- a/Engine/Plugins/Experimental/Mover/Source/Mover/Private/DefaultMovementSet/Modes/FlyingMode.cpp
+++ b/Engine/Plugins/Experimental/Mover/Source/Mover/Private/DefaultMovementSet/Modes/FlyingMode.cpp
@@ -112,12 +112,18 @@ void UFlyingMode::SimulationTick_Implementation(const FSimulationTickParams& Par
 	}
 	
 	FVector MoveDelta = ProposedMove.LinearVelocity * DeltaSeconds;
-	const FQuat OrientQuat = TargetOrient.Quaternion();
+
+	FQuat TargetOrientQuat = TargetOrient.Quaternion();
+	if (CommonLegacySettings->bShouldRemainVertical)
+	{
+		TargetOrientQuat = FRotationMatrix::MakeFromZX(MoverComp->GetUpDirection(), TargetOrientQuat.GetForwardVector()).ToQuat();
+	}
+
 	FHitResult Hit(1.f);
 
 	if (!MoveDelta.IsNearlyZero() || bIsOrientationChanging)
 	{
-		UMovementUtils::TrySafeMoveUpdatedComponent(Params.MovingComps, MoveDelta, OrientQuat, true, Hit, ETeleportType::None, MoveRecord);
+		UMovementUtils::TrySafeMoveUpdatedComponent(Params.MovingComps, MoveDelta, TargetOrientQuat, true, Hit, ETeleportType::None, MoveRecord);
 	}
 
 	if (Hit.IsValidBlockingHit())
@@ -126,7 +132,7 @@ void UFlyingMode::SimulationTick_Implementation(const FSimulationTickParams& Par
 		MoverComp->HandleImpact(ImpactParams);
 
 		// Try to slide the remaining distance along the surface.
-		UMovementUtils::TryMoveToSlideAlongSurface(FMovingComponentSet(MoverComp), MoveDelta, 1.f - Hit.Time, OrientQuat, Hit.Normal, Hit, true, MoveRecord);
+		UMovementUtils::TryMoveToSlideAlongSurface(FMovingComponentSet(MoverComp), MoveDelta, 1.f - Hit.Time, TargetOrientQuat, Hit.Normal, Hit, true, MoveRecord);
 	}
 
 	if (bRespectDistanceOverWalkableSurfaces)
diff --git a/Engine/Plugins/Experimental/Mover/Source/Mover/Private/DefaultMovementSet/Modes/NavWalkingMode.cpp b/Engine/Plugins/Experimental/Mover/Source/Mover/Private/DefaultMovementSet/Modes/NavWalkingMode.cpp
--- a/Engine/Plugins/Experimental/Mover/Source/Mover/Private/DefaultMovementSet/Modes/NavWalkingMode.cpp
+++ b/Engine/Plugins/Experimental/Mover/Source/Mover/Private/DefaultMovementSet/Modes/NavWalkingMode.cpp
@@ -125,7 +125,7 @@ void UNavWalkingMode::SimulationTick_Implementation(const FSimulationTickParams&
 	USceneComponent* UpdatedComponent = Params.MovingComps.UpdatedComponent.Get();
 	UPrimitiveComponent* UpdatedPrimitive = Params.MovingComps.UpdatedPrimitive.Get();
 	const FProposedMove& ProposedMove = Params.ProposedMove;
-	FVector UpDirection = MoverComp->GetUpDirection();
+	const FVector UpDirection = MoverComp->GetUpDirection();
 
 	if (!UpdatedComponent || !UpdatedPrimitive)
 	{
@@ -161,8 +161,11 @@ void UNavWalkingMode::SimulationTick_Implementation(const FSimulationTickParams&
 		bIsOrientationChanging = true;
 	}
 	
-	const FQuat TargetOrientQuat = TargetOrient.Quaternion();
-
+	FQuat TargetOrientQuat = TargetOrient.Quaternion();
+	if (CommonLegacySettings->bShouldRemainVertical)
+	{
+		TargetOrientQuat = FRotationMatrix::MakeFromZX(UpDirection, TargetOrientQuat.GetForwardVector()).ToQuat();
+	}
 	
 	FVector StartingLocation = NavMoverComponent->GetFeetLocation();
 	FVector AdjustedDest = StartingLocation + OrigMoveDelta;
diff --git a/Engine/Plugins/Experimental/Mover/Source/Mover/Private/DefaultMovementSet/Modes/WalkingMode.cpp b/Engine/Plugins/Experimental/Mover/Source/Mover/Private/DefaultMovementSet/Modes/WalkingMode.cpp
--- a/Engine/Plugins/Experimental/Mover/Source/Mover/Private/DefaultMovementSet/Modes/WalkingMode.cpp
+++ b/Engine/Plugins/Experimental/Mover/Source/Mover/Private/DefaultMovementSet/Modes/WalkingMode.cpp
@@ -161,7 +161,12 @@ void UWalkingMode::SimulationTick_Implementation(const FSimulationTickParams& Pa
 		bIsOrientationChanging = (TargetOrient != StartingOrient);
 	}
 
-	const FQuat TargetOrientQuat = TargetOrient.Quaternion();
+	FQuat TargetOrientQuat = TargetOrient.Quaternion();
+	if (CommonLegacySettings->bShouldRemainVertical)
+	{
+		TargetOrientQuat = FRotationMatrix::MakeFromZX(UpDirection, TargetOrientQuat.GetForwardVector()).ToQuat();
+	}
+
 	FHitResult MoveHitResult(1.f);
 	
 	FVector CurMoveDelta = OrigMoveDelta;
