[Mover] StartingMovementMode is set in the MoverComponent in a recent change, so remove this. Additionally, as ModeFSM is not intialized yet at this point, the game crashes on Liaison initialization.

diff --git a/Engine/Plugins/Experimental/Mover/Source/Mover/Private/Backends/MoverNetworkPhysicsLiaisonBase.cpp b/Engine/Plugins/Experimental/Mover/Source/Mover/Private/Backends/MoverNetworkPhysicsLiaisonBase.cpp
--- a/Engine/Plugins/Experimental/Mover/Source/Mover/Private/Backends/MoverNetworkPhysicsLiaisonBase.cpp
+++ b/Engine/Plugins/Experimental/Mover/Source/Mover/Private/Backends/MoverNetworkPhysicsLiaisonBase.cpp
@@ -309,7 +309,6 @@ void UMoverNetworkPhysicsLiaisonComponentBase::InitializeComponent()
 	if (UWorld* World = GetWorld(); World && World->IsGameWorld())
 	{
 		UMoverComponent& MoverComp = GetMoverComponent();
-		MoverComp.ModeFSM->SetModeImmediately(MoverComp.StartingMovementMode);
 	}
 
 	if (Chaos::FPhysicsSolverBase::IsNetworkPhysicsPredictionEnabled() && bUsingAsyncPhysics)
