Mover: making nav walking modes force a navmesh projection if something comes along and invalidates our last floor information, such as with a teleport

diff --git a/Engine/Plugins/Experimental/Mover/Source/Mover/Private/DefaultMovementSet/Modes/AsyncNavWalkingMode.cpp b/Engine/Plugins/Experimental/Mover/Source/Mover/Private/DefaultMovementSet/Modes/AsyncNavWalkingMode.cpp
--- a/Engine/Plugins/Experimental/Mover/Source/Mover/Private/DefaultMovementSet/Modes/AsyncNavWalkingMode.cpp
+++ b/Engine/Plugins/Experimental/Mover/Source/Mover/Private/DefaultMovementSet/Modes/AsyncNavWalkingMode.cpp
@@ -308,7 +308,8 @@ void UAsyncNavWalkingMode::SimulationTick_Implementation(const FSimulationTickPa
 		{
 			if (UMoverBlackboard* SimBlackboard = MoverComp->GetSimBlackboard_Mutable())
 			{
-				SimBlackboard->Invalidate(CommonBlackboard::LastFloorResult);
+				const FFloorCheckResult EmptyFloorCheckResult;
+				SimBlackboard->Set(CommonBlackboard::LastFloorResult, EmptyFloorCheckResult);
 			}
 		}
 
@@ -500,7 +501,7 @@ FVector UAsyncNavWalkingMode::ProjectLocationFromNavMesh(float DeltaSeconds, con
 
 	FFloorCheckResult CachedFloorCheckResult;
 	UMoverBlackboard* SimBlackboard = MoverComp->GetSimBlackboard_Mutable();
-	SimBlackboard->TryGet(CommonBlackboard::LastFloorResult, CachedFloorCheckResult);
+	bool bHasValidFloorResult = SimBlackboard->TryGet(CommonBlackboard::LastFloorResult, CachedFloorCheckResult);
 	FHitResult CachedProjectedNavMeshHitResult = CachedFloorCheckResult.HitResult;
 
 	// We can skip this trace if we are checking at the same location as the last trace (ie, we haven't moved).
@@ -508,8 +509,9 @@ FVector UAsyncNavWalkingMode::ProjectLocationFromNavMesh(float DeltaSeconds, con
 		CachedProjectedNavMeshHitResult.TraceStart == TraceStart &&
 		CachedProjectedNavMeshHitResult.TraceEnd == TraceEnd);
 
+	// Check periodically or if we have no information about our last floor result
 	NavMeshProjectionTimer -= DeltaSeconds;
-	if (NavMeshProjectionTimer <= 0.0f)
+	if (NavMeshProjectionTimer <= 0.0f || !bHasValidFloorResult)
 	{
 		if (!bCachedLocationStillValid)
 		{
@@ -522,7 +524,8 @@ FVector UAsyncNavWalkingMode::ProjectLocationFromNavMesh(float DeltaSeconds, con
 			if (HitResult.bStartPenetrating || !HitResult.bBlockingHit)
 			{
 				CachedProjectedNavMeshHitResult.Reset();
-				SimBlackboard->Invalidate(CommonBlackboard::LastFloorResult);
+				const FFloorCheckResult EmptyFloorCheckResult;
+				SimBlackboard->Set(CommonBlackboard::LastFloorResult, EmptyFloorCheckResult);
 			}
 			else
 			{
diff --git a/Engine/Plugins/Experimental/Mover/Source/Mover/Private/DefaultMovementSet/Modes/NavWalkingMode.cpp b/Engine/Plugins/Experimental/Mover/Source/Mover/Private/DefaultMovementSet/Modes/NavWalkingMode.cpp
--- a/Engine/Plugins/Experimental/Mover/Source/Mover/Private/DefaultMovementSet/Modes/NavWalkingMode.cpp
+++ b/Engine/Plugins/Experimental/Mover/Source/Mover/Private/DefaultMovementSet/Modes/NavWalkingMode.cpp
@@ -301,7 +301,8 @@ void UNavWalkingMode::SimulationTick_Implementation(const FSimulationTickParams&
 		{
 			if (UMoverBlackboard* SimBlackboard = MoverComp->GetSimBlackboard_Mutable())
 			{
-				SimBlackboard->Invalidate(CommonBlackboard::LastFloorResult);
+				const FFloorCheckResult EmptyFloorCheckResult;
+				SimBlackboard->Set(CommonBlackboard::LastFloorResult, EmptyFloorCheckResult);
 			}	
 		}
 
@@ -374,7 +375,8 @@ void UNavWalkingMode::SetCollisionForNavWalking(bool bEnable)
 				{
 					if (UMoverBlackboard* SimBlackboard = MoverComponent->GetSimBlackboard_Mutable())
 					{
-						SimBlackboard->Invalidate(CommonBlackboard::LastFloorResult);
+						const FFloorCheckResult EmptyFloorCheckResult;
+						SimBlackboard->Set(CommonBlackboard::LastFloorResult, EmptyFloorCheckResult);
 					}
 
 					// Stagger timed updates so many different characters spawned at the same time don't update on the same frame.
@@ -507,7 +509,7 @@ FVector UNavWalkingMode::ProjectLocationFromNavMesh(float DeltaSeconds, const FV
 
 	FFloorCheckResult CachedFloorCheckResult;
 	UMoverBlackboard* SimBlackboard = MoverComp->GetSimBlackboard_Mutable();
-	SimBlackboard->TryGet(CommonBlackboard::LastFloorResult, CachedFloorCheckResult);
+	bool bHasValidFloorResult = SimBlackboard->TryGet(CommonBlackboard::LastFloorResult, CachedFloorCheckResult);
 	FHitResult CachedProjectedNavMeshHitResult = CachedFloorCheckResult.HitResult;
 	
 	// We can skip this trace if we are checking at the same location as the last trace (ie, we haven't moved).
@@ -515,8 +517,9 @@ FVector UNavWalkingMode::ProjectLocationFromNavMesh(float DeltaSeconds, const FV
 											CachedProjectedNavMeshHitResult.TraceStart == TraceStart &&
 											CachedProjectedNavMeshHitResult.TraceEnd == TraceEnd);
 
+	// Check periodically or if we have no information about our last floor result
 	NavMeshProjectionTimer -= DeltaSeconds;
-	if (NavMeshProjectionTimer <= 0.0f)
+	if (NavMeshProjectionTimer <= 0.0f || !bHasValidFloorResult)
 	{
 		if (!bCachedLocationStillValid)
 		{
@@ -529,7 +532,8 @@ FVector UNavWalkingMode::ProjectLocationFromNavMesh(float DeltaSeconds, const FV
 			if (HitResult.bStartPenetrating || !HitResult.bBlockingHit)
 			{
 				CachedProjectedNavMeshHitResult.Reset();
-				SimBlackboard->Invalidate(CommonBlackboard::LastFloorResult);
+				const FFloorCheckResult EmptyFloorCheckResult;
+				SimBlackboard->Set(CommonBlackboard::LastFloorResult, EmptyFloorCheckResult);
 			}
 			else
 			{
