Fix uninitialized properties in various anim files

diff --git a/Engine/Plugins/Experimental/AnimNextAnimGraph/Source/AnimNextAnimGraph/Private/Traits/ModifyCurveTrait.h b/Engine/Plugins/Experimental/AnimNextAnimGraph/Source/AnimNextAnimGraph/Private/Traits/ModifyCurveTrait.h
--- a/Engine/Plugins/Experimental/AnimNextAnimGraph/Source/AnimNextAnimGraph/Private/Traits/ModifyCurveTrait.h
+++ b/Engine/Plugins/Experimental/AnimNextAnimGraph/Source/AnimNextAnimGraph/Private/Traits/ModifyCurveTrait.h
@@ -35,7 +35,7 @@ struct FModifyCurveParameters
 	FName CurveName;
 
 	UPROPERTY(Category=CurveData, EditAnywhere)
-	float CurveValue;
+	float CurveValue = 0.0f;
 	
 };
 
@@ -45,7 +45,7 @@ struct FAnimNextModifyCurveTraitSharedData : public FAnimNextTraitSharedData
 	GENERATED_BODY()
 
 	UPROPERTY(Category=Settings, EditAnywhere)
-	EAnimNext_ModifyCurveApplyMode ApplyMode;
+	EAnimNext_ModifyCurveApplyMode ApplyMode = EAnimNext_ModifyCurveApplyMode::Add;
 
 	UPROPERTY(Category=Settings, EditAnywhere, meta=(Inline, Hidden))
 	TArray<FModifyCurveParameters> ModifyCurveParameters;
diff --git a/Engine/Plugins/Experimental/AnimNextAnimGraph/Source/AnimNextAnimGraphEditor/Private/RewindDebugger/SequenceInfoTrack.h b/Engine/Plugins/Experimental/AnimNextAnimGraph/Source/AnimNextAnimGraphEditor/Private/RewindDebugger/SequenceInfoTrack.h
--- a/Engine/Plugins/Experimental/AnimNextAnimGraph/Source/AnimNextAnimGraphEditor/Private/RewindDebugger/SequenceInfoTrack.h
+++ b/Engine/Plugins/Experimental/AnimNextAnimGraph/Source/AnimNextAnimGraphEditor/Private/RewindDebugger/SequenceInfoTrack.h
@@ -19,7 +19,7 @@ struct FAnimNextSyncMarkerTraceInfo
 	GENERATED_BODY()
 	
 	UPROPERTY()
-	float Time;
+	float Time = 0.0f;
 
 	UPROPERTY()
 	FName Name;
@@ -39,10 +39,10 @@ struct FAnimNextSequenceTraceInfo
 	TWeakObjectPtr<UAnimSequence> AnimSequence;
 
 	UPROPERTY(VisibleAnywhere, Category = Properties)
-	float DurationSeconds;
+	float DurationSeconds = 0.0f;
 
 	UPROPERTY(VisibleAnywhere, Category = Properties)
-	float CurrentTimeSeconds;
+	float CurrentTimeSeconds = 0.0f;
 
 	UPROPERTY(VisibleAnywhere, Category = Properties)
 	TArray<FAnimNextSyncMarkerTraceInfo> SyncMarkers;
diff --git a/Engine/Plugins/Experimental/Animation/HierarchyTableAnimation/Source/Runtime/Public/HierarchyTableBlendProfile.h b/Engine/Plugins/Experimental/Animation/HierarchyTableAnimation/Source/Runtime/Public/HierarchyTableBlendProfile.h
--- a/Engine/Plugins/Experimental/Animation/HierarchyTableAnimation/Source/Runtime/Public/HierarchyTableBlendProfile.h
+++ b/Engine/Plugins/Experimental/Animation/HierarchyTableAnimation/Source/Runtime/Public/HierarchyTableBlendProfile.h
@@ -22,10 +22,10 @@ struct FMaskedAttributeWeightSerialised
 	FName AttributeName;
 	
 	UPROPERTY()
-	int32 AttributeIndex;
+	int32 AttributeIndex = INDEX_NONE;
 
 	UPROPERTY()
-	float Weight;
+	float Weight = 0.0f;
 };
 
 USTRUCT()
@@ -37,7 +37,7 @@ struct FMaskedCurveWeightSerialised
 	FName CurveName;
 
 	UPROPERTY()
-	float Weight;
+	float Weight = 0.0f;
 };
 
 USTRUCT()
diff --git a/Engine/Plugins/Experimental/Mover/Source/Mover/Public/MoveLibrary/MovementUtils.h b/Engine/Plugins/Experimental/Mover/Source/Mover/Public/MoveLibrary/MovementUtils.h
--- a/Engine/Plugins/Experimental/Mover/Source/Mover/Public/MoveLibrary/MovementUtils.h
+++ b/Engine/Plugins/Experimental/Mover/Source/Mover/Public/MoveLibrary/MovementUtils.h
@@ -35,19 +35,19 @@ struct FTrajectorySampleInfo
 
 	// Velocity at the time of this sample (world space, units/sec)
 	UPROPERTY(BlueprintReadWrite, EditAnywhere, Category = Mover)
-	FVector LinearVelocity;
+	FVector LinearVelocity = FVector::ZeroVector;
 
 	// Acceleration at the time of this sample (world space, units/sec^2)
 	UPROPERTY(BlueprintReadWrite, EditAnywhere, Category = Mover)
-	FVector InstantaneousAcceleration;
+	FVector InstantaneousAcceleration = FVector::ZeroVector;
 
 	// Rotational velocity (world space, degrees/sec)
 	UPROPERTY(BlueprintReadWrite, EditAnywhere, Category = Mover)
-	FRotator AngularVelocity;
+	FRotator AngularVelocity = FRotator::ZeroRotator;
 
 	// Time stamp of this sample, in server simulation time
 	UPROPERTY(BlueprintReadWrite, EditAnywhere, Category = Mover)
-	float SimTimeMs;
+	float SimTimeMs = 0.0f;
 };
 
 // Input parameters for compute velocity function
