Chaos Visual Debugger: Fixing a silly copy paste error that meant the mover info input structs weren't properly displaying in the mover info tab

diff --git a/Engine/Plugins/Experimental/Mover/Source/Mover/Private/ChaosVisualDebugger/MoverCVDRuntimeTrace.cpp b/Engine/Plugins/Experimental/Mover/Source/Mover/Private/ChaosVisualDebugger/MoverCVDRuntimeTrace.cpp
--- a/Engine/Plugins/Experimental/Mover/Source/Mover/Private/ChaosVisualDebugger/MoverCVDRuntimeTrace.cpp
+++ b/Engine/Plugins/Experimental/Mover/Source/Mover/Private/ChaosVisualDebugger/MoverCVDRuntimeTrace.cpp
@@ -88,8 +88,8 @@ void FMoverCVDRuntimeTrace::UnwrapSimData(const FMoverCVDSimDataWrapper& InSimDa
 		OutSyncState = MakeShared<FMoverSyncState>();
 		FMoverSyncState::StaticStruct()->SerializeBin(ArReader, OutSyncState.Get());
 		// Sync State's Collection of custom structs
-		FSkipObjectRefsMemoryReader ArInputCollectionReader(InSimDataWrapper.SyncStateDataCollectionBytes, true);
-		OutInputCmd->InputCollection.SerializeDebugData(ArInputCollectionReader);
+		FSkipObjectRefsMemoryReader ArSyncStateCollectionReader(InSimDataWrapper.SyncStateDataCollectionBytes, true);
+		OutSyncState->SyncStateCollection.SerializeDebugData(ArSyncStateCollectionReader);
 	}
 
 	{
