Mover: Adding double buffer for cached sync state to avoid issues with reading in data that may be freed/overwritten

diff --git a/Engine/Plugins/Experimental/Mover/Source/Mover/Private/Backends/MoverNetworkPhysicsLiaisonBase.cpp b/Engine/Plugins/Experimental/Mover/Source/Mover/Private/Backends/MoverNetworkPhysicsLiaisonBase.cpp
--- a/Engine/Plugins/Experimental/Mover/Source/Mover/Private/Backends/MoverNetworkPhysicsLiaisonBase.cpp
+++ b/Engine/Plugins/Experimental/Mover/Source/Mover/Private/Backends/MoverNetworkPhysicsLiaisonBase.cpp
@@ -621,8 +621,11 @@ void UMoverNetworkPhysicsLiaisonComponentBase::PostPhysicsUpdate_External()
 
 	// The MoverComponent relies on its CachedLastSyncState for a lot of information, so setting it here is what makes the resulting state of the completed sim frame
 	// "real" from the perspective the MoverComp and any objects that call getters on it
-	MoverComp.CachedLastSyncState = InterpolatedSyncState;
-	MoverComp.LastMoverDefaultSyncState = MoverComp.CachedLastSyncState.SyncStateCollection.FindDataByType<FMoverDefaultSyncState>();
+	FMoverSyncState& BufferedSyncState = MoverComp.MoverSyncStateDoubleBuffer.GetWritable();
+	BufferedSyncState = InterpolatedSyncState;
+	MoverComp.LastMoverDefaultSyncState = BufferedSyncState.SyncStateCollection.FindDataByType<FMoverDefaultSyncState>();
+	MoverComp.MoverSyncStateDoubleBuffer.Flip();
+	
 	MoverComp.CachedLastUsedInputCmd = InterpolatedInputCmd;
 
 	// Note this may return something different now because its dependent on the CachedLastSyncState that we just changed above
diff --git a/Engine/Plugins/Experimental/Mover/Source/Mover/Private/Debug/MoverDebugComponent.cpp b/Engine/Plugins/Experimental/Mover/Source/Mover/Private/Debug/MoverDebugComponent.cpp
--- a/Engine/Plugins/Experimental/Mover/Source/Mover/Private/Debug/MoverDebugComponent.cpp
+++ b/Engine/Plugins/Experimental/Mover/Source/Mover/Private/Debug/MoverDebugComponent.cpp
@@ -66,7 +66,7 @@ void UMoverDebugComponent::TickComponent(float DeltaTime, ELevelTick TickType, F
 		UMoverComponent* MoverComp = GetOwner()->FindComponentByClass<UMoverComponent>();
 		if (MoverComp && MoverComp->CachedLastSimTickTimeStep.ServerFrame > 0)
 		{
-			UpdateHistoryTrackingForFrame(MoverComp->CachedLastSimTickTimeStep, &MoverComp->CachedLastSyncState, &MoverComp->CachedLastAuxState);
+			UpdateHistoryTrackingForFrame(MoverComp->CachedLastSimTickTimeStep, &MoverComp->GetSyncState(), &MoverComp->CachedLastAuxState);
 		}
 	}
 	
diff --git a/Engine/Plugins/Experimental/Mover/Source/Mover/Private/MoverComponent.cpp b/Engine/Plugins/Experimental/Mover/Source/Mover/Private/MoverComponent.cpp
--- a/Engine/Plugins/Experimental/Mover/Source/Mover/Private/MoverComponent.cpp
+++ b/Engine/Plugins/Experimental/Mover/Source/Mover/Private/MoverComponent.cpp
@@ -150,9 +150,11 @@ void UMoverComponent::InitializeComponent()
 	}
 
 	// Gather initial state to fulfill queries
-	CreateDefaultInputAndState(CachedLastProducedInputCmd, CachedLastSyncState, CachedLastAuxState);
+	FMoverSyncState DefaultMoverSyncState;
+	CreateDefaultInputAndState(CachedLastProducedInputCmd, DefaultMoverSyncState, CachedLastAuxState);
+	MoverSyncStateDoubleBuffer.SetBufferedData(DefaultMoverSyncState);
 	CachedLastUsedInputCmd = CachedLastProducedInputCmd;
-	LastMoverDefaultSyncState = CachedLastSyncState.SyncStateCollection.FindDataByType<FMoverDefaultSyncState>();
+	LastMoverDefaultSyncState = MoverSyncStateDoubleBuffer.GetReadable().SyncStateCollection.FindDataByType<FMoverDefaultSyncState>();
 }
 
 
@@ -454,8 +456,9 @@ void UMoverComponent::InitializeSimulationState(FMoverSyncState* OutSync, FMover
 	CreateDefaultInputAndState(CachedLastProducedInputCmd, *OutSync, *OutAux);
 
 	CachedLastUsedInputCmd = CachedLastProducedInputCmd;
-	CachedLastSyncState = *OutSync;
-	LastMoverDefaultSyncState = CachedLastSyncState.SyncStateCollection.FindDataByType<FMoverDefaultSyncState>();;
+	MoverSyncStateDoubleBuffer.SetBufferedData(*OutSync);
+	LastMoverDefaultSyncState = MoverSyncStateDoubleBuffer.GetReadable().SyncStateCollection.FindDataByType<FMoverDefaultSyncState>();
+	
 	CachedLastAuxState = *OutAux;
 
 }
@@ -668,8 +671,10 @@ bool UMoverComponent::IsModifierActiveOrQueued(const FMovementModifierHandle& Mo
 
 const FMovementModifierBase* UMoverComponent::FindMovementModifier(const FMovementModifierHandle& ModifierHandle) const
 {
+	const FMoverSyncState& CachedSyncState = MoverSyncStateDoubleBuffer.GetReadable();
+	
 	// Check active modifiers for modifier handle
-	for (auto ActiveModifierFromSyncStateIt = CachedLastSyncState.MovementModifiers.GetActiveModifiersIterator(); ActiveModifierFromSyncStateIt; ++ActiveModifierFromSyncStateIt)
+	for (auto ActiveModifierFromSyncStateIt = CachedSyncState.MovementModifiers.GetActiveModifiersIterator(); ActiveModifierFromSyncStateIt; ++ActiveModifierFromSyncStateIt)
 	{
 		const TSharedPtr<FMovementModifierBase> ActiveModifierFromSyncState = *ActiveModifierFromSyncStateIt;
 
@@ -680,7 +685,7 @@ const FMovementModifierBase* UMoverComponent::FindMovementModifier(const FMoveme
 	}
 
 	// Check queued modifiers for modifier handle
-	for (auto QueuedModifierFromSyncStateIt = CachedLastSyncState.MovementModifiers.GetQueuedModifiersIterator(); QueuedModifierFromSyncStateIt; ++QueuedModifierFromSyncStateIt)
+	for (auto QueuedModifierFromSyncStateIt = CachedSyncState.MovementModifiers.GetQueuedModifiersIterator(); QueuedModifierFromSyncStateIt; ++QueuedModifierFromSyncStateIt)
 	{
 		const TSharedPtr<FMovementModifierBase> QueuedModifierFromSyncState = *QueuedModifierFromSyncStateIt;
 
@@ -695,8 +700,10 @@ const FMovementModifierBase* UMoverComponent::FindMovementModifier(const FMoveme
 
 const FMovementModifierBase* UMoverComponent::FindMovementModifierByType(const UScriptStruct* DataStructType) const
 {
+	const FMoverSyncState& CachedSyncState = MoverSyncStateDoubleBuffer.GetReadable();
+	
 	// Check active modifiers for modifier handle
-	for (auto ActiveModifierFromSyncStateIt = CachedLastSyncState.MovementModifiers.GetActiveModifiersIterator(); ActiveModifierFromSyncStateIt; ++ActiveModifierFromSyncStateIt)
+	for (auto ActiveModifierFromSyncStateIt = CachedSyncState.MovementModifiers.GetActiveModifiersIterator(); ActiveModifierFromSyncStateIt; ++ActiveModifierFromSyncStateIt)
 	{
 		const TSharedPtr<FMovementModifierBase> ActiveModifierFromSyncState = *ActiveModifierFromSyncStateIt;
 
@@ -707,7 +714,7 @@ const FMovementModifierBase* UMoverComponent::FindMovementModifierByType(const U
 	}
 
 	// Check queued modifiers for modifier handle
-	for (auto QueuedModifierFromSyncStateIt = CachedLastSyncState.MovementModifiers.GetQueuedModifiersIterator(); QueuedModifierFromSyncStateIt; ++QueuedModifierFromSyncStateIt)
+	for (auto QueuedModifierFromSyncStateIt = CachedSyncState.MovementModifiers.GetQueuedModifiersIterator(); QueuedModifierFromSyncStateIt; ++QueuedModifierFromSyncStateIt)
 	{
 		const TSharedPtr<FMovementModifierBase> QueuedModifierFromSyncState = *QueuedModifierFromSyncStateIt;
 
@@ -746,8 +753,9 @@ bool UMoverComponent::HasGameplayTag(FGameplayTag TagToFind, bool bExactMatch) c
 		}
 	}
 		
+	const FMoverSyncState& CachedSyncState = MoverSyncStateDoubleBuffer.GetReadable();
 	// Search Movement Modifiers
-	for (auto ModifierFromSyncStateIt = CachedLastSyncState.MovementModifiers.GetActiveModifiersIterator(); ModifierFromSyncStateIt; ++ModifierFromSyncStateIt)
+	for (auto ModifierFromSyncStateIt = CachedSyncState.MovementModifiers.GetActiveModifiersIterator(); ModifierFromSyncStateIt; ++ModifierFromSyncStateIt)
 	{
 		if (const TSharedPtr<FMovementModifierBase> ModifierFromSyncState = *ModifierFromSyncStateIt)
 		{
@@ -759,7 +767,7 @@ bool UMoverComponent::HasGameplayTag(FGameplayTag TagToFind, bool bExactMatch) c
 	}
 		
 	// Search Layered Moves
-	for (const TSharedPtr<FLayeredMoveBase>& LayeredMove :  CachedLastSyncState.LayeredMoves.GetActiveMoves())
+	for (const TSharedPtr<FLayeredMoveBase>& LayeredMove :  CachedSyncState.LayeredMoves.GetActiveMoves())
 	{
 		if (LayeredMove->HasGameplayTag(TagToFind, bExactMatch))
 		{
@@ -793,8 +801,12 @@ void UMoverComponent::RemoveGameplayTags(const FGameplayTagContainer& TagsToRemo
 void UMoverComponent::UpdateCachedFrameState(const FMoverSyncState* SyncState, const FMoverAuxStateContext* AuxState)
 {
 	// TODO integrate dirty tracking
-	CachedLastSyncState = *SyncState;
-	LastMoverDefaultSyncState = CachedLastSyncState.SyncStateCollection.FindDataByType<FMoverDefaultSyncState>();
+	FMoverSyncState& BufferedSyncState = MoverSyncStateDoubleBuffer.GetWritable();
+	BufferedSyncState = *SyncState;
+	LastMoverDefaultSyncState = BufferedSyncState.SyncStateCollection.FindDataByType<FMoverDefaultSyncState>();
+	MoverSyncStateDoubleBuffer.Flip();
+
+	// TODO: when AuxState starts getting used we need to double buffer it here as well
 	CachedLastAuxState = *AuxState;
 	CachedLastSimTickTimeStep.BaseSimTimeMs = BackendLiaisonComp->GetCurrentSimTimeMs();
 	CachedLastSimTickTimeStep.ServerFrame = BackendLiaisonComp->GetCurrentSimFrame();
@@ -1564,7 +1576,8 @@ DEFINE_FUNCTION(UMoverComponent::execK2_FindActiveLayeredMove)
 
 const FLayeredMoveBase* UMoverComponent::FindActiveLayeredMoveByType(const UScriptStruct* LayeredMoveStructType) const
 {
-	return CachedLastSyncState.LayeredMoves.FindActiveMove(LayeredMoveStructType);
+	const FMoverSyncState& CachedSyncState = MoverSyncStateDoubleBuffer.GetReadable();
+	return CachedSyncState.LayeredMoves.FindActiveMove(LayeredMoveStructType);
 }
 
 void UMoverComponent::QueueNextMode(FName DesiredModeName, bool bShouldReenter)
@@ -1923,7 +1936,7 @@ TArray<FTrajectorySampleInfo> UMoverComponent::GetPredictedTrajectory(FMoverPred
 		}
 		else
 		{
-			StepState.SyncState = CachedLastSyncState;
+			StepState.SyncState = MoverSyncStateDoubleBuffer.GetReadable();
 		}
 
 		if (PredictionParams.OptionalStartAuxState.IsSet())
@@ -2032,7 +2045,7 @@ TArray<FTrajectorySampleInfo> UMoverComponent::GetPredictedTrajectory(FMoverPred
 
 FName UMoverComponent::GetMovementModeName() const
 { 
-	return CachedLastSyncState.MovementMode;
+	return MoverSyncStateDoubleBuffer.GetReadable().MovementMode;
 }
 
 const UBaseMovementMode* UMoverComponent::GetMovementMode() const
@@ -2067,7 +2080,7 @@ bool UMoverComponent::HasValidCachedState() const
 
 const FMoverSyncState& UMoverComponent::GetSyncState() const
 {
-	return CachedLastSyncState;
+	return MoverSyncStateDoubleBuffer.GetReadable();
 }
 
 bool UMoverComponent::TryGetFloorCheckHitResult(FHitResult& OutHitResult) const
@@ -2175,7 +2188,7 @@ UBaseMovementMode* UMoverComponent::FindMode_Mutable(TSubclassOf<UBaseMovementMo
 
 UBaseMovementMode* UMoverComponent::GetActiveModeInternal(TSubclassOf<UBaseMovementMode> ModeType, bool bRequireExactClass) const
 {
-	if (const TObjectPtr<UBaseMovementMode>* CurrentMode = MovementModes.Find(CachedLastSyncState.MovementMode))
+	if (const TObjectPtr<UBaseMovementMode>* CurrentMode = MovementModes.Find(GetMovementModeName()))
 	{
 		if ((!bRequireExactClass && CurrentMode->IsA(ModeType)) ||
 			CurrentMode->GetClass() == ModeType)
@@ -2193,8 +2206,10 @@ void UMoverComponent::SetSimulationOutput(const FMoverTimeStep& TimeStep, const
 
 	CachedLastUsedInputCmd = OutputData.LastUsedInputCmd;
 
-	CachedLastSyncState = OutputData.SyncState;
-	LastMoverDefaultSyncState = CachedLastSyncState.SyncStateCollection.FindDataByType<FMoverDefaultSyncState>();
+	FMoverSyncState& BufferedSyncState = MoverSyncStateDoubleBuffer.GetWritable();
+	BufferedSyncState = OutputData.SyncState;
+	LastMoverDefaultSyncState = BufferedSyncState.SyncStateCollection.FindDataByType<FMoverDefaultSyncState>();
+	MoverSyncStateDoubleBuffer.Flip();
 
 	// This is for things like the ground info that we want to cache and interpolate but isn't part of the networked sync state
 	// AdditionalOutputData is generic because ground info might not be useful for platforms, say, but platforms may want to store something else.
diff --git a/Engine/Plugins/Experimental/Mover/Source/Mover/Public/MoverComponent.h b/Engine/Plugins/Experimental/Mover/Source/Mover/Public/MoverComponent.h
--- a/Engine/Plugins/Experimental/Mover/Source/Mover/Public/MoverComponent.h
+++ b/Engine/Plugins/Experimental/Mover/Source/Mover/Public/MoverComponent.h
@@ -593,7 +593,8 @@ protected:
 	
 	FMoverInputCmdContext CachedLastUsedInputCmd;
 	
-	FMoverSyncState CachedLastSyncState;	
+	FMoverDoubleBuffer<FMoverSyncState> MoverSyncStateDoubleBuffer;
+	
 	const FMoverDefaultSyncState* LastMoverDefaultSyncState = nullptr;
 
 	FMoverTimeStep CachedLastSimTickTimeStep;	// Saved timestep info from our last simulation tick, used during rollback handling. This will rewind during corrections.
diff --git a/Engine/Plugins/Experimental/Mover/Source/Mover/Public/MoverSimulationTypes.h b/Engine/Plugins/Experimental/Mover/Source/Mover/Public/MoverSimulationTypes.h
--- a/Engine/Plugins/Experimental/Mover/Source/Mover/Public/MoverSimulationTypes.h
+++ b/Engine/Plugins/Experimental/Mover/Source/Mover/Public/MoverSimulationTypes.h
@@ -177,6 +177,42 @@ public:
 	}
 };
 
+/** 
+ *  Double Buffer struct for various Mover data. 
+ */
+template<typename T>
+struct FMoverDoubleBuffer
+{
+	// Sets all buffered data - usually used for initializing data
+	void SetBufferedData(const T& InDataToCopy)
+	{
+		Buffer[0] = InDataToCopy;
+		Buffer[1] = InDataToCopy;
+	}
+	
+	// Gets data that is safe to read and is not being written to
+	const T& GetReadable() const
+	{
+		return Buffer[ReadIndex];
+	}
+
+	// Gets data that is being written to and is expected to change
+	T& GetWritable()
+	{
+		return Buffer[(ReadIndex + 1) % 2];
+	}
+
+	// Flips which data in the buffer we return for reading and writing
+	void Flip()
+	{
+		ReadIndex = (ReadIndex + 1) % 2;
+	}
+	
+private:
+	uint32 ReadIndex = 0;
+	T Buffer[2];
+};
+
 // Auxiliary state that is input into the simulation (changes rarely)
 USTRUCT(BlueprintType)
 struct FMoverAuxStateContext
