Mover: Finding a movement modifier by handle and by type now searches through the queued modifiers on the Mover State Machine as well. Fixes an issue where modifiers weren't found immediately after queueing them.

diff --git a/Engine/Plugins/Experimental/Mover/Source/Mover/Private/MovementModeStateMachine.cpp b/Engine/Plugins/Experimental/Mover/Source/Mover/Private/MovementModeStateMachine.cpp
--- a/Engine/Plugins/Experimental/Mover/Source/Mover/Private/MovementModeStateMachine.cpp
+++ b/Engine/Plugins/Experimental/Mover/Source/Mover/Private/MovementModeStateMachine.cpp
@@ -522,6 +522,32 @@ void UMovementModeStateMachine::CancelModifierFromHandle(FMovementModifierHandle
 	}
 }
 
+const FMovementModifierBase* UMovementModeStateMachine::FindQueuedModifier(FMovementModifierHandle ModifierHandle) const
+{
+	for (const TSharedPtr<FMovementModifierBase>& QueuedModifier : QueuedMovementModifiers)
+	{
+		if (QueuedModifier->GetHandle() == ModifierHandle)
+		{
+			return QueuedModifier.Get();
+		}
+	}
+
+	return nullptr;
+}
+
+const FMovementModifierBase* UMovementModeStateMachine::FindQueuedModifierByType(const UScriptStruct* ModifierType) const
+{
+	for (const TSharedPtr<FMovementModifierBase>& QueuedModifier : QueuedMovementModifiers)
+	{
+		if (QueuedModifier->GetScriptStruct() == ModifierType)
+		{
+			return QueuedModifier.Get();
+		}
+	}
+
+	return nullptr;
+}
+
 
 void UMovementModeStateMachine::ConstructDefaultModes()
 {
diff --git a/Engine/Plugins/Experimental/Mover/Source/Mover/Private/MoverComponent.cpp b/Engine/Plugins/Experimental/Mover/Source/Mover/Private/MoverComponent.cpp
--- a/Engine/Plugins/Experimental/Mover/Source/Mover/Private/MoverComponent.cpp
+++ b/Engine/Plugins/Experimental/Mover/Source/Mover/Private/MoverComponent.cpp
@@ -695,7 +695,7 @@ const FMovementModifierBase* UMoverComponent::FindMovementModifier(const FMoveme
 		}
 	}
 	
-	return nullptr;
+	return ModeFSM->FindQueuedModifier(ModifierHandle);
 }
 
 const FMovementModifierBase* UMoverComponent::FindMovementModifierByType(const UScriptStruct* DataStructType) const
@@ -724,7 +724,7 @@ const FMovementModifierBase* UMoverComponent::FindMovementModifierByType(const U
 		}
 	}
 	
-	return nullptr;
+	return ModeFSM->FindQueuedModifierByType(DataStructType);
 }
 
 bool UMoverComponent::HasGameplayTag(FGameplayTag TagToFind, bool bExactMatch) const
diff --git a/Engine/Plugins/Experimental/Mover/Source/Mover/Public/MovementModeStateMachine.h b/Engine/Plugins/Experimental/Mover/Source/Mover/Public/MovementModeStateMachine.h
--- a/Engine/Plugins/Experimental/Mover/Source/Mover/Public/MovementModeStateMachine.h
+++ b/Engine/Plugins/Experimental/Mover/Source/Mover/Public/MovementModeStateMachine.h
@@ -59,6 +59,9 @@ public:
 	UE_API FMovementModifierHandle QueueMovementModifier(TSharedPtr<FMovementModifierBase> Modifier);
 
  	UE_API void CancelModifierFromHandle(FMovementModifierHandle ModifierHandle);
+
+ 	UE_API const FMovementModifierBase* FindQueuedModifier(FMovementModifierHandle ModifierHandle) const;
+ 	UE_API const FMovementModifierBase* FindQueuedModifierByType(const UScriptStruct* ModifierType) const;
  	
 protected:
 
