Mover: supporting the combination of root motion + falling modes such that everything except vertical movement can be root motion, with falling modes governing gravity acceleration. Also fixes issues where gameplay tag checks could be against last-frame's state.

diff --git a/Engine/Plugins/Experimental/ChaosMover/Source/ChaosMover/Private/ChaosMoverStateMachine.cpp b/Engine/Plugins/Experimental/ChaosMover/Source/ChaosMover/Private/ChaosMoverStateMachine.cpp
--- a/Engine/Plugins/Experimental/ChaosMover/Source/ChaosMover/Private/ChaosMoverStateMachine.cpp
+++ b/Engine/Plugins/Experimental/ChaosMover/Source/ChaosMover/Private/ChaosMoverStateMachine.cpp
@@ -289,7 +289,13 @@ namespace UE::ChaosMover
 
 				if (bHasModeMoveContribution && bHasLayeredMoveContributions)
 				{
-					MovementMixer->MixProposedMoves(CombinedLayeredMove, CombinedMove);
+					FVector UpDir = FVector::UpVector;
+					if (const FChaosMoverSimulationDefaultInputs* DefaultSimInputs = Simulation->GetLocalSimInput().FindDataByType<FChaosMoverSimulationDefaultInputs>())
+					{
+						UpDir = DefaultSimInputs->UpDir;
+					}
+
+					MovementMixer->MixProposedMoves(CombinedLayeredMove, UpDir, CombinedMove);
 				}
 				else if (bHasLayeredMoveContributions && !bHasModeMoveContribution)
 				{
diff --git a/Engine/Plugins/Experimental/Mover/Source/Mover/Private/DefaultMovementSet/LayeredMoves/RootMotionAttributeLayeredMove.cpp b/Engine/Plugins/Experimental/Mover/Source/Mover/Private/DefaultMovementSet/LayeredMoves/RootMotionAttributeLayeredMove.cpp
--- a/Engine/Plugins/Experimental/Mover/Source/Mover/Private/DefaultMovementSet/LayeredMoves/RootMotionAttributeLayeredMove.cpp
+++ b/Engine/Plugins/Experimental/Mover/Source/Mover/Private/DefaultMovementSet/LayeredMoves/RootMotionAttributeLayeredMove.cpp
@@ -131,7 +131,7 @@ bool FLayeredMove_RootMotionAttribute::GenerateMove(const FMoverTickStartData& S
 		const FCharacterDefaultInputs* InputCmd = SimState.InputCmd.InputCollection.FindDataByType<FCharacterDefaultInputs>();
 
 		if (!bHasValidWarpingContext &&
-				((InputCmd && InputCmd->bIsJumpJustPressed) || MoverComp->HasGameplayTag(Mover_SkipAnimRootMotion, false)))
+				((InputCmd && InputCmd->bIsJumpJustPressed) || MoverComp->HasGameplayTagInState(SimState.SyncState, Mover_SkipAnimRootMotion, false)))
 		{
 			return false;	// do not perform root motion while we have the SkipAnimRootMotion tag. i.e. jumping or falling, so that we can have air control, unless we're under montage influence
 		}
@@ -151,6 +151,11 @@ bool FLayeredMove_RootMotionAttribute::GenerateMove(const FMoverTickStartData& S
 		OutProposedMove = FProposedMove();
 		OutProposedMove.MixMode = MixMode;
 
+		if (MoverComp->HasGameplayTagInState(SimState.SyncState, Mover_SkipVerticalAnimRootMotion, false) && MixMode == EMoveMixMode::OverrideAll)
+		{
+			OutProposedMove.MixMode = EMoveMixMode::OverrideAllExceptVerticalVelocity;
+		}
+
 		// Convert the transform into linear and angular velocities
 		OutProposedMove.LinearVelocity  = WorldSpaceRootMotion.GetTranslation() * (1.f / DeltaSeconds);
 		OutProposedMove.AngularVelocity = WorldSpaceRootMotion.GetRotation().Rotator() * (1.f / DeltaSeconds);
diff --git a/Engine/Plugins/Experimental/Mover/Source/Mover/Private/DefaultMovementSet/Modes/AsyncFallingMode.cpp b/Engine/Plugins/Experimental/Mover/Source/Mover/Private/DefaultMovementSet/Modes/AsyncFallingMode.cpp
--- a/Engine/Plugins/Experimental/Mover/Source/Mover/Private/DefaultMovementSet/Modes/AsyncFallingMode.cpp
+++ b/Engine/Plugins/Experimental/Mover/Source/Mover/Private/DefaultMovementSet/Modes/AsyncFallingMode.cpp
@@ -31,7 +31,7 @@ UAsyncFallingMode::UAsyncFallingMode(const FObjectInitializer& ObjectInitializer
 
 	GameplayTags.AddTag(Mover_IsInAir);
 	GameplayTags.AddTag(Mover_IsFalling);
-	GameplayTags.AddTag(Mover_SkipAnimRootMotion);
+	GameplayTags.AddTag(Mover_SkipVerticalAnimRootMotion);	// allows combination of gravity falling and root motion
 }
 
 
diff --git a/Engine/Plugins/Experimental/Mover/Source/Mover/Private/DefaultMovementSet/Modes/FallingMode.cpp b/Engine/Plugins/Experimental/Mover/Source/Mover/Private/DefaultMovementSet/Modes/FallingMode.cpp
--- a/Engine/Plugins/Experimental/Mover/Source/Mover/Private/DefaultMovementSet/Modes/FallingMode.cpp
+++ b/Engine/Plugins/Experimental/Mover/Source/Mover/Private/DefaultMovementSet/Modes/FallingMode.cpp
@@ -28,7 +28,7 @@ UFallingMode::UFallingMode(const FObjectInitializer& ObjectInitializer)
 
 	GameplayTags.AddTag(Mover_IsInAir);
 	GameplayTags.AddTag(Mover_IsFalling);
-	GameplayTags.AddTag(Mover_SkipAnimRootMotion);
+	GameplayTags.AddTag(Mover_SkipVerticalAnimRootMotion);	// allows combination of gravity falling and root motion
 }
 
 
diff --git a/Engine/Plugins/Experimental/Mover/Source/Mover/Private/MoveLibrary/MovementMixer.cpp b/Engine/Plugins/Experimental/Mover/Source/Mover/Private/MoveLibrary/MovementMixer.cpp
--- a/Engine/Plugins/Experimental/Mover/Source/Mover/Private/MoveLibrary/MovementMixer.cpp
+++ b/Engine/Plugins/Experimental/Mover/Source/Mover/Private/MoveLibrary/MovementMixer.cpp
@@ -29,7 +29,7 @@ void UMovementMixer::MixLayeredMove(const FLayeredMoveBase& ActiveMove, const FP
 		OutCumulativeMove.DirectionIntent = MoveStep.DirectionIntent;
 	}
 
-	if (ActiveMove.MixMode == EMoveMixMode::OverrideVelocity)
+	if (MoveStep.MixMode == EMoveMixMode::OverrideVelocity)
 	{
 		if (CheckPriority(&ActiveMove, CurrentHighestPriority, CurrentLayeredMoveStartTime))
 		{
@@ -48,7 +48,7 @@ void UMovementMixer::MixLayeredMove(const FLayeredMoveBase& ActiveMove, const FP
 			OutCumulativeMove.AngularVelocity = MoveStep.AngularVelocity;
 		}
 	}
-	else if (ActiveMove.MixMode == EMoveMixMode::AdditiveVelocity)
+	else if (MoveStep.MixMode == EMoveMixMode::AdditiveVelocity)
 	{
 		if (OutCumulativeMove.MixMode != EMoveMixMode::OverrideVelocity && OutCumulativeMove.MixMode != EMoveMixMode::OverrideAll)
 		{
@@ -61,7 +61,7 @@ void UMovementMixer::MixLayeredMove(const FLayeredMoveBase& ActiveMove, const FP
 			OutCumulativeMove.AngularVelocity += MoveStep.AngularVelocity;
 		}
 	}
-	else if (ActiveMove.MixMode == EMoveMixMode::OverrideAll)
+	else if (MoveStep.MixMode == EMoveMixMode::OverrideAll)
 	{
 		if (CheckPriority(&ActiveMove, CurrentHighestPriority, CurrentLayeredMoveStartTime))
 		{
@@ -74,13 +74,26 @@ void UMovementMixer::MixLayeredMove(const FLayeredMoveBase& ActiveMove, const FP
 			OutCumulativeMove.MixMode = EMoveMixMode::OverrideAll;
 		}
 	}
+	else if (MoveStep.MixMode == EMoveMixMode::OverrideAllExceptVerticalVelocity)
+	{
+		if (CheckPriority(&ActiveMove, CurrentHighestPriority, CurrentLayeredMoveStartTime))
+		{
+			if (OutCumulativeMove.MixMode == EMoveMixMode::OverrideVelocity || OutCumulativeMove.MixMode == EMoveMixMode::OverrideAll || OutCumulativeMove.MixMode == EMoveMixMode::OverrideAllExceptVerticalVelocity)
+			{
+				UE_LOG(LogMover, Log, TEXT("Multiple LayeredMoves with Override mix mode are active simultaneously. Layered move with the highest priority will take effect."));
+			}
+
+			OutCumulativeMove = MoveStep;
+			OutCumulativeMove.MixMode = EMoveMixMode::OverrideAllExceptVerticalVelocity;
+		}
+	}
 	else
 	{
 		ensureMsgf(false, TEXT("Unhandled move mix mode was found."));
 	}
 }
 
-void UMovementMixer::MixProposedMoves(const FProposedMove& MoveToMix, FProposedMove& OutCumulativeMove)
+void UMovementMixer::MixProposedMoves(const FProposedMove& MoveToMix, FVector UpDirection, FProposedMove& OutCumulativeMove)
 {
 	if (MoveToMix.bHasDirIntent && OutCumulativeMove.MixMode != EMoveMixMode::OverrideAll)
 	{
@@ -103,6 +116,15 @@ void UMovementMixer::MixProposedMoves(const FProposedMove& MoveToMix, FProposedM
 		OutCumulativeMove.LinearVelocity = MoveToMix.LinearVelocity;
 		OutCumulativeMove.AngularVelocity = MoveToMix.AngularVelocity;
 	}
+	else if (MoveToMix.MixMode == EMoveMixMode::OverrideAllExceptVerticalVelocity)
+	{
+		const FVector IncomingVerticalVelocity = MoveToMix.LinearVelocity.ProjectOnToNormal(UpDirection);
+		const FVector IncomingNonVerticalVelocity = MoveToMix.LinearVelocity - IncomingVerticalVelocity;
+		const FVector ExistingVerticalVelocity = OutCumulativeMove.LinearVelocity.ProjectOnToNormal(UpDirection);
+
+		OutCumulativeMove = MoveToMix;
+		OutCumulativeMove.LinearVelocity = IncomingNonVerticalVelocity + ExistingVerticalVelocity;
+	}
 	else
 	{
 		ensureMsgf(false, TEXT("Unhandled move mix mode was found."));
diff --git a/Engine/Plugins/Experimental/Mover/Source/Mover/Private/MovementModeStateMachine.cpp b/Engine/Plugins/Experimental/Mover/Source/Mover/Private/MovementModeStateMachine.cpp
--- a/Engine/Plugins/Experimental/Mover/Source/Mover/Private/MovementModeStateMachine.cpp
+++ b/Engine/Plugins/Experimental/Mover/Source/Mover/Private/MovementModeStateMachine.cpp
@@ -176,6 +176,8 @@ void UMovementModeStateMachine::OnSimulationTick(USceneComponent* UpdatedCompone
 	int32 NumConsecutiveFullRefundedSubsteps = 0;
 	while (SubTimeStep.BaseSimTimeMs < EndingSimTimeMs)
 	{
+		WorkingSubstepStartData.SyncState.MovementMode = CurrentModeName;
+
 		FMoverDefaultSyncState* OutputSyncState = &OutputState.SyncState.SyncStateCollection.FindOrAddMutableDataByType<FMoverDefaultSyncState>();
 		OutputState.SyncState.MovementMode = CurrentModeName;
 
@@ -243,6 +245,8 @@ void UMovementModeStateMachine::OnSimulationTick(USceneComponent* UpdatedCompone
 		for (TSharedPtr<FLayeredMoveBase>& ActiveMove : ActiveMoves)
 		{
 			FProposedMove MoveStep;
+			MoveStep.MixMode = ActiveMove->MixMode;	// Initialize using the move's mixmode, but allow it to be changed in GenerateMove
+
 			if (ActiveMove->GenerateMove(WorkingSubstepStartData, SubTimeStep, MoverComp, SimBlackboard, MoveStep))
 			{
 				// If this active move is already past it's first tick we don't need to set the preferred mode again
@@ -280,7 +284,7 @@ void UMovementModeStateMachine::OnSimulationTick(USceneComponent* UpdatedCompone
 
 			if (bHasModeMoveContribution && bHasLayeredMoveContributions)
 			{
-				MoverComp->MovementMixer->MixProposedMoves(CombinedLayeredMove, CombinedMove);
+				MoverComp->MovementMixer->MixProposedMoves(CombinedLayeredMove, MoverComp->GetUpDirection(), CombinedMove);
 			}
 			else if (bHasLayeredMoveContributions && !bHasModeMoveContribution)
 			{
@@ -429,6 +433,11 @@ void UMovementModeStateMachine::OnSimulationRollback(const FMoverSyncState* Sync
 {
 	ClearQueuedMode();
 
+	if (CurrentModeName != SyncState->MovementMode)
+	{
+		SetModeImmediately(SyncState->MovementMode);
+	}
+
 	{
 		FRWScopeLock QueueLock(LayeredMoveQueueLock, SLT_Write);
 		QueuedLayeredMoves.Empty();
diff --git a/Engine/Plugins/Experimental/Mover/Source/Mover/Private/MoverComponent.cpp b/Engine/Plugins/Experimental/Mover/Source/Mover/Private/MoverComponent.cpp
--- a/Engine/Plugins/Experimental/Mover/Source/Mover/Private/MoverComponent.cpp
+++ b/Engine/Plugins/Experimental/Mover/Source/Mover/Private/MoverComponent.cpp
@@ -729,12 +729,18 @@ const FMovementModifierBase* UMoverComponent::FindMovementModifierByType(const U
 
 bool UMoverComponent::HasGameplayTag(FGameplayTag TagToFind, bool bExactMatch) const
 {
+	return HasGameplayTagInState(MoverSyncStateDoubleBuffer.GetReadable(), TagToFind, bExactMatch);
+}
+
+bool UMoverComponent::HasGameplayTagInState(const FMoverSyncState& SyncState, FGameplayTag TagToFind, bool bExactMatch) const 
+{
+	// Check loose / external tags
 	if (bExactMatch)
 	{
 		if (ExternalGameplayTags.HasTagExact(TagToFind))
 		{
 			return true;
-		}	
+		}
 	}
 	else
 	{
@@ -743,19 +749,18 @@ bool UMoverComponent::HasGameplayTag(FGameplayTag TagToFind, bool bExactMatch) c
 			return true;
 		}
 	}
-	
-	// Search Movement Modes
-	if (const UBaseMovementMode* ActiveMovementMode = GetMovementMode())
+
+	// Check active Movement Mode
+	if (const UBaseMovementMode* ActiveMovementMode = FindMovementModeByName(SyncState.MovementMode))
 	{
 		if (ActiveMovementMode->HasGameplayTag(TagToFind, bExactMatch))
 		{
 			return true;
 		}
 	}
-		
-	const FMoverSyncState& CachedSyncState = MoverSyncStateDoubleBuffer.GetReadable();
+
 	// Search Movement Modifiers
-	for (auto ModifierFromSyncStateIt = CachedSyncState.MovementModifiers.GetActiveModifiersIterator(); ModifierFromSyncStateIt; ++ModifierFromSyncStateIt)
+	for (auto ModifierFromSyncStateIt = SyncState.MovementModifiers.GetActiveModifiersIterator(); ModifierFromSyncStateIt; ++ModifierFromSyncStateIt)
 	{
 		if (const TSharedPtr<FMovementModifierBase> ModifierFromSyncState = *ModifierFromSyncStateIt)
 		{
@@ -765,19 +770,20 @@ bool UMoverComponent::HasGameplayTag(FGameplayTag TagToFind, bool bExactMatch) c
 			}
 		}
 	}
-		
+
 	// Search Layered Moves
-	for (const TSharedPtr<FLayeredMoveBase>& LayeredMove :  CachedSyncState.LayeredMoves.GetActiveMoves())
+	for (const TSharedPtr<FLayeredMoveBase>& LayeredMove : SyncState.LayeredMoves.GetActiveMoves())
 	{
 		if (LayeredMove->HasGameplayTag(TagToFind, bExactMatch))
 		{
 			return true;
 		}
 	}
-	
+
 	return false;
 }
 
+
 void UMoverComponent::AddGameplayTag(FGameplayTag TagToAdd)
 {
 	ExternalGameplayTags.AddTag(TagToAdd);
diff --git a/Engine/Plugins/Experimental/Mover/Source/Mover/Private/MoverTypes.cpp b/Engine/Plugins/Experimental/Mover/Source/Mover/Private/MoverTypes.cpp
--- a/Engine/Plugins/Experimental/Mover/Source/Mover/Private/MoverTypes.cpp
+++ b/Engine/Plugins/Experimental/Mover/Source/Mover/Private/MoverTypes.cpp
@@ -18,6 +18,7 @@ UE_DEFINE_GAMEPLAY_TAG_COMMENT(Mover_IsSwimming, "Mover.IsSwimming", "Default Mo
 UE_DEFINE_GAMEPLAY_TAG_COMMENT(Mover_IsCrouching, "Mover.Stance.IsCrouching", "Default Mover state flag indicating character is crouching.");
 UE_DEFINE_GAMEPLAY_TAG_COMMENT(Mover_IsNavWalking, "Mover.IsNavWalking", "Default Mover state flag indicating character is NavWalking.");
 UE_DEFINE_GAMEPLAY_TAG_COMMENT(Mover_SkipAnimRootMotion, "Mover.SkipAnimRootMotion", "Default Mover state flag indicating Animation Root Motion proposed movement should be skipped.");
+UE_DEFINE_GAMEPLAY_TAG_COMMENT(Mover_SkipVerticalAnimRootMotion, "Mover.SkipVerticalAnimRootMotion", "Default Mover state flag indicating Animation Root Motion proposed movements should not include a vertical velocity component (along the up/down axis).");
 
 FMoverOnImpactParams::FMoverOnImpactParams() 
 	: AttemptedMoveDelta(0) 
diff --git a/Engine/Plugins/Experimental/Mover/Source/Mover/Public/MoveLibrary/MovementMixer.h b/Engine/Plugins/Experimental/Mover/Source/Mover/Public/MoveLibrary/MovementMixer.h
--- a/Engine/Plugins/Experimental/Mover/Source/Mover/Public/MoveLibrary/MovementMixer.h
+++ b/Engine/Plugins/Experimental/Mover/Source/Mover/Public/MoveLibrary/MovementMixer.h
@@ -25,7 +25,7 @@ public:
 	UE_API virtual void MixLayeredMove(const FLayeredMoveBase& ActiveMove, const FProposedMove& MoveStep, FProposedMove& OutCumulativeMove);
 
 	/** In charge of mixing proposed moves together. Is similar to MixLayeredMove but is only responsible for mixing proposed moves instead of layered moves. */
-	UE_API virtual void MixProposedMoves(const FProposedMove& MoveToMix, FProposedMove& OutCumulativeMove);
+	UE_API virtual void MixProposedMoves(const FProposedMove& MoveToMix, FVector UpDirection, FProposedMove& OutCumulativeMove);
 
 	/** Resets all state used for mixing. Should be called before or after finished mixing moves. */
 	UE_API virtual void ResetMixerState();
diff --git a/Engine/Plugins/Experimental/Mover/Source/Mover/Public/MoveLibrary/MovementUtilsTypes.h b/Engine/Plugins/Experimental/Mover/Source/Mover/Public/MoveLibrary/MovementUtilsTypes.h
--- a/Engine/Plugins/Experimental/Mover/Source/Mover/Public/MoveLibrary/MovementUtilsTypes.h
+++ b/Engine/Plugins/Experimental/Mover/Source/Mover/Public/MoveLibrary/MovementUtilsTypes.h
@@ -22,6 +22,10 @@ enum class EMoveMixMode : uint8
 	OverrideVelocity = 1,
 	/** All move parameters should override others */
 	OverrideAll      = 2,
+
+	/** All move parameters should override others except linear velocity along the up/down axis. Commonly used for falling. */
+	OverrideAllExceptVerticalVelocity = 3,
+
 };
 
 
diff --git a/Engine/Plugins/Experimental/Mover/Source/Mover/Public/MoverComponent.h b/Engine/Plugins/Experimental/Mover/Source/Mover/Public/MoverComponent.h
--- a/Engine/Plugins/Experimental/Mover/Source/Mover/Public/MoverComponent.h
+++ b/Engine/Plugins/Experimental/Mover/Source/Mover/Public/MoverComponent.h
@@ -478,6 +478,17 @@ public:	// Queries
 	UFUNCTION(BlueprintPure, Category = Mover, meta = (Keywords = "HasTag"))
 	MOVER_API bool HasGameplayTag(FGameplayTag TagToFind, bool bExactMatch) const;
 
+	/**
+	 * Check Mover systems for a gameplay tag. Use the given state, as well as any loose tags on the MoverComponent.
+	 *
+	 * @param TagToFind			Tag to check on the MoverComponent or state
+	 * @param bExactMatch		If true, the tag has to be exactly present, if false then TagToFind will include it's parent tags while matching
+	 *
+	 * @return True if the TagToFind was found
+	 */
+	UFUNCTION(BlueprintPure, Category = Mover, meta = (Keywords = "HasTag"))
+	MOVER_API bool HasGameplayTagInState(const FMoverSyncState& SyncState, FGameplayTag TagToFind, bool bExactMatch) const;
+
 	/**
   	 * Adds a gameplay tag to this Mover Component.
   	 * Note: Duplicate tags will not be added
diff --git a/Engine/Plugins/Experimental/Mover/Source/Mover/Public/MoverTypes.h b/Engine/Plugins/Experimental/Mover/Source/Mover/Public/MoverTypes.h
--- a/Engine/Plugins/Experimental/Mover/Source/Mover/Public/MoverTypes.h
+++ b/Engine/Plugins/Experimental/Mover/Source/Mover/Public/MoverTypes.h
@@ -21,6 +21,7 @@ MOVER_API UE_DECLARE_GAMEPLAY_TAG_EXTERN(Mover_IsSwimming);
 MOVER_API UE_DECLARE_GAMEPLAY_TAG_EXTERN(Mover_IsCrouching);
 MOVER_API UE_DECLARE_GAMEPLAY_TAG_EXTERN(Mover_IsNavWalking);
 MOVER_API UE_DECLARE_GAMEPLAY_TAG_EXTERN(Mover_SkipAnimRootMotion);
+MOVER_API UE_DECLARE_GAMEPLAY_TAG_EXTERN(Mover_SkipVerticalAnimRootMotion);
 
 /** List of Mover tick phases where different parts of work are being performed */
 UENUM(BlueprintType)
