Mover: Fixing floating point precision error in mover's fixed delta time calculation

diff --git a/Engine/Plugins/Experimental/Mover/Source/Mover/Private/MovementModeStateMachine.cpp b/Engine/Plugins/Experimental/Mover/Source/Mover/Private/MovementModeStateMachine.cpp
--- a/Engine/Plugins/Experimental/Mover/Source/Mover/Private/MovementModeStateMachine.cpp
+++ b/Engine/Plugins/Experimental/Mover/Source/Mover/Private/MovementModeStateMachine.cpp
@@ -154,8 +154,6 @@ void UMovementModeStateMachine::ClearQueuedMode()
 
 void UMovementModeStateMachine::OnSimulationTick(USceneComponent* UpdatedComponent, UPrimitiveComponent* UpdatedPrimitive, UMoverBlackboard* SimBlackboard, const FMoverTickStartData& StartState, const FMoverTimeStep& TimeStep, FMoverTickEndData& OutputState)
 {
-	const float EndingSimTimeMs = TimeStep.BaseSimTimeMs + TimeStep.StepMs;
-
 	//GEngine->AddOnScreenDebugMessage(-1, -0.1f, FColor::White, FString::Printf(TEXT("Mode Tick: %s  Queued: %s"), *CurrentModeName.ToString(), *NextModeName.ToString()));
 
 	FMoverTimeStep SubTimeStep = TimeStep;
@@ -174,7 +172,9 @@ void UMovementModeStateMachine::OnSimulationTick(USceneComponent* UpdatedCompone
 	int SubStepCount = 0;
 	const int32 MaxConsecutiveFullRefundedSubsteps = GetDefault<UMoverDeveloperSettings>()->MaxTimesToRefundSubstep;
 	int32 NumConsecutiveFullRefundedSubsteps = 0;
-	while (SubTimeStep.BaseSimTimeMs < EndingSimTimeMs)
+
+	float TotalUsedMs = 0.0f;
+	while (TotalUsedMs < TimeStep.StepMs)
 	{
 		WorkingSubstepStartData.SyncState.MovementMode = CurrentModeName;
 
@@ -183,7 +183,7 @@ void UMovementModeStateMachine::OnSimulationTick(USceneComponent* UpdatedCompone
 
 		OutputState.MovementEndState.ResetToDefaults();
 
-		SubTimeStep.StepMs = EndingSimTimeMs - SubTimeStep.BaseSimTimeMs;		// TODO: convert this to an overridable function that can support MaxStepTime, MaxIterations, etc.
+		SubTimeStep.StepMs = TimeStep.StepMs - TotalUsedMs;		// TODO: convert this to an overridable function that can support MaxStepTime, MaxIterations, etc.
 
 		// Transfer any queued moves into the starting state. They'll be started during the move generation.
 		FlushQueuedMovesToGroup(WorkingSubstepStartData.SyncState.LayeredMoves);
@@ -377,7 +377,7 @@ void UMovementModeStateMachine::OnSimulationTick(USceneComponent* UpdatedCompone
 				if (NumConsecutiveFullRefundedSubsteps >= MaxConsecutiveFullRefundedSubsteps)
 				{
 					UE_LOG(LogMover, Warning, TEXT("Movement mode %s and %s on %s are stuck giving time back to each other. Overriding to advance to next substep."), *CurrentModeName.ToString(), *OutputState.MovementEndState.NextModeName.ToString(), *MoverComp->GetOwner()->GetName());
-					SubTimeStep.BaseSimTimeMs += SubTimeStep.StepMs;
+					TotalUsedMs += SubTimeStep.StepMs;
 				}
 			}
 			else
@@ -398,8 +398,10 @@ void UMovementModeStateMachine::OnSimulationTick(USceneComponent* UpdatedCompone
 		}
 		
 		const float RemainingMs = FMath::Clamp(OutputState.MovementEndState.RemainingMs, 0.0f, SubTimeStep.StepMs);
-		SubTimeStep.BaseSimTimeMs += (SubTimeStep.StepMs - RemainingMs);
-		SubTimeStep.StepMs = EndingSimTimeMs - SubTimeStep.BaseSimTimeMs;
+		const float SubstepUsedMs = (SubTimeStep.StepMs - RemainingMs);
+		SubTimeStep.BaseSimTimeMs += SubstepUsedMs;
+		TotalUsedMs += SubstepUsedMs;
+		SubTimeStep.StepMs = RemainingMs;
 
 		WorkingSubstepStartData.SyncState = OutputState.SyncState;
 		WorkingSubstepStartData.AuxState  = OutputState.AuxState;
