Mover: framework to allow state-copying operations when simulation state doesn't change. Supported by standalone backend and async nav walking mode for this first pass.

diff --git a/Engine/Plugins/Experimental/Mover/Source/Mover/Private/Backends/MoverStandaloneLiaison.cpp b/Engine/Plugins/Experimental/Mover/Source/Mover/Private/Backends/MoverStandaloneLiaison.cpp
--- a/Engine/Plugins/Experimental/Mover/Source/Mover/Private/Backends/MoverStandaloneLiaison.cpp
+++ b/Engine/Plugins/Experimental/Mover/Source/Mover/Private/Backends/MoverStandaloneLiaison.cpp
@@ -49,7 +49,8 @@ namespace MoverStandaloneBackendCVars
 		TEXT("Whether to allow Mover simulation ticks to run on any thread. Requires use of threadsafe movement modes that do not modify scene components.\nIndividuals must also have their UseAsyncMovementSimulationTick option enabled.\n0: Game thread only, 1: Any thread"),
 		FConsoleVariableDelegate::CreateStatic(&OnRunMovementSimOnAnyThreadChanged),
 		ECVF_Default);
-}
+
+}	// end MoverStandaloneBackendCVars
 
 UMoverStandaloneLiaisonComponent::UMoverStandaloneLiaisonComponent(const FObjectInitializer& ObjectInitializer)
 	: Super(ObjectInitializer)
@@ -118,6 +119,7 @@ bool UMoverStandaloneLiaisonComponent::WritePendingSyncState(const FMoverSyncSta
 	{
 		FRWScopeLock Lock(StateDataLock, SLT_Write);
 		CachedLastSyncState = SyncStateToWrite;
+		bIsCachedStateDirty = true;
 	}
 
 	return true;
@@ -139,6 +141,7 @@ void UMoverStandaloneLiaisonComponent::BeginPlay()
 
 			FRWScopeLock Lock(StateDataLock, SLT_Write);
 			MoverComp->InitializeSimulationState(OUT &CachedLastSyncState, OUT &CachedLastAuxState);
+			bIsCachedStateDirty = true;
 		}
 		else
 		{
@@ -355,6 +358,8 @@ void UMoverStandaloneLiaisonComponent::TickMovementSimulation(float DeltaSeconds
 
 	WorkingStartData.InputCmd = LastProducedInputCmd;
 
+	WorkingEndData.InitForNewFrame();
+
 	{
 		FRWScopeLock Lock(StateDataLock, SLT_ReadOnly);
 		WorkingStartData.SyncState = CachedLastSyncState;
@@ -367,13 +372,20 @@ void UMoverStandaloneLiaisonComponent::TickMovementSimulation(float DeltaSeconds
 	}
 
 	{
-		FRWScopeLock Lock(StateDataLock, SLT_Write);
-		CachedLastSyncState = WorkingEndData.SyncState;
-		CachedLastAuxState = WorkingEndData.AuxState;
+		if (!CachedLastSyncState.HasSameContents(WorkingEndData.SyncState) ||
+			!WorkingEndData.MovementEndState.bEndedWithNoChanges)
+		{
+			FRWScopeLock Lock(StateDataLock, SLT_Write);
+			CachedLastSyncState = WorkingEndData.SyncState;
+			CachedLastAuxState = WorkingEndData.AuxState;
+			bIsCachedStateDirty = true;
+		}
 	}
 
 }
 
+
+
 void UMoverStandaloneLiaisonComponent::TickApplySimulationState(float DeltaSeconds)
 {
 	TGuardValue<bool> InApplySimulationState(bIsInApplySimulationState, true);
@@ -381,8 +393,23 @@ void UMoverStandaloneLiaisonComponent::TickApplySimulationState(float DeltaSecon
 	{
 		SCOPED_NAMED_EVENT(StandaloneMoverBackendLiaison_Tick_FinalizeFrame, FColor::Green);
 
-		FRWScopeLock Lock(StateDataLock, SLT_ReadOnly);
-		MoverComp->FinalizeFrame(&CachedLastSyncState, &CachedLastAuxState);
+		bool bNeedsFinalizeUnchanged = true;
+
+		{
+			FRWScopeLock Lock(StateDataLock, SLT_ReadOnly);
+			if (bIsCachedStateDirty)
+			{
+				MoverComp->FinalizeFrame(&CachedLastSyncState, &CachedLastAuxState);
+
+				bIsCachedStateDirty = false;
+				bNeedsFinalizeUnchanged = false;
+			}
+		}
+
+		if (bNeedsFinalizeUnchanged)
+		{
+			MoverComp->FinalizeUnchangedFrame();
+		}
 	}
 }
 
diff --git a/Engine/Plugins/Experimental/Mover/Source/Mover/Private/DefaultMovementSet/Modes/AsyncNavWalkingMode.cpp b/Engine/Plugins/Experimental/Mover/Source/Mover/Private/DefaultMovementSet/Modes/AsyncNavWalkingMode.cpp
--- a/Engine/Plugins/Experimental/Mover/Source/Mover/Private/DefaultMovementSet/Modes/AsyncNavWalkingMode.cpp
+++ b/Engine/Plugins/Experimental/Mover/Source/Mover/Private/DefaultMovementSet/Modes/AsyncNavWalkingMode.cpp
@@ -131,6 +131,7 @@ void UAsyncNavWalkingMode::SimulationTick_Implementation(const FSimulationTickPa
 	const UMoverComponent* MoverComp = Cast<UMoverComponent>(GetMoverComponent());
 	if (!MoverComp || !NavMoverComponent.IsValid() || !CommonLegacySettings.IsValid())
 	{
+		OutputState.MovementEndState.bEndedWithNoChanges = true;
 		return;
 	}
 
@@ -287,7 +288,7 @@ void UAsyncNavWalkingMode::SimulationTick_Implementation(const FSimulationTickPa
 					break;
 			}
 
-			CaptureFinalState(LocationInProgress, RotationInProgress.Rotator(), MoveRecord, OutputSyncState);
+			CaptureOutputState(*StartingSyncState, LocationInProgress, RotationInProgress.Rotator(), MoveRecord, OutputSyncState, OutputState);
 			return;
 		}
 
@@ -344,7 +345,7 @@ void UAsyncNavWalkingMode::SimulationTick_Implementation(const FSimulationTickPa
 		MoveRecord.SetDeltaSeconds((Params.TimeStep.StepMs - OutputState.MovementEndState.RemainingMs) * 0.001f);
 	}
 
-	CaptureFinalState(LocationInProgress, RotationInProgress.Rotator(), MoveRecord, OutputSyncState);
+	CaptureOutputState(*StartingSyncState, LocationInProgress, RotationInProgress.Rotator(), MoveRecord, OutputSyncState, OutputState);
 }
 
 bool UAsyncNavWalkingMode::FindNavFloor(const FVector& TestLocation, FNavLocation& OutNavFloorLocation, const INavigationDataInterface* NavData) const
@@ -618,7 +619,7 @@ void UAsyncNavWalkingMode::OnUnregistered()
 	Super::OnUnregistered();
 }
 
-void UAsyncNavWalkingMode::CaptureFinalState(const FVector FinalLocation, const FRotator FinalRotation, const FMovementRecord& Record, FMoverDefaultSyncState& OutputSyncState) const
+void UAsyncNavWalkingMode::CaptureOutputState(const FMoverDefaultSyncState& StartSyncState, const FVector& FinalLocation, const FRotator& FinalRotation, const FMovementRecord& Record, FMoverDefaultSyncState& OutputSyncState, FMoverTickEndData& TickEndData) const
 {
 	UMoverBlackboard* SimBlackboard = GetMoverComponent()->GetSimBlackboard_Mutable();
 
@@ -630,5 +631,7 @@ void UAsyncNavWalkingMode::CaptureFinalState(const FVector FinalLocation, const
 		FinalRotation,
 		Record.GetRelevantVelocity(),
 		nullptr);	// no movement base
+
+	TickEndData.MovementEndState.bEndedWithNoChanges = OutputSyncState.IsNearlyEqual(StartSyncState);
 }
 
diff --git a/Engine/Plugins/Experimental/Mover/Source/Mover/Private/LayeredMove.cpp b/Engine/Plugins/Experimental/Mover/Source/Mover/Private/LayeredMove.cpp
--- a/Engine/Plugins/Experimental/Mover/Source/Mover/Private/LayeredMove.cpp
+++ b/Engine/Plugins/Experimental/Mover/Source/Mover/Private/LayeredMove.cpp
@@ -285,6 +285,35 @@ bool FLayeredMoveGroup::operator!=(const FLayeredMoveGroup& Other) const
 	return !(FLayeredMoveGroup::operator==(Other));
 }
 
+bool FLayeredMoveGroup::HasSameContents(const FLayeredMoveGroup& Other) const
+{
+	// Only compare the types of moves contained, not the state
+	if (ActiveLayeredMoves.Num() != Other.ActiveLayeredMoves.Num() ||
+		QueuedLayeredMoves.Num() != Other.QueuedLayeredMoves.Num())
+	{
+		return false;
+	}
+
+	for (int32 i = 0; i < ActiveLayeredMoves.Num(); ++i)
+	{
+		if (ActiveLayeredMoves[i]->GetScriptStruct() != Other.ActiveLayeredMoves[i]->GetScriptStruct())
+		{
+			return false;
+		}
+	}
+
+	for (int32 i = 0; i < QueuedLayeredMoves.Num(); ++i)
+	{
+		if (QueuedLayeredMoves[i]->GetScriptStruct() != Other.QueuedLayeredMoves[i]->GetScriptStruct())
+		{
+			return false;
+		}
+	}
+
+	return true;
+}
+
+
 void FLayeredMoveGroup::AddStructReferencedObjects(FReferenceCollector& Collector) const
 {
 	for (const TSharedPtr<FLayeredMoveBase>& LayeredMove : ActiveLayeredMoves)
diff --git a/Engine/Plugins/Experimental/Mover/Source/Mover/Private/MoveLibrary/MovementRecord.cpp b/Engine/Plugins/Experimental/Mover/Source/Mover/Private/MoveLibrary/MovementRecord.cpp
--- a/Engine/Plugins/Experimental/Mover/Source/Mover/Private/MoveLibrary/MovementRecord.cpp
+++ b/Engine/Plugins/Experimental/Mover/Source/Mover/Private/MoveLibrary/MovementRecord.cpp
@@ -3,6 +3,18 @@
 #include "MoveLibrary/MovementRecord.h"
 #include "MoverLog.h"
 
+void FMovementRecord::Reset()
+{
+	TotalMoveDelta = FVector::ZeroVector;
+	RelevantMoveDelta = FVector::ZeroVector;
+	TotalDeltaSeconds = 0.f;
+
+	bIsRelevancyLocked = false;
+	bRelevancyLockValue = false;
+
+	Substeps.Empty(Substeps.Num());	// leaving space for the next user, as this is likely to fill up again
+}
+
 void FMovementRecord::Append(FMovementSubstep Substep)
 {
 	if (bIsRelevancyLocked)
diff --git a/Engine/Plugins/Experimental/Mover/Source/Mover/Private/MovementModifier.cpp b/Engine/Plugins/Experimental/Mover/Source/Mover/Private/MovementModifier.cpp
--- a/Engine/Plugins/Experimental/Mover/Source/Mover/Private/MovementModifier.cpp
+++ b/Engine/Plugins/Experimental/Mover/Source/Mover/Private/MovementModifier.cpp
@@ -250,6 +250,36 @@ bool FMovementModifierGroup::operator!=(const FMovementModifierGroup& Other) con
 	return !(FMovementModifierGroup::operator==(Other));
 }
 
+
+bool FMovementModifierGroup::HasSameContents(const FMovementModifierGroup& Other) const
+{
+	// Only compare the types of modifiers contained, not the state
+	if (ActiveModifiers.Num() != Other.ActiveModifiers.Num() ||
+		QueuedModifiers.Num() != Other.QueuedModifiers.Num())
+	{
+		return false;
+	}
+
+	for (int32 i = 0; i < ActiveModifiers.Num(); ++i)
+	{
+		if (ActiveModifiers[i]->GetScriptStruct() != Other.ActiveModifiers[i]->GetScriptStruct())
+		{
+			return false;
+		}
+	}
+
+	for (int32 i = 0; i < QueuedModifiers.Num(); ++i)
+	{
+		if (QueuedModifiers[i]->GetScriptStruct() != Other.QueuedModifiers[i]->GetScriptStruct())
+		{
+			return false;
+		}
+	}
+
+	return true;
+}
+
+
 void FMovementModifierGroup::AddStructReferencedObjects(FReferenceCollector& Collector) const
 {
 	for (const TSharedPtr<FMovementModifierBase>& Modifier : ActiveModifiers)
diff --git a/Engine/Plugins/Experimental/Mover/Source/Mover/Private/MoverComponent.cpp b/Engine/Plugins/Experimental/Mover/Source/Mover/Private/MoverComponent.cpp
--- a/Engine/Plugins/Experimental/Mover/Source/Mover/Private/MoverComponent.cpp
+++ b/Engine/Plugins/Experimental/Mover/Source/Mover/Private/MoverComponent.cpp
@@ -349,10 +349,22 @@ void UMoverComponent::FinalizeFrame(const FMoverSyncState* SyncState, const FMov
 
 	if (OnPostFinalize.IsBound())
 	{
-		OnPostFinalize.Broadcast(*SyncState, *AuxState);
+		OnPostFinalize.Broadcast(MoverSyncStateDoubleBuffer.GetReadable(), CachedLastAuxState);
 	}
 }
 
+void UMoverComponent::FinalizeUnchangedFrame()
+{
+	CachedLastSimTickTimeStep.BaseSimTimeMs = BackendLiaisonComp->GetCurrentSimTimeMs();
+	CachedLastSimTickTimeStep.ServerFrame = BackendLiaisonComp->GetCurrentSimFrame();
+
+	if (OnPostFinalize.IsBound())
+	{
+		OnPostFinalize.Broadcast(MoverSyncStateDoubleBuffer.GetReadable(), CachedLastAuxState);
+	}
+}
+
+
 void UMoverComponent::FinalizeSmoothingFrame(const FMoverSyncState* SyncState, const FMoverAuxStateContext* AuxState)
 {
 	if (PrimaryVisualComponent)
diff --git a/Engine/Plugins/Experimental/Mover/Source/Mover/Private/MoverDataModelTypes.cpp b/Engine/Plugins/Experimental/Mover/Source/Mover/Private/MoverDataModelTypes.cpp
--- a/Engine/Plugins/Experimental/Mover/Source/Mover/Private/MoverDataModelTypes.cpp
+++ b/Engine/Plugins/Experimental/Mover/Source/Mover/Private/MoverDataModelTypes.cpp
@@ -355,6 +355,21 @@ bool FMoverDefaultSyncState::UpdateCurrentMovementBase()
 	return bDidGetBaseTransform;
 }
 
+bool FMoverDefaultSyncState::IsNearlyEqual(const FMoverDefaultSyncState& Other) const
+{
+	const bool bHasSameBaseBaseInfo = (!MovementBase.IsValid() && !Other.MovementBase.IsValid()) ||
+											(MovementBase == Other.MovementBase && 
+											 MovementBaseBoneName == Other.MovementBaseBoneName &&
+											 (MovementBasePos - Other.MovementBasePos).IsNearlyZero() && 
+											 MovementBaseQuat.Equals(Other.MovementBaseQuat));
+
+	return (Location-Other.Location).IsNearlyZero() &&
+		   (Orientation-Other.Orientation).IsNearlyZero() &&
+		   (Velocity-Other.Velocity).IsNearlyZero() &&
+		   (MoveDirectionIntent-Other.MoveDirectionIntent).IsNearlyZero() &&
+		    bHasSameBaseBaseInfo;
+}
+
 FVector FMoverDefaultSyncState::GetLocation_WorldSpace() const
 {
 	if (MovementBase.IsValid())
diff --git a/Engine/Plugins/Experimental/Mover/Source/Mover/Private/MoverTypes.cpp b/Engine/Plugins/Experimental/Mover/Source/Mover/Private/MoverTypes.cpp
--- a/Engine/Plugins/Experimental/Mover/Source/Mover/Private/MoverTypes.cpp
+++ b/Engine/Plugins/Experimental/Mover/Source/Mover/Private/MoverTypes.cpp
@@ -321,6 +321,25 @@ void FMoverDataCollection::Decay(float DecayAmount)
 	}
 }
 
+
+bool FMoverDataCollection::HasSameContents(const FMoverDataCollection& Other) const
+{
+	if (DataArray.Num() != Other.DataArray.Num())
+	{
+		return false;
+	}
+
+	for (int32 i = 0; i < DataArray.Num(); ++i)
+	{
+		if (DataArray[i]->GetScriptStruct() != Other.DataArray[i]->GetScriptStruct())
+		{
+			return false;
+		}
+	}
+
+	return true;
+}
+
 void FMoverDataCollection::AddStructReferencedObjects(FReferenceCollector& Collector) const
 {
 	for (const TSharedPtr<FMoverDataStructBase>& Data : DataArray)
diff --git a/Engine/Plugins/Experimental/Mover/Source/Mover/Public/Backends/MoverStandaloneLiaison.h b/Engine/Plugins/Experimental/Mover/Source/Mover/Public/Backends/MoverStandaloneLiaison.h
--- a/Engine/Plugins/Experimental/Mover/Source/Mover/Public/Backends/MoverStandaloneLiaison.h
+++ b/Engine/Plugins/Experimental/Mover/Source/Mover/Public/Backends/MoverStandaloneLiaison.h
@@ -182,6 +182,7 @@ protected:
 
 	FMoverSyncState CachedLastSyncState;
 	FMoverAuxStateContext CachedLastAuxState;
+	bool bIsCachedStateDirty = false;	// If true, we need to propagate the state to our MoverComponent during ApplyState
 
 	FRWLock StateDataLock;	// used when reading/writing to our cached state (sync/aux) data
 
diff --git a/Engine/Plugins/Experimental/Mover/Source/Mover/Public/DefaultMovementSet/Modes/AsyncNavWalkingMode.h b/Engine/Plugins/Experimental/Mover/Source/Mover/Public/DefaultMovementSet/Modes/AsyncNavWalkingMode.h
--- a/Engine/Plugins/Experimental/Mover/Source/Mover/Public/DefaultMovementSet/Modes/AsyncNavWalkingMode.h
+++ b/Engine/Plugins/Experimental/Mover/Source/Mover/Public/DefaultMovementSet/Modes/AsyncNavWalkingMode.h
@@ -129,7 +129,7 @@ protected:
 	UE_API virtual void OnRegistered(const FName ModeName) override;
 	UE_API virtual void OnUnregistered() override;
 
-	UE_API void CaptureFinalState(const FVector FinalLocation, const FRotator FinalRotation, const FMovementRecord& Record, FMoverDefaultSyncState& OutputSyncState) const;
+	UE_API void CaptureOutputState(const FMoverDefaultSyncState& StartSyncState, const FVector& FinalLocation, const FRotator& FinalRotation, const FMovementRecord& Record, FMoverDefaultSyncState& OutputSyncState, FMoverTickEndData& OutputState) const;
 
 	TWeakObjectPtr<const UCommonLegacyMovementSettings> CommonLegacySettings;
  
diff --git a/Engine/Plugins/Experimental/Mover/Source/Mover/Public/LayeredMove.h b/Engine/Plugins/Experimental/Mover/Source/Mover/Public/LayeredMove.h
--- a/Engine/Plugins/Experimental/Mover/Source/Mover/Public/LayeredMove.h
+++ b/Engine/Plugins/Experimental/Mover/Source/Mover/Public/LayeredMove.h
@@ -206,6 +206,9 @@ struct FLayeredMoveGroup
 	/** Comparison operator */
 	UE_API bool operator!=(const FLayeredMoveGroup& Other) const;
 
+	/** Checks only whether there are matching LayeredMoves, but NOT necessarily identical states of each move */
+	UE_API bool HasSameContents(const FLayeredMoveGroup& Other) const;
+
 	/** Exposes references to GC system */
 	UE_API void AddStructReferencedObjects(FReferenceCollector& Collector) const;
 
diff --git a/Engine/Plugins/Experimental/Mover/Source/Mover/Public/MoveLibrary/MovementRecord.h b/Engine/Plugins/Experimental/Mover/Source/Mover/Public/MoveLibrary/MovementRecord.h
--- a/Engine/Plugins/Experimental/Mover/Source/Mover/Public/MoveLibrary/MovementRecord.h
+++ b/Engine/Plugins/Experimental/Mover/Source/Mover/Public/MoveLibrary/MovementRecord.h
@@ -33,6 +33,8 @@ struct FMovementRecord
 {
 	GENERATED_BODY()
 
+	UE_API void Reset();
+
 	UE_API void Append(FMovementSubstep Substep);
 
 	// TODO: consider changing this to only allow locking as irrelevant
diff --git a/Engine/Plugins/Experimental/Mover/Source/Mover/Public/MovementModifier.h b/Engine/Plugins/Experimental/Mover/Source/Mover/Public/MovementModifier.h
--- a/Engine/Plugins/Experimental/Mover/Source/Mover/Public/MovementModifier.h
+++ b/Engine/Plugins/Experimental/Mover/Source/Mover/Public/MovementModifier.h
@@ -206,6 +206,9 @@ struct FMovementModifierGroup
 	/** Comparison operator */
 	UE_API bool operator!=(const FMovementModifierGroup& Other) const;
 
+	/** Checks only whether there are matching modifiers, but NOT necessarily identical states of each one */
+	UE_API bool HasSameContents(const FMovementModifierGroup& Other) const;
+
 	/** Exposes references to GC system */
 	UE_API void AddStructReferencedObjects(FReferenceCollector& Collector) const;
 
diff --git a/Engine/Plugins/Experimental/Mover/Source/Mover/Public/MoverComponent.h b/Engine/Plugins/Experimental/Mover/Source/Mover/Public/MoverComponent.h
--- a/Engine/Plugins/Experimental/Mover/Source/Mover/Public/MoverComponent.h
+++ b/Engine/Plugins/Experimental/Mover/Source/Mover/Public/MoverComponent.h
@@ -129,6 +129,9 @@ public:
 	// Take output for simulation. Called by backend system.
 	MOVER_API void FinalizeFrame(const FMoverSyncState* SyncState, const FMoverAuxStateContext* AuxState);
 
+	// Take output for simulation when no simulation state changes have occurred. Called by backend system.
+	MOVER_API void FinalizeUnchangedFrame();
+
 	// Take smoothed simulation state. Called by backend system, if supported.
 	MOVER_API void FinalizeSmoothingFrame(const FMoverSyncState* SyncState, const FMoverAuxStateContext* AuxState);
 
diff --git a/Engine/Plugins/Experimental/Mover/Source/Mover/Public/MoverDataModelTypes.h b/Engine/Plugins/Experimental/Mover/Source/Mover/Public/MoverDataModelTypes.h
--- a/Engine/Plugins/Experimental/Mover/Source/Mover/Public/MoverDataModelTypes.h
+++ b/Engine/Plugins/Experimental/Mover/Source/Mover/Public/MoverDataModelTypes.h
@@ -216,6 +216,8 @@ public:
 	UE_API bool UpdateCurrentMovementBase();
 
 	// Queries
+	bool IsNearlyEqual(const FMoverDefaultSyncState& Other) const;
+
 	UPrimitiveComponent* GetMovementBase() const { return MovementBase.Get(); }
 	FName GetMovementBaseBoneName() const { return MovementBaseBoneName; }
 	FVector GetCapturedMovementBasePos() const { return MovementBasePos; }
diff --git a/Engine/Plugins/Experimental/Mover/Source/Mover/Public/MoverSimulationTypes.h b/Engine/Plugins/Experimental/Mover/Source/Mover/Public/MoverSimulationTypes.h
--- a/Engine/Plugins/Experimental/Mover/Source/Mover/Public/MoverSimulationTypes.h
+++ b/Engine/Plugins/Experimental/Mover/Source/Mover/Public/MoverSimulationTypes.h
@@ -50,6 +50,7 @@ struct FMovementModeTickEndState
 	{
 		RemainingMs = 0.f;
 		NextModeName = NAME_None;
+		bEndedWithNoChanges = false;
 	}
 
 	// Any unused tick time
@@ -58,6 +59,11 @@ struct FMovementModeTickEndState
 
 	UPROPERTY(BlueprintReadWrite, Category=Mover)
 	FName NextModeName = NAME_None;
+
+	// Affirms that no state changes were made during this simulation tick. Can help optimizations if modes set this during sim tick.
+	UPROPERTY(BlueprintReadWrite, Category=Mover)
+	bool bEndedWithNoChanges = false;
+
 };
 
 
@@ -131,8 +137,17 @@ public:
 		MovementMode = NAME_None;
 	}
 
+	bool HasSameContents(const FMoverSyncState& Other) const
+	{
+		return MovementMode == Other.MovementMode &&
+			LayeredMoves.HasSameContents(Other.LayeredMoves) &&
+			MovementModifiers.HasSameContents(Other.MovementModifiers) &&
+			SyncStateCollection.HasSameContents(Other.SyncStateCollection);
+	}
+
 	UScriptStruct* GetStruct() const { return StaticStruct(); }
 
+
 	void NetSerialize(const FNetSerializeParams& P)
 	{
 		P.Ar << MovementMode;
@@ -290,6 +305,12 @@ struct FMoverTickEndData
 		this->AuxState = *AuxState;
 	}
 
+	void InitForNewFrame()
+	{
+		MovementEndState.ResetToDefaults();
+		MoveRecord.Reset();
+	}
+
 	UPROPERTY(BlueprintReadWrite, EditAnywhere, Category=Mover)
 	FMoverSyncState SyncState;
 	UPROPERTY(BlueprintReadWrite, EditAnywhere, Category=Mover)
diff --git a/Engine/Plugins/Experimental/Mover/Source/Mover/Public/MoverTypes.h b/Engine/Plugins/Experimental/Mover/Source/Mover/Public/MoverTypes.h
--- a/Engine/Plugins/Experimental/Mover/Source/Mover/Public/MoverTypes.h
+++ b/Engine/Plugins/Experimental/Mover/Source/Mover/Public/MoverTypes.h
@@ -236,6 +236,9 @@ struct FMoverDataCollection
 	/** Decay input based on DecayAmount for resimulation and forward prediction */
 	UE_API void Decay(float DecayAmount);
 
+	/** Checks only whether there are matching structs inside, but NOT necessarily identical states of each one */
+	UE_API bool HasSameContents(const FMoverDataCollection& Other) const;
+
 	/** Exposes references to GC system */
 	UE_API void AddStructReferencedObjects(FReferenceCollector& Collector) const;
 
