Mover: avoid a state copy operation for simulation ticks that result in only 1 substep and do not have any final instant effects.

diff --git a/Engine/Plugins/Experimental/Mover/Source/Mover/Private/MovementModeStateMachine.cpp b/Engine/Plugins/Experimental/Mover/Source/Mover/Private/MovementModeStateMachine.cpp
--- a/Engine/Plugins/Experimental/Mover/Source/Mover/Private/MovementModeStateMachine.cpp
+++ b/Engine/Plugins/Experimental/Mover/Source/Mover/Private/MovementModeStateMachine.cpp
@@ -155,9 +155,9 @@ void UMovementModeStateMachine::ClearQueuedMode()
 void UMovementModeStateMachine::OnSimulationTick(USceneComponent* UpdatedComponent, UPrimitiveComponent* UpdatedPrimitive, UMoverBlackboard* SimBlackboard, const FMoverTickStartData& StartState, const FMoverTimeStep& TimeStep, FMoverTickEndData& OutputState)
 {
 	//GEngine->AddOnScreenDebugMessage(-1, -0.1f, FColor::White, FString::Printf(TEXT("Mode Tick: %s  Queued: %s"), *CurrentModeName.ToString(), *NextModeName.ToString()));
-
 	FMoverTimeStep SubTimeStep = TimeStep;
 	WorkingSubstepStartData = StartState;
+	bool bIsWorkingStartStateReady = true;	// this flag is used to avoid unneeded data copying after substeps
 
 	UMoverComponent* MoverComp = CastChecked<UMoverComponent>(GetOuter());
 	check(MoverComp->MovementMixer);
@@ -176,6 +176,13 @@ void UMovementModeStateMachine::OnSimulationTick(USceneComponent* UpdatedCompone
 	float TotalUsedMs = 0.0f;
 	while (TotalUsedMs < TimeStep.StepMs)
 	{
+		if (!bIsWorkingStartStateReady)
+		{
+			WorkingSubstepStartData.SyncState = OutputState.SyncState;
+			WorkingSubstepStartData.AuxState = OutputState.AuxState;
+			bIsWorkingStartStateReady = true;
+		}
+
 		WorkingSubstepStartData.SyncState.MovementMode = CurrentModeName;
 
 		FMoverDefaultSyncState* OutputSyncState = &OutputState.SyncState.SyncStateCollection.FindOrAddMutableDataByType<FMoverDefaultSyncState>();
@@ -403,25 +410,33 @@ void UMovementModeStateMachine::OnSimulationTick(USceneComponent* UpdatedCompone
 		TotalUsedMs += SubstepUsedMs;
 		SubTimeStep.StepMs = RemainingMs;
 
-		WorkingSubstepStartData.SyncState = OutputState.SyncState;
-		WorkingSubstepStartData.AuxState  = OutputState.AuxState;
-
+		bIsWorkingStartStateReady = false;
 		++SubStepCount;
 	}
 
-	FApplyMovementEffectParams EffectParams;
-	EffectParams.MoverComp = MoverComp;
-	EffectParams.StartState = &WorkingSubstepStartData;
-	EffectParams.TimeStep = &SubTimeStep;
-	EffectParams.UpdatedComponent = UpdatedComponent;
-	EffectParams.UpdatedPrimitive = UpdatedPrimitive;
-	
-	// Apply any instant effects that were queued up during this tick and didn't get handled in a substep
-	if (ApplyInstantEffects(EffectParams, OutputState.SyncState))
+	if (HasAnyInstantEffectsQueued())
 	{
-		if (CurrentModeName != OutputState.SyncState.MovementMode)
+		if (!bIsWorkingStartStateReady)
 		{
-			SetModeImmediately(OutputState.SyncState.MovementMode);
+			WorkingSubstepStartData.SyncState = OutputState.SyncState;
+			WorkingSubstepStartData.AuxState = OutputState.AuxState;
+			bIsWorkingStartStateReady = true;
+		}
+
+		FApplyMovementEffectParams EffectParams;
+		EffectParams.MoverComp = MoverComp;
+		EffectParams.StartState = &WorkingSubstepStartData;
+		EffectParams.TimeStep = &SubTimeStep;
+		EffectParams.UpdatedComponent = UpdatedComponent;
+		EffectParams.UpdatedPrimitive = UpdatedPrimitive;
+	
+		// Apply any instant effects that were queued up during this tick and didn't get handled in a substep
+		if (ApplyInstantEffects(EffectParams, OutputState.SyncState))
+		{
+			if (CurrentModeName != OutputState.SyncState.MovementMode)
+			{
+				SetModeImmediately(OutputState.SyncState.MovementMode);
+			}
 		}
 	}
 }
@@ -695,6 +710,13 @@ void UMovementModeStateMachine::RollbackModifiers(const FMoverSyncState* Invalid
 	}
 }
 
+bool UMovementModeStateMachine::HasAnyInstantEffectsQueued() const 
+{
+	FRWScopeLock QueueLock(InstantEffectsQueueLock, SLT_ReadOnly);
+
+	return !QueuedInstantEffects.IsEmpty();
+}
+
 bool UMovementModeStateMachine::ApplyInstantEffects(FApplyMovementEffectParams& ApplyEffectParams, FMoverSyncState& OutputState)
 {
 	bool bInstantMovementEffectApplied = false;
diff --git a/Engine/Plugins/Experimental/Mover/Source/Mover/Public/MovementModeStateMachine.h b/Engine/Plugins/Experimental/Mover/Source/Mover/Public/MovementModeStateMachine.h
--- a/Engine/Plugins/Experimental/Mover/Source/Mover/Public/MovementModeStateMachine.h
+++ b/Engine/Plugins/Experimental/Mover/Source/Mover/Public/MovementModeStateMachine.h
@@ -106,6 +106,7 @@ private:
  	UE_API void FlushQueuedModifiersToGroup(FMovementModifierGroup& ModifierGroup);
  	UE_API void FlushModifierCancellationsToGroup(FMovementModifierGroup& ActiveModifierGroup);
  	UE_API void RollbackModifiers(const FMoverSyncState* InvalidSyncState, const FMoverSyncState* SyncState, const FMoverAuxStateContext* InvalidAuxState, const FMoverAuxStateContext* AuxState);
+	UE_API bool HasAnyInstantEffectsQueued() const;
  	UE_API bool ApplyInstantEffects(FApplyMovementEffectParams& ApplyEffectParams, FMoverSyncState& OutputState);
 	UE_API AActor* GetOwnerActor() const;
 };
