Mover: fix bug in standalone mode where the time step was being rounded down to the nearest whole millisecond.

diff --git a/Engine/Plugins/Experimental/Mover/Source/Mover/Private/Backends/MoverStandaloneLiaison.cpp b/Engine/Plugins/Experimental/Mover/Source/Mover/Private/Backends/MoverStandaloneLiaison.cpp
--- a/Engine/Plugins/Experimental/Mover/Source/Mover/Private/Backends/MoverStandaloneLiaison.cpp
+++ b/Engine/Plugins/Experimental/Mover/Source/Mover/Private/Backends/MoverStandaloneLiaison.cpp
@@ -354,7 +354,7 @@ void UMoverStandaloneLiaisonComponent::TickMovementSimulation(float DeltaSeconds
 	FMoverTimeStep TimeStep;
 	TimeStep.ServerFrame = CurrentSimFrame;
 	TimeStep.BaseSimTimeMs = CurrentSimTimeMs;
-	TimeStep.StepMs = FMath::Floor(DeltaSeconds * 1000.f);
+	TimeStep.StepMs = DeltaSeconds * 1000.f;
 
 	WorkingStartData.InputCmd = LastProducedInputCmd;
 
