Mover: fixing angular velocity application for mutli-axis rotation cases. Moving common function into a utility function for all modes to use.

diff --git a/Engine/Plugins/Experimental/ChaosMover/Source/ChaosMover/Private/Character/Modes/ChaosFallingMode.cpp b/Engine/Plugins/Experimental/ChaosMover/Source/ChaosMover/Private/Character/Modes/ChaosFallingMode.cpp
--- a/Engine/Plugins/Experimental/ChaosMover/Source/ChaosMover/Private/Character/Modes/ChaosFallingMode.cpp
+++ b/Engine/Plugins/Experimental/ChaosMover/Source/ChaosMover/Private/Character/Modes/ChaosFallingMode.cpp
@@ -191,11 +191,7 @@ void UChaosFallingMode::SimulationTick_Implementation(const FSimulationTickParam
 
 	const float DeltaSeconds = Params.TimeStep.StepMs * 0.001f;
 
-	FRotator TargetOrient = StartingSyncState->GetOrientation_WorldSpace();
-	if (!UMovementUtils::IsAngularVelocityZero(ProposedMove.AngularVelocity))
-	{
-		TargetOrient += (ProposedMove.AngularVelocity * DeltaSeconds);
-	}
+	const FRotator TargetOrient = UMovementUtils::ApplyAngularVelocity(StartingSyncState->GetOrientation_WorldSpace(), ProposedMove.AngularVelocity, DeltaSeconds);
 
 	// The physics simulation applies Z-only gravity acceleration via physics volumes, so we need to account for it here 
 	FVector TargetVel = ProposedMove.LinearVelocity - DefaultSimInputs->PhysicsObjectGravity * FVector::UpVector * DeltaSeconds;
diff --git a/Engine/Plugins/Experimental/ChaosMover/Source/ChaosMover/Private/Character/Modes/ChaosFlyingMode.cpp b/Engine/Plugins/Experimental/ChaosMover/Source/ChaosMover/Private/Character/Modes/ChaosFlyingMode.cpp
--- a/Engine/Plugins/Experimental/ChaosMover/Source/ChaosMover/Private/Character/Modes/ChaosFlyingMode.cpp
+++ b/Engine/Plugins/Experimental/ChaosMover/Source/ChaosMover/Private/Character/Modes/ChaosFlyingMode.cpp
@@ -120,11 +120,7 @@ void UChaosFlyingMode::SimulationTick_Implementation(const FSimulationTickParams
 
 	const float DeltaSeconds = Params.TimeStep.StepMs * 0.001f;
 
-	FRotator TargetOrient = StartingSyncState->GetOrientation_WorldSpace();
-	if (!UMovementUtils::IsAngularVelocityZero(ProposedMove.AngularVelocity))
-	{
-		TargetOrient += (ProposedMove.AngularVelocity * DeltaSeconds);
-	}
+	const FRotator TargetOrient = UMovementUtils::ApplyAngularVelocity(StartingSyncState->GetOrientation_WorldSpace(), ProposedMove.AngularVelocity, DeltaSeconds);
 
 	// The physics simulation applies Z-only gravity acceleration via physics volumes, so we need to account for it here 
 	FVector TargetVel = ProposedMove.LinearVelocity - DefaultSimInputs->PhysicsObjectGravity * FVector::UpVector * DeltaSeconds;
diff --git a/Engine/Plugins/Experimental/ChaosMover/Source/ChaosMover/Private/Character/Modes/ChaosWalkingMode.cpp b/Engine/Plugins/Experimental/ChaosMover/Source/ChaosMover/Private/Character/Modes/ChaosWalkingMode.cpp
--- a/Engine/Plugins/Experimental/ChaosMover/Source/ChaosMover/Private/Character/Modes/ChaosWalkingMode.cpp
+++ b/Engine/Plugins/Experimental/ChaosMover/Source/ChaosMover/Private/Character/Modes/ChaosWalkingMode.cpp
@@ -269,11 +269,7 @@ void UChaosWalkingMode::SimulationTick_Implementation(const FSimulationTickParam
 
 		// Target orientation
 		// This is always applied regardless of whether the character is supported
-		FRotator TargetOrientation = StartingSyncState->GetOrientation_WorldSpace();
-		if (!UMovementUtils::IsAngularVelocityZero(ProposedMove.AngularVelocity))
-		{
-			TargetOrientation += (ProposedMove.AngularVelocity * DeltaSeconds);
-		}
+		const FRotator TargetOrientation = UMovementUtils::ApplyAngularVelocity(StartingSyncState->GetOrientation_WorldSpace(), ProposedMove.AngularVelocity, DeltaSeconds);
 
 		OutputSyncState.SetTransforms_WorldSpace(TargetPosition, TargetOrientation, TargetVelocity);
 	}
diff --git a/Engine/Plugins/Experimental/Mover/Source/Mover/Private/DefaultMovementSet/Modes/AsyncFallingMode.cpp b/Engine/Plugins/Experimental/Mover/Source/Mover/Private/DefaultMovementSet/Modes/AsyncFallingMode.cpp
--- a/Engine/Plugins/Experimental/Mover/Source/Mover/Private/DefaultMovementSet/Modes/AsyncFallingMode.cpp
+++ b/Engine/Plugins/Experimental/Mover/Source/Mover/Private/DefaultMovementSet/Modes/AsyncFallingMode.cpp
@@ -181,13 +181,7 @@ void UAsyncFallingMode::SimulationTick_Implementation(const FSimulationTickParam
 
 	// Use the orientation intent directly. If no intent is provided, use last frame's orientation. Note that we are assuming rotation changes can't fail. 
 	const FRotator StartingOrient = StartingSyncState->GetOrientation_WorldSpace();
-	FRotator TargetOrient = StartingOrient;
-
-	// Apply orientation changes (if any)
-	if (!UMovementUtils::IsAngularVelocityZero(ProposedMove.AngularVelocity))
-	{
-		TargetOrient += (ProposedMove.AngularVelocity * DeltaSeconds);
-	}
+	const FRotator TargetOrient = UMovementUtils::ApplyAngularVelocity(StartingOrient, ProposedMove.AngularVelocity, DeltaSeconds);
 
 	const FVector StartingFallingVelocity = StartingSyncState->GetVelocity_WorldSpace();
 
diff --git a/Engine/Plugins/Experimental/Mover/Source/Mover/Private/DefaultMovementSet/Modes/AsyncFlyingMode.cpp b/Engine/Plugins/Experimental/Mover/Source/Mover/Private/DefaultMovementSet/Modes/AsyncFlyingMode.cpp
--- a/Engine/Plugins/Experimental/Mover/Source/Mover/Private/DefaultMovementSet/Modes/AsyncFlyingMode.cpp
+++ b/Engine/Plugins/Experimental/Mover/Source/Mover/Private/DefaultMovementSet/Modes/AsyncFlyingMode.cpp
@@ -96,16 +96,9 @@ void UAsyncFlyingMode::SimulationTick_Implementation(const FSimulationTickParams
 
 	// Use the orientation intent directly. If no intent is provided, use last frame's orientation. Note that we are assuming rotation changes can't fail. 
 	const FRotator StartingOrient = StartingSyncState->GetOrientation_WorldSpace();
-	FRotator TargetOrient = StartingOrient;
 
-	bool bIsOrientationChanging = false;
-
-	// Apply orientation changes (if any)
-	if (!UMovementUtils::IsAngularVelocityZero(ProposedMove.AngularVelocity))
-	{
-		TargetOrient += (ProposedMove.AngularVelocity * DeltaSeconds);
-		bIsOrientationChanging = (TargetOrient != StartingOrient);
-	}
+	const FRotator TargetOrient = UMovementUtils::ApplyAngularVelocity(StartingOrient, ProposedMove.AngularVelocity, DeltaSeconds);
+	const bool bIsOrientationChanging = (TargetOrient != StartingOrient);
 	
 	const FVector StartLocation = StartingSyncState->GetLocation_WorldSpace();
 	const FVector TargetLocation = StartLocation + (ProposedMove.LinearVelocity * DeltaSeconds);
diff --git a/Engine/Plugins/Experimental/Mover/Source/Mover/Private/DefaultMovementSet/Modes/AsyncNavWalkingMode.cpp b/Engine/Plugins/Experimental/Mover/Source/Mover/Private/DefaultMovementSet/Modes/AsyncNavWalkingMode.cpp
--- a/Engine/Plugins/Experimental/Mover/Source/Mover/Private/DefaultMovementSet/Modes/AsyncNavWalkingMode.cpp
+++ b/Engine/Plugins/Experimental/Mover/Source/Mover/Private/DefaultMovementSet/Modes/AsyncNavWalkingMode.cpp
@@ -154,14 +154,8 @@ void UAsyncNavWalkingMode::SimulationTick_Implementation(const FSimulationTickPa
 	OutputSyncState.MoveDirectionIntent = (ProposedMove.bHasDirIntent ? ProposedMove.DirectionIntent : FVector::ZeroVector);
 
 	const FRotator StartingOrient = StartingSyncState->GetOrientation_WorldSpace();
-	FRotator TargetOrient = StartingOrient;
-
-	// Apply orientation changes (if any)
-	if (!UMovementUtils::IsAngularVelocityZero(ProposedMove.AngularVelocity))
-	{
-		TargetOrient += (ProposedMove.AngularVelocity * DeltaSeconds);
-	}
-
+	
+	const FRotator TargetOrient = UMovementUtils::ApplyAngularVelocity(StartingOrient, ProposedMove.AngularVelocity, DeltaSeconds);
 
 	const FVector StartingFeetLocation = NavMoverComponent->GetFeetLocationAt(StartingSyncState->GetLocation_WorldSpace());
 	FVector TargetFeetLocation = StartingFeetLocation + (ProposedMove.LinearVelocity * DeltaSeconds);
diff --git a/Engine/Plugins/Experimental/Mover/Source/Mover/Private/DefaultMovementSet/Modes/AsyncWalkingMode.cpp b/Engine/Plugins/Experimental/Mover/Source/Mover/Private/DefaultMovementSet/Modes/AsyncWalkingMode.cpp
--- a/Engine/Plugins/Experimental/Mover/Source/Mover/Private/DefaultMovementSet/Modes/AsyncWalkingMode.cpp
+++ b/Engine/Plugins/Experimental/Mover/Source/Mover/Private/DefaultMovementSet/Modes/AsyncWalkingMode.cpp
@@ -158,16 +158,8 @@ void UAsyncWalkingMode::SimulationTick_Implementation(const FSimulationTickParam
 	OutputSyncState.MoveDirectionIntent = (ProposedMove.bHasDirIntent ? ProposedMove.DirectionIntent : FVector::ZeroVector);
 
 	const FRotator StartingOrient = StartingSyncState->GetOrientation_WorldSpace();
-	FRotator TargetOrient = StartingOrient;
-
-	bool bIsOrientationChanging = false;
-
-	// Apply orientation changes (if any)
-	if (!UMovementUtils::IsAngularVelocityZero(ProposedMove.AngularVelocity))
-	{ 
-		TargetOrient += (ProposedMove.AngularVelocity * DeltaSeconds);
-		bIsOrientationChanging = (TargetOrient != StartingOrient);
-	}
+	const FRotator TargetOrient = UMovementUtils::ApplyAngularVelocity(StartingOrient, ProposedMove.AngularVelocity, DeltaSeconds);
+	const bool bIsOrientationChanging = (TargetOrient != StartingOrient);
 
 	const FQuat StartRotation = StartingOrient.Quaternion();
 
diff --git a/Engine/Plugins/Experimental/Mover/Source/Mover/Private/DefaultMovementSet/Modes/FallingMode.cpp b/Engine/Plugins/Experimental/Mover/Source/Mover/Private/DefaultMovementSet/Modes/FallingMode.cpp
--- a/Engine/Plugins/Experimental/Mover/Source/Mover/Private/DefaultMovementSet/Modes/FallingMode.cpp
+++ b/Engine/Plugins/Experimental/Mover/Source/Mover/Private/DefaultMovementSet/Modes/FallingMode.cpp
@@ -172,13 +172,8 @@ void UFallingMode::SimulationTick_Implementation(const FSimulationTickParams& Pa
 	OutputSyncState.MoveDirectionIntent = (ProposedMove.bHasDirIntent ? ProposedMove.DirectionIntent : FVector::ZeroVector);
 
 	// Use the orientation intent directly. If no intent is provided, use last frame's orientation. Note that we are assuming rotation changes can't fail. 
-	FRotator TargetOrient = StartingSyncState->GetOrientation_WorldSpace();
-
-	// Apply orientation changes (if any)
-	if (!UMovementUtils::IsAngularVelocityZero(ProposedMove.AngularVelocity))
-	{
-		TargetOrient += (ProposedMove.AngularVelocity * DeltaSeconds);
-	}
+	const FRotator StartingOrient = StartingSyncState->GetOrientation_WorldSpace();
+	const FRotator TargetOrient = UMovementUtils::ApplyAngularVelocity(StartingOrient, ProposedMove.AngularVelocity, DeltaSeconds);
 	
 	const FVector StartingFallingVelocity = StartingSyncState->GetVelocity_WorldSpace();
 
diff --git a/Engine/Plugins/Experimental/Mover/Source/Mover/Private/DefaultMovementSet/Modes/FlyingMode.cpp b/Engine/Plugins/Experimental/Mover/Source/Mover/Private/DefaultMovementSet/Modes/FlyingMode.cpp
--- a/Engine/Plugins/Experimental/Mover/Source/Mover/Private/DefaultMovementSet/Modes/FlyingMode.cpp
+++ b/Engine/Plugins/Experimental/Mover/Source/Mover/Private/DefaultMovementSet/Modes/FlyingMode.cpp
@@ -100,16 +100,9 @@ void UFlyingMode::SimulationTick_Implementation(const FSimulationTickParams& Par
 
 	// Use the orientation intent directly. If no intent is provided, use last frame's orientation. Note that we are assuming rotation changes can't fail. 
 	const FRotator StartingOrient = StartingSyncState->GetOrientation_WorldSpace();
-	FRotator TargetOrient = StartingOrient;
 
-	bool bIsOrientationChanging = false;
-
-	// Apply orientation changes (if any)
-	if (!UMovementUtils::IsAngularVelocityZero(ProposedMove.AngularVelocity))
-	{
-		TargetOrient += (ProposedMove.AngularVelocity * DeltaSeconds);
-		bIsOrientationChanging = (TargetOrient != StartingOrient);
-	}
+	const FRotator TargetOrient = UMovementUtils::ApplyAngularVelocity(StartingOrient, ProposedMove.AngularVelocity, DeltaSeconds);
+	const bool bIsOrientationChanging = (TargetOrient != StartingOrient);
 	
 	FVector MoveDelta = ProposedMove.LinearVelocity * DeltaSeconds;
 
diff --git a/Engine/Plugins/Experimental/Mover/Source/Mover/Private/DefaultMovementSet/Modes/NavWalkingMode.cpp b/Engine/Plugins/Experimental/Mover/Source/Mover/Private/DefaultMovementSet/Modes/NavWalkingMode.cpp
--- a/Engine/Plugins/Experimental/Mover/Source/Mover/Private/DefaultMovementSet/Modes/NavWalkingMode.cpp
+++ b/Engine/Plugins/Experimental/Mover/Source/Mover/Private/DefaultMovementSet/Modes/NavWalkingMode.cpp
@@ -150,16 +150,8 @@ void UNavWalkingMode::SimulationTick_Implementation(const FSimulationTickParams&
 	OutputSyncState.MoveDirectionIntent = (ProposedMove.bHasDirIntent ? ProposedMove.DirectionIntent : FVector::ZeroVector);
 
 	const FRotator StartingOrient = StartingSyncState->GetOrientation_WorldSpace();
-	FRotator TargetOrient = StartingOrient;
-
-	bool bIsOrientationChanging = false;
-
-	// Apply orientation changes (if any)
-	if (!UMovementUtils::IsAngularVelocityZero(ProposedMove.AngularVelocity))
-	{ 
-		TargetOrient += (ProposedMove.AngularVelocity * DeltaSeconds);
-		bIsOrientationChanging = true;
-	}
+	const FRotator TargetOrient = UMovementUtils::ApplyAngularVelocity(StartingOrient, ProposedMove.AngularVelocity, DeltaSeconds);
+	const bool bIsOrientationChanging = (TargetOrient != StartingOrient);
 	
 	FQuat TargetOrientQuat = TargetOrient.Quaternion();
 	if (CommonLegacySettings->bShouldRemainVertical)
diff --git a/Engine/Plugins/Experimental/Mover/Source/Mover/Private/DefaultMovementSet/Modes/WalkingMode.cpp b/Engine/Plugins/Experimental/Mover/Source/Mover/Private/DefaultMovementSet/Modes/WalkingMode.cpp
--- a/Engine/Plugins/Experimental/Mover/Source/Mover/Private/DefaultMovementSet/Modes/WalkingMode.cpp
+++ b/Engine/Plugins/Experimental/Mover/Source/Mover/Private/DefaultMovementSet/Modes/WalkingMode.cpp
@@ -150,16 +150,8 @@ void UWalkingMode::SimulationTick_Implementation(const FSimulationTickParams& Pa
 	OutputSyncState.MoveDirectionIntent = (ProposedMove.bHasDirIntent ? ProposedMove.DirectionIntent : FVector::ZeroVector);
 
 	const FRotator StartingOrient = StartingSyncState->GetOrientation_WorldSpace();
-	FRotator TargetOrient = StartingOrient;
-
-	bool bIsOrientationChanging = false;
-
-	// Apply orientation changes (if any)
-	if (!UMovementUtils::IsAngularVelocityZero(ProposedMove.AngularVelocity))
-	{ 
-		TargetOrient += (ProposedMove.AngularVelocity * DeltaSeconds);
-		bIsOrientationChanging = (TargetOrient != StartingOrient);
-	}
+	const FRotator TargetOrient = UMovementUtils::ApplyAngularVelocity(StartingOrient, ProposedMove.AngularVelocity, DeltaSeconds);
+	const bool bIsOrientationChanging = (TargetOrient != StartingOrient);
 
 	FQuat TargetOrientQuat = TargetOrient.Quaternion();
 	if (CommonLegacySettings->bShouldRemainVertical)
diff --git a/Engine/Plugins/Experimental/Mover/Source/Mover/Private/MoveLibrary/MovementUtils.cpp b/Engine/Plugins/Experimental/Mover/Source/Mover/Private/MoveLibrary/MovementUtils.cpp
--- a/Engine/Plugins/Experimental/Mover/Source/Mover/Private/MoveLibrary/MovementUtils.cpp
+++ b/Engine/Plugins/Experimental/Mover/Source/Mover/Private/MoveLibrary/MovementUtils.cpp
@@ -627,6 +627,19 @@ bool UMovementUtils::IsAngularVelocityZero(const FRotator& AngularVelocity)
 }
 
 
+FRotator UMovementUtils::ApplyAngularVelocity(const FRotator& StartingOrient, const FRotator& AngularVelocity, float DeltaSeconds)
+{
+	if (!UMovementUtils::IsAngularVelocityZero(AngularVelocity))
+	{
+		const FQuat ProposedMoveQuat(AngularVelocity * DeltaSeconds);
+		const FQuat TargetOrientQuat = FQuat(StartingOrient) * ProposedMoveQuat;
+		return TargetOrientQuat.Rotator();
+	}
+
+	return StartingOrient;
+}
+
+
 bool UMovementUtils::FindTeleportSpot(const UMoverComponent* MoverComp, FVector& TestLocation, FRotator TestRotation)
 {
 	if (!MoverComp || !MoverComp->GetUpdatedComponent())
diff --git a/Engine/Plugins/Experimental/Mover/Source/Mover/Private/PhysicsMover/Modes/PhysicsDrivenFallingMode.cpp b/Engine/Plugins/Experimental/Mover/Source/Mover/Private/PhysicsMover/Modes/PhysicsDrivenFallingMode.cpp
--- a/Engine/Plugins/Experimental/Mover/Source/Mover/Private/PhysicsMover/Modes/PhysicsDrivenFallingMode.cpp
+++ b/Engine/Plugins/Experimental/Mover/Source/Mover/Private/PhysicsMover/Modes/PhysicsDrivenFallingMode.cpp
@@ -135,12 +135,7 @@ void UPhysicsDrivenFallingMode::SimulationTick_Implementation(const FSimulationT
 	}
 
 	// In air steering
-
-	FRotator TargetOrient = StartingSyncState->GetOrientation_WorldSpace();
-	if (!UMovementUtils::IsAngularVelocityZero(ProposedMove.AngularVelocity)) 
-	{
-		TargetOrient += (ProposedMove.AngularVelocity * DeltaSeconds);
-	}
+	const FRotator TargetOrient = UMovementUtils::ApplyAngularVelocity(StartingSyncState->GetOrientation_WorldSpace(), ProposedMove.AngularVelocity, DeltaSeconds);
 
 	FVector TargetVel = ProposedMove.LinearVelocity;
 	if (const APhysicsVolume* CurPhysVolume = UpdatedComponent->GetPhysicsVolume())
diff --git a/Engine/Plugins/Experimental/Mover/Source/Mover/Private/PhysicsMover/Modes/PhysicsDrivenFlyingMode.cpp b/Engine/Plugins/Experimental/Mover/Source/Mover/Private/PhysicsMover/Modes/PhysicsDrivenFlyingMode.cpp
--- a/Engine/Plugins/Experimental/Mover/Source/Mover/Private/PhysicsMover/Modes/PhysicsDrivenFlyingMode.cpp
+++ b/Engine/Plugins/Experimental/Mover/Source/Mover/Private/PhysicsMover/Modes/PhysicsDrivenFlyingMode.cpp
@@ -66,12 +66,7 @@ void UPhysicsDrivenFlyingMode::SimulationTick_Implementation(const FSimulationTi
 	}
 
 	// In air steering
-
-	FRotator TargetOrient = StartingSyncState->GetOrientation_WorldSpace();
-	if (!UMovementUtils::IsAngularVelocityZero(ProposedMove.AngularVelocity))
-	{
-		TargetOrient += (ProposedMove.AngularVelocity * DeltaSeconds);
-	}
+	const FRotator TargetOrient = UMovementUtils::ApplyAngularVelocity(StartingSyncState->GetOrientation_WorldSpace(), ProposedMove.AngularVelocity, DeltaSeconds);
 
 	FVector TargetVel = ProposedMove.LinearVelocity;
 	if (const APhysicsVolume* CurPhysVolume = UpdatedComponent->GetPhysicsVolume())
diff --git a/Engine/Plugins/Experimental/Mover/Source/Mover/Private/PhysicsMover/Modes/PhysicsDrivenSwimmingMode.cpp b/Engine/Plugins/Experimental/Mover/Source/Mover/Private/PhysicsMover/Modes/PhysicsDrivenSwimmingMode.cpp
--- a/Engine/Plugins/Experimental/Mover/Source/Mover/Private/PhysicsMover/Modes/PhysicsDrivenSwimmingMode.cpp
+++ b/Engine/Plugins/Experimental/Mover/Source/Mover/Private/PhysicsMover/Modes/PhysicsDrivenSwimmingMode.cpp
@@ -153,11 +153,7 @@ void UPhysicsDrivenSwimmingMode::SimulationTick_Implementation(const FSimulation
 		const bool bWalkTrigger = WaterResult.WaterSplineData.ImmersionDepth < CommonLegacySettings->SwimmingStopImmersionDepth;
 		const bool bFallTrigger = FMath::Clamp((WaterResult.WaterSplineData.ImmersionDepth + TargetHeight) / (2 * TargetHeight), -2.f, 2.f) < -1.f;
 	
-		FRotator TargetOrient = StartingSyncState->GetOrientation_WorldSpace();
-		if (!UMovementUtils::IsAngularVelocityZero(ProposedMove.AngularVelocity))
-		{
-			TargetOrient += (ProposedMove.AngularVelocity * DeltaSeconds);
-		}
+		const FRotator TargetOrient = UMovementUtils::ApplyAngularVelocity(StartingSyncState->GetOrientation_WorldSpace(), ProposedMove.AngularVelocity, DeltaSeconds);
 
 		FVector TargetVel = ProposedMove.LinearVelocity;		
 		if (const APhysicsVolume* CurPhysVolume = UpdatedComponent->GetPhysicsVolume())
diff --git a/Engine/Plugins/Experimental/Mover/Source/Mover/Private/PhysicsMover/Modes/PhysicsDrivenWalkingMode.cpp b/Engine/Plugins/Experimental/Mover/Source/Mover/Private/PhysicsMover/Modes/PhysicsDrivenWalkingMode.cpp
--- a/Engine/Plugins/Experimental/Mover/Source/Mover/Private/PhysicsMover/Modes/PhysicsDrivenWalkingMode.cpp
+++ b/Engine/Plugins/Experimental/Mover/Source/Mover/Private/PhysicsMover/Modes/PhysicsDrivenWalkingMode.cpp
@@ -441,11 +441,7 @@ void UPhysicsDrivenWalkingMode::SimulationTick_Implementation(const FSimulationT
 
 		// Target orientation
 		// This is always applied regardless of whether the character is supported
-		FRotator TargetOrientation = StartingSyncState->GetOrientation_WorldSpace();
-		if (!UMovementUtils::IsAngularVelocityZero(ProposedMove.AngularVelocity)) 
-		{
-			TargetOrientation += (ProposedMove.AngularVelocity * DeltaSeconds);
-		}
+		const FRotator TargetOrientation = UMovementUtils::ApplyAngularVelocity(StartingSyncState->GetOrientation_WorldSpace(), ProposedMove.AngularVelocity, DeltaSeconds);
 
 		if (bIsSupported)
 		{
diff --git a/Engine/Plugins/Experimental/Mover/Source/Mover/Public/MoveLibrary/MovementUtils.h b/Engine/Plugins/Experimental/Mover/Source/Mover/Public/MoveLibrary/MovementUtils.h
--- a/Engine/Plugins/Experimental/Mover/Source/Mover/Public/MoveLibrary/MovementUtils.h
+++ b/Engine/Plugins/Experimental/Mover/Source/Mover/Public/MoveLibrary/MovementUtils.h
@@ -291,6 +291,10 @@ public:
 	UFUNCTION(BlueprintCallable, Category = Mover)
 	static UE_API bool IsAngularVelocityZero(const FRotator& AngularVelocity);
 
+	/** Applies a desired angular velocity to a starting rotator over time. Returns the new orientation rotator. */
+	UFUNCTION(BlueprintCallable, Category = Mover)
+	static UE_API FRotator ApplyAngularVelocity(const FRotator& StartingOrient, const FRotator& AngularVelocity, float DeltaSeconds);
+
 	/**
 	 * Try to find an acceptable non-colliding location to place TestActor as close to possible to TestLocation. Expects TestLocation to be a valid location inside the level.
 	 * Returns true if a location without blocking collision is found, in which case TestLocation is overwritten with the new clear location.
