Mover: Updating NavWalkingModes to try updating CachedNavLocation with FindNavFloor if we're sliding along navmeshes and the CachedNavLocation was stale.

#rnx

diff --git a/Engine/Plugins/Experimental/Mover/Source/Mover/Private/DefaultMovementSet/Modes/AsyncNavWalkingMode.cpp b/Engine/Plugins/Experimental/Mover/Source/Mover/Private/DefaultMovementSet/Modes/AsyncNavWalkingMode.cpp
--- a/Engine/Plugins/Experimental/Mover/Source/Mover/Private/DefaultMovementSet/Modes/AsyncNavWalkingMode.cpp
+++ b/Engine/Plugins/Experimental/Mover/Source/Mover/Private/DefaultMovementSet/Modes/AsyncNavWalkingMode.cpp
@@ -239,11 +239,20 @@ void UAsyncNavWalkingMode::SimulationTick_Implementation(const FSimulationTickPa
 		{
 			if (bSlideAlongNavMeshEdge && CachedNavLocation.HasNodeRef())
 			{
-				bHasNavigationData = NavDataInterface->FindMoveAlongSurface(CachedNavLocation, TargetFeetLocation, OUT DestNavLocation);
+				bool bHasValidCachedNavLocation = NavDataInterface->IsNodeRefValid(CachedNavLocation.NodeRef);
+				if (!bHasValidCachedNavLocation)
+				{
+					bHasValidCachedNavLocation = FindNavFloor(TargetFeetLocation, OUT CachedNavLocation, NavDataInterface.Get());
+				}
 
-				if (bHasNavigationData)
+				if (bHasValidCachedNavLocation)
 				{
-					TargetFeetLocation = UMovementUtils::ProjectToGravityFloor(DestNavLocation.Location, UpDirection) + UMovementUtils::GetGravityVerticalComponent(TargetFeetLocation, UpDirection);
+					bHasNavigationData = NavDataInterface->FindMoveAlongSurface(CachedNavLocation, TargetFeetLocation, OUT DestNavLocation);
+
+					if (bHasNavigationData)
+					{
+						TargetFeetLocation = UMovementUtils::ProjectToGravityFloor(DestNavLocation.Location, UpDirection) + UMovementUtils::GetGravityVerticalComponent(TargetFeetLocation, UpDirection);
+					}
 				}
 			}
 			else
diff --git a/Engine/Plugins/Experimental/Mover/Source/Mover/Private/DefaultMovementSet/Modes/NavWalkingMode.cpp b/Engine/Plugins/Experimental/Mover/Source/Mover/Private/DefaultMovementSet/Modes/NavWalkingMode.cpp
--- a/Engine/Plugins/Experimental/Mover/Source/Mover/Private/DefaultMovementSet/Modes/NavWalkingMode.cpp
+++ b/Engine/Plugins/Experimental/Mover/Source/Mover/Private/DefaultMovementSet/Modes/NavWalkingMode.cpp
@@ -228,11 +228,20 @@ void UNavWalkingMode::SimulationTick_Implementation(const FSimulationTickParams&
 		{
 			if (bSlideAlongNavMeshEdge && CachedNavLocation.HasNodeRef())
 			{
-				bHasNavigationData = NavDataInterface->FindMoveAlongSurface(CachedNavLocation, AdjustedDest, OUT DestNavLocation);
+				bool bHasValidCachedNavLocation = NavDataInterface->IsNodeRefValid(CachedNavLocation.NodeRef);
+				if (!bHasValidCachedNavLocation)
+				{
+					bHasValidCachedNavLocation = FindNavFloor(AdjustedDest, OUT CachedNavLocation, NavDataInterface.Get());
+				}
 
-				if (bHasNavigationData)
+				if (bHasValidCachedNavLocation)
 				{
-					AdjustedDest = UMovementUtils::ProjectToGravityFloor(DestNavLocation.Location, UpDirection) + UMovementUtils::GetGravityVerticalComponent(AdjustedDest, UpDirection);
+					bHasNavigationData = NavDataInterface->FindMoveAlongSurface(CachedNavLocation, AdjustedDest, OUT DestNavLocation);
+
+					if (bHasNavigationData)
+					{
+						AdjustedDest = UMovementUtils::ProjectToGravityFloor(DestNavLocation.Location, UpDirection) + UMovementUtils::GetGravityVerticalComponent(AdjustedDest, UpDirection);
+					}
 				}
 			}
 			else
