Mover: fix BP editor crash during validation if there's a shared settings object set to None

diff --git a/Engine/Plugins/Experimental/Mover/Source/Mover/Private/MoverComponent.cpp b/Engine/Plugins/Experimental/Mover/Source/Mover/Private/MoverComponent.cpp
--- a/Engine/Plugins/Experimental/Mover/Source/Mover/Private/MoverComponent.cpp
+++ b/Engine/Plugins/Experimental/Mover/Source/Mover/Private/MoverComponent.cpp
@@ -2143,7 +2143,7 @@ IMovementSettingsInterface* UMoverComponent::FindSharedSettings_Mutable(const UC
 
 	for (const TObjectPtr<UObject>& SettingsObj : SharedSettings)
 	{
-		if (SettingsObj->IsA(ByType))
+		if (SettingsObj && SettingsObj->IsA(ByType))
 		{
 			return Cast<IMovementSettingsInterface>(SettingsObj);
 		}
