Chaos Mover - Fix output time and events

#rnx

diff --git a/Engine/Plugins/Experimental/ChaosMover/Source/ChaosMover/Private/Backends/ChaosMoverAsyncCallback.cpp b/Engine/Plugins/Experimental/ChaosMover/Source/ChaosMover/Private/Backends/ChaosMoverAsyncCallback.cpp
--- a/Engine/Plugins/Experimental/ChaosMover/Source/ChaosMover/Private/Backends/ChaosMoverAsyncCallback.cpp
+++ b/Engine/Plugins/Experimental/ChaosMover/Source/ChaosMover/Private/Backends/ChaosMoverAsyncCallback.cpp
@@ -46,7 +46,6 @@ namespace UE::ChaosMover
 		if (AsyncInput->PhysicsSolver)
 		{
 			TimeStep.ServerFrame = AsyncInput->PhysicsSolver->GetCurrentFrame() + AsyncInput->NetworkPhysicsTickOffset;
-			TimeStep.BaseSimTimeMs += AsyncInput->NetworkPhysicsTickOffset * TimeStep.StepMs;
 			TimeStep.bIsResimulating = AsyncInput->PhysicsSolver->GetEvolution()->IsResimming();
 		}
 
@@ -82,7 +81,6 @@ namespace UE::ChaosMover
 		if (AsyncInput->PhysicsSolver)
 		{
 			TimeStep.ServerFrame = AsyncInput->PhysicsSolver->GetCurrentFrame() + AsyncInput->NetworkPhysicsTickOffset;
-			TimeStep.BaseSimTimeMs += AsyncInput->NetworkPhysicsTickOffset * TimeStep.StepMs;
 			TimeStep.bIsResimulating = AsyncInput->PhysicsSolver->GetEvolution()->IsResimming();
 		}
 
diff --git a/Engine/Plugins/Experimental/ChaosMover/Source/ChaosMover/Private/Backends/ChaosMoverBackend.cpp b/Engine/Plugins/Experimental/ChaosMover/Source/ChaosMover/Private/Backends/ChaosMoverBackend.cpp
--- a/Engine/Plugins/Experimental/ChaosMover/Source/ChaosMover/Private/Backends/ChaosMoverBackend.cpp
+++ b/Engine/Plugins/Experimental/ChaosMover/Source/ChaosMover/Private/Backends/ChaosMoverBackend.cpp
@@ -455,7 +455,7 @@ void UChaosMoverBackendComponent::FinalizeFrame(float ResultsTimeInMs)
 
 	FMoverTimeStep TimeStep;
 	UE::ChaosMover::FSimulationOutputData InterpolatedOutput;
-	SimOutputRecord.GetInterpolated(ResultsTimeInMs, TimeStep, InterpolatedOutput);
+	SimOutputRecord.CreateInterpolatedResult(ResultsTimeInMs, TimeStep, InterpolatedOutput);
 
 	// Physics interactions in the last frame may have caused a change in position or velocity that's different from what a simple lerp would predict,
 	// so stomp the lerped sync state's transform data with that of the actual particle after the last sim frame
@@ -466,6 +466,8 @@ void UChaosMoverBackendComponent::FinalizeFrame(float ResultsTimeInMs)
 
 		// Make sure the move direction intent is in base space (the base quat is identity if there's no base, effectively making this a no-op)
 		TransformSyncState.MoveDirectionIntent = TransformSyncState.GetCapturedMovementBaseQuat().UnrotateVector(TransformSyncState.MoveDirectionIntent);
+	
+		check(Chaos::FVec3::IsNearlyEqual(Chaos::FVec3(TransformSyncState.GetLocation_WorldSpace()), Particle->GetX(), UE_KINDA_SMALL_NUMBER));
 	}
 
 	MoverComp.SetSimulationOutput(TimeStep, InterpolatedOutput);
diff --git a/Engine/Plugins/Experimental/ChaosMover/Source/ChaosMover/Private/Backends/ChaosMoverSubsystem.cpp b/Engine/Plugins/Experimental/ChaosMover/Source/ChaosMover/Private/Backends/ChaosMoverSubsystem.cpp
--- a/Engine/Plugins/Experimental/ChaosMover/Source/ChaosMover/Private/Backends/ChaosMoverSubsystem.cpp
+++ b/Engine/Plugins/Experimental/ChaosMover/Source/ChaosMover/Private/Backends/ChaosMoverSubsystem.cpp
@@ -115,6 +115,9 @@ void UChaosMoverSubsystem::InjectInputs_External(int32 PhysicsStep, int32 NumSte
 			{
 				AsyncInput->PhysicsSolver = Solver;
 
+				TimeStep.BaseSimTimeMs = Solver->GetPhysicsResultsTime_External() * 1000.0f;
+				TimeStep.ServerFrame = Solver->GetCurrentFrame() + NetPhysicsTickOffset;
+
 				if (Solver->IsUsingAsyncResults())
 				{
 					TimeStep.StepMs = Solver->GetAsyncDeltaTime() * 1000.0f;
@@ -123,14 +126,11 @@ void UChaosMoverSubsystem::InjectInputs_External(int32 PhysicsStep, int32 NumSte
 				{
 					TimeStep.StepMs = FMath::Clamp(World->GetDeltaSeconds(), Solver->GetMinDeltaTime_External(), Solver->GetMaxDeltaTime_External()) * 1000.0f;
 				}
-
-				TimeStep.ServerFrame = Solver->GetCurrentFrame() + NetPhysicsTickOffset;
-				TimeStep.BaseSimTimeMs = Solver->GetPhysicsResultsTime_External() * 1000.0f + NetPhysicsTickOffset * TimeStep.StepMs;
 			}
 		}
 	}
 
-	auto LambdaParallelUpdate = [NetPhysicsTickOffset, PhysicsStep, NumSteps, AsyncInput, &TimeStep, this](int32 Idx) {
+	auto LambdaParallelUpdate = [PhysicsStep, NumSteps, AsyncInput, &TimeStep, this](int32 Idx) {
 		Backends[Idx]->ProduceInputData(PhysicsStep, NumSteps, TimeStep, AsyncInput->InputData[Idx]);
 	};
 
@@ -148,8 +148,6 @@ void UChaosMoverSubsystem::OnPostPhysicsTick(FChaosScene* Scene)
 	const bool ForceSingleThread = UE::ChaosMover::CVars::bForceSingleThreadedGT;
 
 	// Pop each async output
-	float LastDtInMs = 0.0f;
-
 	while (Chaos::TSimCallbackOutputHandle<UE::ChaosMover::FAsyncCallbackOutput> AsyncOutput = AsyncCallback->PopFutureOutputData_External())
 	{
 		auto LambdaParallelUpdate = [&AsyncOutput](int32 Idx) {
@@ -159,11 +157,6 @@ void UChaosMoverSubsystem::OnPostPhysicsTick(FChaosScene* Scene)
 			}
 		};
 
-		if (AsyncOutput->TimeStep.Num() > 0)
-		{
-			LastDtInMs = AsyncOutput->TimeStep[0].StepMs;
-		}
-
 		Chaos::PhysicsParallelFor(AsyncOutput->Backends.Num(), LambdaParallelUpdate, ForceSingleThread);
 	}
 
@@ -172,13 +165,12 @@ void UChaosMoverSubsystem::OnPostPhysicsTick(FChaosScene* Scene)
 	{
 		if (Chaos::FPhysicsSolver* Solver = Scene->GetSolver())
 		{
-			// Want the base time at the start of the frame
-			ResultsTimeInMs = Solver->GetPhysicsResultsTime_External() * 1000.0f - LastDtInMs;
+			ResultsTimeInMs = Solver->GetPhysicsResultsTime_External() * 1000.0f;
 		}
 	}
 
 	// Call finalize frame
-	auto LambdaParallelUpdate = [this, &ResultsTimeInMs](int32 Idx) {
+	auto LambdaParallelUpdate = [this, ResultsTimeInMs](int32 Idx) {
 		if (UChaosMoverBackendComponent* Backend = Backends[Idx].Get())
 		{
 			Backend->FinalizeFrame(ResultsTimeInMs);
diff --git a/Engine/Plugins/Experimental/Mover/Source/Mover/Private/MoverSimulationTypes.cpp b/Engine/Plugins/Experimental/Mover/Source/Mover/Private/MoverSimulationTypes.cpp
--- a/Engine/Plugins/Experimental/Mover/Source/Mover/Private/MoverSimulationTypes.cpp
+++ b/Engine/Plugins/Experimental/Mover/Source/Mover/Private/MoverSimulationTypes.cpp
@@ -23,27 +23,6 @@ namespace UE::Mover
 		SyncState.Interpolate(&From.SyncState, &To.SyncState, Alpha);
 		LastUsedInputCmd.Interpolate(&From.LastUsedInputCmd, &To.LastUsedInputCmd, Alpha);
 		AdditionalOutputData.Interpolate(From.AdditionalOutputData, To.AdditionalOutputData, Alpha);
-
-		for (const TSharedPtr<FMoverSimulationEventData>& EventData : From.Events)
-		{
-			if (const FMoverSimulationEventData* DataPtr = EventData.Get())
-			{
-				if (DataPtr->EventTimeMs <= SimTimeMs)
-				{
-					Events.Add(EventData);
-				}
-			}
-		}
-		for (const TSharedPtr<FMoverSimulationEventData>& EventData : To.Events)
-		{
-			if (const FMoverSimulationEventData* DataPtr = EventData.Get())
-			{
-				if (DataPtr->EventTimeMs <= SimTimeMs)
-				{
-					Events.Add(EventData);
-				}
-			}
-		}
 	}
 
 	void FSimulationOutputRecord::FData::Reset()
@@ -56,6 +35,30 @@ namespace UE::Mover
 	{
 		CurrentIndex = (CurrentIndex + 1) % 2;
 		Data[CurrentIndex] = { InTimeStep, InData };
+
+		// When we're resimulating we can be adding data that is at an earlier time to
+		// the currently stored data, so fix up if necessary
+		if (InTimeStep.bIsResimulating)
+		{
+			// Remove previous result if it is ahead in time of this latest result
+			const uint8 PrevIndex = (CurrentIndex + 1) % 2;
+			if (Data[PrevIndex].TimeStep.BaseSimTimeMs > InTimeStep.BaseSimTimeMs)
+			{
+				Data[PrevIndex] = Data[CurrentIndex];
+			}
+
+			// If resimulating find any events stored after the time step and remove
+			Events.RemoveAllSwap([&InTimeStep](const TSharedPtr<FMoverSimulationEventData> Event) {
+				return !Event.IsValid() || (Event->EventTimeMs >= InTimeStep.BaseSimTimeMs);
+				});
+		}
+
+		// Transfer events to local Events array
+		for (TSharedPtr<FMoverSimulationEventData>& Event : Data[CurrentIndex].SimOutputData.Events)
+		{
+			Events.Add(Event);
+		}
+		Data[CurrentIndex].SimOutputData.Events.Empty();
 	}
 
 	const FSimulationOutputData& FSimulationOutputRecord::GetLatest() const
@@ -63,7 +66,7 @@ namespace UE::Mover
 		return Data[CurrentIndex].SimOutputData;
 	}
 
-	void FSimulationOutputRecord::GetInterpolated(float AtBaseTimeMs, FMoverTimeStep& OutTimeStep, FSimulationOutputData& OutData) const
+	void FSimulationOutputRecord::CreateInterpolatedResult(float AtBaseTimeMs, FMoverTimeStep& OutTimeStep, FSimulationOutputData& OutData)
 	{
 		const uint8 PrevIndex = (CurrentIndex + 1) % 2;
 		const float PrevTimeMs = Data[PrevIndex].TimeStep.BaseSimTimeMs;
@@ -87,6 +90,18 @@ namespace UE::Mover
 		}
 
 		OutTimeStep.BaseSimTimeMs = AtBaseTimeMs;
+
+		for (TSharedPtr<FMoverSimulationEventData>& Event : Events)
+		{
+			if (Event->EventTimeMs <= AtBaseTimeMs)
+			{
+				OutData.Events.Add(MoveTemp(Event));
+			}
+		}
+
+		Events.RemoveAllSwap([](const TSharedPtr<FMoverSimulationEventData> Event) {
+			return !Event.IsValid();
+			});
 	}
 
 	void FSimulationOutputRecord::Clear()
@@ -94,6 +109,7 @@ namespace UE::Mover
 		CurrentIndex = 1;
 		Data[0].Reset();
 		Data[1].Reset();
+		Events.Empty();
 	}
 
 } // namespace UE::Mover
\ No newline at end of file
diff --git a/Engine/Plugins/Experimental/Mover/Source/Mover/Public/MoverSimulationTypes.h b/Engine/Plugins/Experimental/Mover/Source/Mover/Public/MoverSimulationTypes.h
--- a/Engine/Plugins/Experimental/Mover/Source/Mover/Public/MoverSimulationTypes.h
+++ b/Engine/Plugins/Experimental/Mover/Source/Mover/Public/MoverSimulationTypes.h
@@ -488,12 +488,17 @@ namespace UE::Mover
 		};
 
 		MOVER_API void Add(const FMoverTimeStep& InTimeStep, const FSimulationOutputData& InData);
+
 		MOVER_API const FSimulationOutputData& GetLatest() const;
-		MOVER_API void GetInterpolated(float AtBaseTimeMs, FMoverTimeStep& OutTimeStep, FSimulationOutputData& OutData) const;
+
+		/** This will create an interpolated output and extract events from the stored data with time stamps up until the input time */
+		MOVER_API void CreateInterpolatedResult(float AtBaseTimeMs, FMoverTimeStep& OutTimeStep, FSimulationOutputData& OutData);
+
 		MOVER_API void Clear();
 
 	private:
 		FData Data[2];
+		TArray<TSharedPtr<FMoverSimulationEventData>> Events;
 		uint8 CurrentIndex = 1;
 	};
 
