[Mover] [Bugfix] Pawn slides when crouching on slopes

diff --git a/Engine/Plugins/Experimental/Mover/Source/Mover/Private/Backends/MoverNetworkPhysicsLiaison.cpp b/Engine/Plugins/Experimental/Mover/Source/Mover/Private/Backends/MoverNetworkPhysicsLiaison.cpp
--- a/Engine/Plugins/Experimental/Mover/Source/Mover/Private/Backends/MoverNetworkPhysicsLiaison.cpp
+++ b/Engine/Plugins/Experimental/Mover/Source/Mover/Private/Backends/MoverNetworkPhysicsLiaison.cpp
@@ -402,7 +402,17 @@ void UMoverNetworkPhysicsLiaisonComponent::PerformPreSimulate_Internal(const FPh
 			
 		CharacterParticle->SetW(AngularVelocity);
 	}		
-	
+
+	// Update the TargetHeight Constraint Settings
+	if (const UBaseMovementMode* CurrentMode = MoverComp.FindMovementModeByName(Output.SyncState.MovementMode))
+	{
+		if (const IPhysicsCharacterMovementModeInterface* PhysicsMode = CastChecked<IPhysicsCharacterMovementModeInterface>(CurrentMode))
+		{
+			Chaos::FCharacterGroundConstraintSettings& ConstraintSettings = ConstraintHandle->GetSettings_Mutable();
+			ConstraintSettings.TargetHeight = PhysicsMode->GetTargetHeight();
+		}
+	}
+
 	// Update the constraint data based on the floor result
 	if (Output.FloorResult.bBlockingHit)
 	{
