Attempt to stop NPC sinking through floor:
Use navmesh normal instead of physics hit normal
Always taking highest hit location in z when searching for ground location

diff --git a/Engine/Plugins/Experimental/Mover/Source/Mover/Private/DefaultMovementSet/Modes/AsyncNavWalkingMode.cpp b/Engine/Plugins/Experimental/Mover/Source/Mover/Private/DefaultMovementSet/Modes/AsyncNavWalkingMode.cpp
--- a/Engine/Plugins/Experimental/Mover/Source/Mover/Private/DefaultMovementSet/Modes/AsyncNavWalkingMode.cpp
+++ b/Engine/Plugins/Experimental/Mover/Source/Mover/Private/DefaultMovementSet/Modes/AsyncNavWalkingMode.cpp
@@ -16,9 +16,25 @@
 #include "MoveLibrary/MovementUtils.h"
 #include "DrawDebugHelpers.h"
 #include "NavigationSystem.h"
+#include "MoveLibrary/NavMovementUtils.h"
+#include "NavMesh/RecastNavMesh.h"
+#include "VisualLogger/VisualLogger.h"
 
 #include UE_INLINE_GENERATED_CPP_BY_NAME(AsyncNavWalkingMode)
 
+#if ENABLE_VISUAL_LOG
+static const FName AsyncWalkModeLogCategory = FName("AsyncNavWalkingMode");
+#endif
+
+namespace AsyncWalkingModeCVars
+{
+	float OverrideRaycastInterval = -1.0f;
+	FAutoConsoleVariableRef CVarOverrideNavProjectionInterval(TEXT("Mover.AsyncNav.OverrideRaycastInterval"), OverrideRaycastInterval, TEXT(""));
+
+	bool bUseNavMeshNormal = false;
+	FAutoConsoleVariableRef CVarUseNavMeshNormal(TEXT("Mover.AsyncNav.UseNavMeshNormal"), bUseNavMeshNormal, TEXT(""));
+}
+
 static const FName MoveWithoutNavMeshSubstepName = "MoveWithoutNavMesh";
 
 UAsyncNavWalkingMode::UAsyncNavWalkingMode()
@@ -433,23 +449,6 @@ void UAsyncNavWalkingMode::FindBestNavMeshLocation(const FVector& TraceStart, co
 	TArray<FHitResult> MultiTraceHits;
 	GetWorld()->LineTraceMultiByChannel(MultiTraceHits, TraceStart, TraceEnd, ECC_WorldStatic, Params, ResponseParams);
 
-	struct FCompareFHitResultNavMeshTrace
-	{
-		explicit FCompareFHitResultNavMeshTrace(const FVector& inSourceLocation) : SourceLocation(inSourceLocation)
-		{
-		}
-
-		FORCEINLINE bool operator()(const FHitResult& A, const FHitResult& B) const
-		{
-			const float ADistSqr = (SourceLocation - A.ImpactPoint).SizeSquared();
-			const float BDistSqr = (SourceLocation - B.ImpactPoint).SizeSquared();
-
-			return (ADistSqr < BDistSqr);
-		}
-
-		const FVector& SourceLocation;
-	};
-
 	struct FRemoveNotBlockingResponseNavMeshTrace
 	{
 		FRemoveNotBlockingResponseNavMeshTrace(bool bInCheckOnlyWorldStatic) : bCheckOnlyWorldStatic(bInCheckOnlyWorldStatic) {}
@@ -478,13 +477,36 @@ void UAsyncNavWalkingMode::FindBestNavMeshLocation(const FVector& TraceStart, co
 	MultiTraceHits.RemoveAllSwap(FRemoveNotBlockingResponseNavMeshTrace(!bProjectNavMeshOnBothWorldChannels), EAllowShrinking::No);
 	if (MultiTraceHits.Num() > 0)
 	{
-		// Sort the hits by the closest to our origin.
-		MultiTraceHits.Sort(FCompareFHitResultNavMeshTrace(TargetNavLocation));
+		FVector UpDirection = GetMoverComponent()->GetUpDirection();
+		MultiTraceHits.Sort([UpDirection](const FHitResult& A, const FHitResult& B)
+		{
+			return A.ImpactPoint.Dot(UpDirection) > B.ImpactPoint.Dot(UpDirection);
+		});
 
 		// Cache the closest hit and treat it as a blocking hit (we used an overlap to get all the world static hits so we could sort them ourselves)
 		OutHitResult = MultiTraceHits[0];
 		OutHitResult.bBlockingHit = true;
 	}
+
+#if ENABLE_VISUAL_LOG
+	const float DebugNormalLength = 50.0f;
+	UE_VLOG_SEGMENT(GetMoverComponent(), AsyncWalkModeLogCategory, Display, TraceStart, TraceEnd, FColor::Green, TEXT(""));
+	for (int i = 0; i < MultiTraceHits.Num(); i++)
+	{
+		UE_VLOG_SPHERE(GetMoverComponent(), "AsyncNavWalkingMode", Display, MultiTraceHits[i].ImpactPoint, i == 0 ? 5.0f : 2.5f, i == 0 ? FColor::Red : FColor::Yellow, TEXT("%i"), i);
+		UE_VLOG_ARROW(GetMoverComponent(), "AsyncNavWalkingMode", Display, MultiTraceHits[i].ImpactPoint, (MultiTraceHits[i].ImpactPoint + DebugNormalLength * MultiTraceHits[i].ImpactNormal), i == 0 ? FColor::Red : FColor::Yellow, TEXT(""));
+	}
+#endif // ENABLE_VISUAL_LOG
+
+	if (AsyncWalkingModeCVars::bUseNavMeshNormal && CachedNavLocation.HasNodeRef() && MultiTraceHits.Num() > 0)
+	{
+		FVector NavMeshNormal;
+		if (NavMovementUtils::CalculateNavMeshNormal(CachedNavLocation, NavMeshNormal, NavDataInterface.Get(), GetMoverComponent()))
+		{
+			OutHitResult.ImpactNormal = NavMeshNormal;
+			UE_VLOG_ARROW(GetMoverComponent(), "AsyncNavWalkingMode", Display, OutHitResult.ImpactPoint, (OutHitResult.ImpactPoint + DebugNormalLength * OutHitResult.ImpactNormal), FColor::Magenta, TEXT("NavMeshNormal"));
+		}
+	}
 }
 
 FVector UAsyncNavWalkingMode::ProjectLocationFromNavMesh(float DeltaSeconds, const FVector& CurrentFeetLocation, const FVector& TargetNavLocation, float UpOffset, float DownOffset)
@@ -514,6 +536,7 @@ FVector UAsyncNavWalkingMode::ProjectLocationFromNavMesh(float DeltaSeconds, con
 		CachedProjectedNavMeshHitResult.TraceEnd == TraceEnd);
 
 	// Check periodically or if we have no information about our last floor result
+	UE_VLOG_SPHERE(GetMoverComponent(), AsyncWalkModeLogCategory, Display, TargetNavLocation, 5.0f, FColor::Blue, TEXT("TargetNavLocation"));
 	NavMeshProjectionTimer -= DeltaSeconds;
 	if (NavMeshProjectionTimer <= 0.0f || !bHasValidFloorResult)
 	{
@@ -553,12 +576,13 @@ FVector UAsyncNavWalkingMode::ProjectLocationFromNavMesh(float DeltaSeconds, con
 		// Wrap around to maintain same relative offset to tick time changes.
 		// Prevents large framerate spikes from aligning multiple characters to the same frame (if they start staggered, they will now remain staggered).
 		float ModTime = 0.f;
-		if (NavMeshProjectionInterval > UE_SMALL_NUMBER)
+		float Interval = AsyncWalkingModeCVars::OverrideRaycastInterval >= 0.0f ? AsyncWalkingModeCVars::OverrideRaycastInterval : NavMeshProjectionInterval;
+		if (Interval > UE_SMALL_NUMBER)
 		{
-			ModTime = FMath::Fmod(-NavMeshProjectionTimer, NavMeshProjectionInterval);
+			ModTime = FMath::Fmod(-NavMeshProjectionTimer, Interval);
 		}
 
-		NavMeshProjectionTimer = NavMeshProjectionInterval - ModTime;
+		NavMeshProjectionTimer = Interval - ModTime;
 	}
 
 	// Project to last plane we found.
@@ -572,6 +596,7 @@ FVector UAsyncNavWalkingMode::ProjectLocationFromNavMesh(float DeltaSeconds, con
 		else
 		{
 			const FVector ProjectedPoint = FMath::LinePlaneIntersection(TraceStart, TraceEnd, CachedProjectedNavMeshHitResult.ImpactPoint, CachedProjectedNavMeshHitResult.Normal);
+			UE_VLOG_SPHERE(GetMoverComponent(), "AsyncNavWalkingMode", Display, ProjectedPoint, 2.5f, FColor::Orange, TEXT("ProjectedPoint"));
 			FVector::FReal ProjectedVertical = ProjectedPoint.Dot(UpDirection);
 
 			// Limit to not be too far above or below NavMesh location
@@ -638,3 +663,4 @@ void UAsyncNavWalkingMode::CaptureOutputState(const FMoverDefaultSyncState& Star
 	TickEndData.MovementEndState.bEndedWithNoChanges = OutputSyncState.IsNearlyEqual(StartSyncState);
 }
 
+
diff --git a/Engine/Plugins/Experimental/Mover/Source/Mover/Private/MoveLibrary/NavMovementUtils.cpp b/Engine/Plugins/Experimental/Mover/Source/Mover/Private/MoveLibrary/NavMovementUtils.cpp
new file mode 100644
--- /dev/null
+++ b/Engine/Plugins/Experimental/Mover/Source/Mover/Private/MoveLibrary/NavMovementUtils.cpp
@@ -0,0 +1,49 @@
+// Copyright Epic Games, Inc. All Rights Reserved.
+
+#include "NavMovementUtils.h"
+#include "CoreTypes.h"
+#include "NavMesh/RecastNavMesh.h"
+#include "VisualLogger/VisualLogger.h"
+
+bool NavMovementUtils::CalculateNavMeshNormal(const FNavLocation& Location, FVector& OutNormal, const INavigationDataInterface* NavData,
+                                              UObject* LogOwner)
+{
+	const ARecastNavMesh* RecastNavMesh = Cast<ARecastNavMesh>(NavData);
+	if (RecastNavMesh == nullptr)
+	{
+		return false;
+	}
+
+	if (Location.HasNodeRef() == false)
+	{
+		return false;
+	}
+
+	TArray<FVector> OutVerts;
+	if (RecastNavMesh->GetPolyVerts(Location.NodeRef, OutVerts) == false)
+	{
+		return false;
+	}
+
+	check(OutVerts.Num() >= 3);
+	OutNormal = FVector::ZeroVector;
+	float DebugThickness = 2.0f;
+
+	// Calculate normal from navmesh verts (arbitrary number)
+	for (int i = 0; i < OutVerts.Num() - 1; ++i)
+	{
+		OutNormal += OutVerts[i+1].Cross(OutVerts[i]);
+		UE_VLOG_SEGMENT_THICK(LogOwner, "AsyncNavWalkingMode", Display, OutVerts[i], OutVerts[i+1], FColor::Magenta, DebugThickness, TEXT(""));
+	}
+	OutNormal += OutVerts[0].Cross(OutVerts.Last());
+	UE_VLOG_SEGMENT_THICK(LogOwner, "AsyncNavWalkingMode", Display, OutVerts.Last(), OutVerts[0], FColor::Magenta, DebugThickness, TEXT(""));
+
+	// Normalize
+	float SizeSquared = OutNormal.SizeSquared();
+	if (FMath::IsNearlyZero(SizeSquared))
+	{
+		return false;
+	}
+	OutNormal /= FMath::Sqrt(SizeSquared);
+	return true;
+}
diff --git a/Engine/Plugins/Experimental/Mover/Source/Mover/Private/MoveLibrary/NavMovementUtils.h b/Engine/Plugins/Experimental/Mover/Source/Mover/Private/MoveLibrary/NavMovementUtils.h
new file mode 100644
--- /dev/null
+++ b/Engine/Plugins/Experimental/Mover/Source/Mover/Private/MoveLibrary/NavMovementUtils.h
@@ -0,0 +1,16 @@
+// Copyright Epic Games, Inc. All Rights Reserved.
+
+#pragma once
+
+#include "Math/Vector.h"
+
+struct FNavLocation;
+class INavigationDataInterface;
+class UObject;
+
+class NavMovementUtils
+{
+	public:
+	
+	static bool CalculateNavMeshNormal(const FNavLocation& Location, FVector& OutNormal, const INavigationDataInterface* NavData, UObject* LogOwner = nullptr);
+};
\ No newline at end of file
