Mover: fixing issues with trying to maintain the small distance off of floors in falling/flying modes, where we could a character could get pulled closer to the floor. Now in these modes, we only allow pushing away from being too close to a floor.

diff --git a/Engine/Plugins/Experimental/Mover/Source/Mover/Private/DefaultMovementSet/Modes/AsyncFallingMode.cpp b/Engine/Plugins/Experimental/Mover/Source/Mover/Private/DefaultMovementSet/Modes/AsyncFallingMode.cpp
--- a/Engine/Plugins/Experimental/Mover/Source/Mover/Private/DefaultMovementSet/Modes/AsyncFallingMode.cpp
+++ b/Engine/Plugins/Experimental/Mover/Source/Mover/Private/DefaultMovementSet/Modes/AsyncFallingMode.cpp
@@ -226,7 +226,7 @@ void UAsyncFallingMode::SimulationTick_Implementation(const FSimulationTickParam
 				SweepHit, CommonLegacySettings->FloorSweepDistance, CommonLegacySettings->MaxWalkSlopeCosine, OUT LandingFloor))
 		{
 			// Adjust height to float slightly above walkable floor
-			LocationInProgress = UGroundMovementUtils::TestMoveToAdjustToFloor(Params.MovingComps, LocationInProgress, RotationInProgress, CommonLegacySettings->MaxWalkSlopeCosine, IN OUT LandingFloor, IN OUT MoveRecord);
+			LocationInProgress = UGroundMovementUtils::TestMoveToKeepMinHeightAboveFloor(Params.MovingComps, LocationInProgress, RotationInProgress, CommonLegacySettings->MaxWalkSlopeCosine, IN OUT LandingFloor, IN OUT MoveRecord);
 			CaptureFinalState(LocationInProgress, RotationInProgress.Rotator(), LandingFloor, DeltaSeconds, DeltaSeconds * PctTimeApplied, OutputSyncState, OutputState, IN OUT MoveRecord);
 			return;
 		}
@@ -255,7 +255,7 @@ void UAsyncFallingMode::SimulationTick_Implementation(const FSimulationTickParam
 		if (LandingFloor.IsWalkableFloor())
 		{
 			// Adjust height to float slightly above walkable floor
-			LocationInProgress = UGroundMovementUtils::TestMoveToAdjustToFloor(Params.MovingComps, LocationInProgress, RotationInProgress, CommonLegacySettings->MaxWalkSlopeCosine, IN OUT LandingFloor, IN OUT MoveRecord);
+			LocationInProgress = UGroundMovementUtils::TestMoveToKeepMinHeightAboveFloor(Params.MovingComps, LocationInProgress, RotationInProgress, CommonLegacySettings->MaxWalkSlopeCosine, IN OUT LandingFloor, IN OUT MoveRecord);
 		}
 	}
 	else
diff --git a/Engine/Plugins/Experimental/Mover/Source/Mover/Private/DefaultMovementSet/Modes/AsyncFlyingMode.cpp b/Engine/Plugins/Experimental/Mover/Source/Mover/Private/DefaultMovementSet/Modes/AsyncFlyingMode.cpp
--- a/Engine/Plugins/Experimental/Mover/Source/Mover/Private/DefaultMovementSet/Modes/AsyncFlyingMode.cpp
+++ b/Engine/Plugins/Experimental/Mover/Source/Mover/Private/DefaultMovementSet/Modes/AsyncFlyingMode.cpp
@@ -140,7 +140,7 @@ void UAsyncFlyingMode::SimulationTick_Implementation(const FSimulationTickParams
 
 		if (FloorUnderActor.IsWalkableFloor())
 		{
-			LocationInProgress = UGroundMovementUtils::TestMoveToAdjustToFloor(Params.MovingComps, LocationInProgress, RotationInProgress, CommonLegacySettings->MaxWalkSlopeCosine, IN OUT FloorUnderActor, MoveRecord);
+			LocationInProgress = UGroundMovementUtils::TestMoveToKeepMinHeightAboveFloor(Params.MovingComps, LocationInProgress, RotationInProgress, CommonLegacySettings->MaxWalkSlopeCosine, IN OUT FloorUnderActor, MoveRecord);
 		}
 	}
 
diff --git a/Engine/Plugins/Experimental/Mover/Source/Mover/Private/DefaultMovementSet/Modes/FallingMode.cpp b/Engine/Plugins/Experimental/Mover/Source/Mover/Private/DefaultMovementSet/Modes/FallingMode.cpp
--- a/Engine/Plugins/Experimental/Mover/Source/Mover/Private/DefaultMovementSet/Modes/FallingMode.cpp
+++ b/Engine/Plugins/Experimental/Mover/Source/Mover/Private/DefaultMovementSet/Modes/FallingMode.cpp
@@ -189,7 +189,7 @@ void UFallingMode::SimulationTick_Implementation(const FSimulationTickParams& Pa
 
 		if (FloorUnderActor.IsWalkableFloor())
 		{
-			UGroundMovementUtils::TryMoveToAdjustHeightAboveFloor(MoverComponent, FloorUnderActor, CommonLegacySettings->MaxWalkSlopeCosine, MoveRecord);
+			UGroundMovementUtils::TryMoveToKeepMinHeightAboveFloor(MoverComponent, FloorUnderActor, CommonLegacySettings->MaxWalkSlopeCosine, MoveRecord);
 			CaptureFinalState(UpdatedComponent, *StartingSyncState, FloorUnderActor, DeltaSeconds, DeltaSeconds * PctTimeApplied, OutputSyncState, OutputState, MoveRecord);
 			return;
 		}
@@ -222,7 +222,7 @@ void UFallingMode::SimulationTick_Implementation(const FSimulationTickParams& Pa
 		if (UAirMovementUtils::IsValidLandingSpot(Params.MovingComps, UpdatedComponent->GetComponentLocation(),
 			Hit, CommonLegacySettings->FloorSweepDistance, CommonLegacySettings->MaxWalkSlopeCosine, OUT LandingFloor))
 		{
-			UGroundMovementUtils::TryMoveToAdjustHeightAboveFloor(MoverComponent, LandingFloor, CommonLegacySettings->MaxWalkSlopeCosine, MoveRecord); // make sure we maintain a small gap over walking surfaces
+			UGroundMovementUtils::TryMoveToKeepMinHeightAboveFloor(MoverComponent, LandingFloor, CommonLegacySettings->MaxWalkSlopeCosine, MoveRecord); // make sure we maintain a small gap over walking surfaces
 			CaptureFinalState(UpdatedComponent, *StartingSyncState, LandingFloor, DeltaSeconds, DeltaSeconds * PctTimeApplied, OutputSyncState, OutputState, MoveRecord);
 			return;
 		}
@@ -242,7 +242,7 @@ void UFallingMode::SimulationTick_Implementation(const FSimulationTickParams& Pa
 
 		if (LandingFloor.IsWalkableFloor())
 		{
-			UGroundMovementUtils::TryMoveToAdjustHeightAboveFloor(MoverComponent, LandingFloor, CommonLegacySettings->MaxWalkSlopeCosine, MoveRecord); // make sure we maintain a small gap over walking surfaces
+			UGroundMovementUtils::TryMoveToKeepMinHeightAboveFloor(MoverComponent, LandingFloor, CommonLegacySettings->MaxWalkSlopeCosine, MoveRecord); // make sure we maintain a small gap over walking surfaces
 			CaptureFinalState(UpdatedComponent, *StartingSyncState, LandingFloor, DeltaSeconds, DeltaSeconds * PctTimeApplied, OutputSyncState, OutputState, MoveRecord);
 			return;
 		}
diff --git a/Engine/Plugins/Experimental/Mover/Source/Mover/Private/DefaultMovementSet/Modes/FlyingMode.cpp b/Engine/Plugins/Experimental/Mover/Source/Mover/Private/DefaultMovementSet/Modes/FlyingMode.cpp
--- a/Engine/Plugins/Experimental/Mover/Source/Mover/Private/DefaultMovementSet/Modes/FlyingMode.cpp
+++ b/Engine/Plugins/Experimental/Mover/Source/Mover/Private/DefaultMovementSet/Modes/FlyingMode.cpp
@@ -136,7 +136,7 @@ void UFlyingMode::SimulationTick_Implementation(const FSimulationTickParams& Par
 
 		if (FloorUnderActor.IsWalkableFloor())
 		{
-			UGroundMovementUtils::TryMoveToAdjustHeightAboveFloor(MoverComp, FloorUnderActor, CommonLegacySettings->MaxWalkSlopeCosine, MoveRecord);
+			UGroundMovementUtils::TryMoveToKeepMinHeightAboveFloor(MoverComp, FloorUnderActor, CommonLegacySettings->MaxWalkSlopeCosine, MoveRecord);
 		}
 	}
 
diff --git a/Engine/Plugins/Experimental/Mover/Source/Mover/Private/MoveLibrary/GroundMovementUtils.cpp b/Engine/Plugins/Experimental/Mover/Source/Mover/Private/MoveLibrary/GroundMovementUtils.cpp
--- a/Engine/Plugins/Experimental/Mover/Source/Mover/Private/MoveLibrary/GroundMovementUtils.cpp
+++ b/Engine/Plugins/Experimental/Mover/Source/Mover/Private/MoveLibrary/GroundMovementUtils.cpp
@@ -350,7 +350,7 @@ bool UGroundMovementUtils::TryMoveToStepUp(const FMovingComponentSet& MovingComp
 
 }
 
-bool UGroundMovementUtils::TryMoveToAdjustHeightAboveFloor(const FMovingComponentSet& MovingComps, FFloorCheckResult& CurrentFloor, float MaxWalkSlopeCosine, FMovementRecord& MoveRecord)
+static bool TryMoveToAdjustHeightAboveFloorInternal(const FMovingComponentSet& MovingComps, FFloorCheckResult& CurrentFloor, float MaxWalkSlopeCosine, bool bPushAwayOnly, FMovementRecord& MoveRecord)
 {
 	// If we have a floor check that hasn't hit anything, don't adjust height.
 	if (!CurrentFloor.IsWalkableFloor())
@@ -372,20 +372,26 @@ bool UGroundMovementUtils::TryMoveToAdjustHeightAboveFloor(const FMovingComponen
 		OldFloorDist = CurrentFloor.LineDist;
 	}
 
-	USceneComponent* UpdatedComponent = MovingComps.UpdatedComponent.Get();
-
 	// Move up or down to maintain floor height.
 	if (OldFloorDist < UE::FloorQueryUtility::MIN_FLOOR_DIST || OldFloorDist > UE::FloorQueryUtility::MAX_FLOOR_DIST)
 	{
-		FHitResult AdjustHit(1.f);
-		const float InitialUp = UpdatedComponent->GetComponentLocation().Dot(UpDirection);
 		const float AvgFloorDist = (UE::FloorQueryUtility::MIN_FLOOR_DIST + UE::FloorQueryUtility::MAX_FLOOR_DIST) * 0.5f;
 		const float MoveDist = AvgFloorDist - OldFloorDist;
 
+		// Don't adjust if we're only supposed to be staying away and the adjustment would bring us closer
+		if (bPushAwayOnly && MoveDist <= 0.0)
+		{
+			return false;
+		}
+
+		USceneComponent* UpdatedComponent = MovingComps.UpdatedComponent.Get();
+		const float InitialUp = UpdatedComponent->GetComponentLocation().Dot(UpDirection);
+		FHitResult AdjustHit(1.f);
+		
 		MoveRecord.LockRelevancy(false);
 
 		UMovementUtils::TrySafeMoveUpdatedComponent(MovingComps,
-			UpDirection * MoveDist, UpdatedComponent->GetComponentQuat(), 
+			UpDirection * MoveDist, UpdatedComponent->GetComponentQuat(),
 			true, AdjustHit, ETeleportType::None, MoveRecord);
 
 		MoveRecord.UnlockRelevancy();
@@ -416,6 +422,19 @@ bool UGroundMovementUtils::TryMoveToAdjustHeightAboveFloor(const FMovingComponen
 	return false;
 }
 
+
+bool UGroundMovementUtils::TryMoveToAdjustHeightAboveFloor(const FMovingComponentSet& MovingComps, FFloorCheckResult& CurrentFloor, float MaxWalkSlopeCosine, FMovementRecord& MoveRecord)
+{
+	return TryMoveToAdjustHeightAboveFloorInternal(MovingComps, CurrentFloor, MaxWalkSlopeCosine, /*bPushAwayOnly*/false, MoveRecord);
+}
+
+
+bool UGroundMovementUtils::TryMoveToKeepMinHeightAboveFloor(const FMovingComponentSet& MovingComps, FFloorCheckResult& CurrentFloor, float MaxWalkSlopeCosine, FMovementRecord& MoveRecord)
+{
+	return TryMoveToAdjustHeightAboveFloorInternal(MovingComps, CurrentFloor, MaxWalkSlopeCosine, /*bPushAwayOnly*/true, MoveRecord);
+}
+
+
 float UGroundMovementUtils::TryWalkToSlideAlongSurface(const FMovingComponentSet& MovingComps, const FVector& Delta, float PctOfDeltaToMove, const FQuat Rotation, const FVector& Normal, FHitResult& Hit, bool bHandleImpact, FMovementRecord& MoveRecord, float MaxWalkSlopeCosine, float MaxStepHeight)
 {
 	if (!Hit.bBlockingHit)
@@ -652,8 +671,7 @@ float UGroundMovementUtils::TestGroundedMoveAlongHitSurface(const FMovingCompone
 }
 
 
-/* static */
-FVector UGroundMovementUtils::TestMoveToAdjustToFloor(const FMovingComponentSet& MovingComps, const FVector& Location, const FQuat& Rotation, float MaxWalkSlopeCosine, FFloorCheckResult& InOutCurrentFloor, FMovementRecord& InOutMoveRecord)
+static FVector TestMoveToAdjustToFloorInternal(const FMovingComponentSet& MovingComps, const FVector& Location, const FQuat& Rotation, float MaxWalkSlopeCosine, bool bPushAwayOnly, FFloorCheckResult& InOutCurrentFloor, FMovementRecord& InOutMoveRecord)
 {
 	// If we have a floor check that hasn't hit anything, don't adjust height.
 	if (!InOutCurrentFloor.IsWalkableFloor())
@@ -678,10 +696,16 @@ FVector UGroundMovementUtils::TestMoveToAdjustToFloor(const FMovingComponentSet&
 	// Move up or down to maintain floor height.
 	if (OldFloorDist < UE::FloorQueryUtility::MIN_FLOOR_DIST || OldFloorDist > UE::FloorQueryUtility::MAX_FLOOR_DIST)
 	{
-		FHitResult AdjustHit(1.f);
-		const float InitialUp = Location.Dot(UpDirection);
 		const float AvgFloorDist = (UE::FloorQueryUtility::MIN_FLOOR_DIST + UE::FloorQueryUtility::MAX_FLOOR_DIST) * 0.5f;
 		const float MoveDist = AvgFloorDist - OldFloorDist;
+
+		if (bPushAwayOnly && MoveDist <= 0.0f)
+		{
+			return Location;	// we are far enough away from the floor already, so no adjustment is needed
+		}
+
+		FHitResult AdjustHit(1.f);
+		const float InitialUp = Location.Dot(UpDirection);
 		const FVector TargetAdjustedLocation = Location + (UpDirection * MoveDist);
 
 
@@ -719,6 +743,19 @@ FVector UGroundMovementUtils::TestMoveToAdjustToFloor(const FMovingComponentSet&
 	return Location; // no adjustment
 }
 
+
+FVector UGroundMovementUtils::TestMoveToAdjustToFloor(const FMovingComponentSet& MovingComps, const FVector& Location, const FQuat& Rotation, float MaxWalkSlopeCosine, FFloorCheckResult& InOutCurrentFloor, FMovementRecord& InOutMoveRecord)
+{
+	return TestMoveToAdjustToFloorInternal(MovingComps, Location, Rotation, MaxWalkSlopeCosine, /*bPushAwayOnly*/false, InOutCurrentFloor, InOutMoveRecord);
+}
+
+
+FVector UGroundMovementUtils::TestMoveToKeepMinHeightAboveFloor(const FMovingComponentSet& MovingComps, const FVector& Location, const FQuat& Rotation, float MaxWalkSlopeCosine, UPARAM(ref) FFloorCheckResult& InOutCurrentFloor, UPARAM(ref) FMovementRecord& InOutMoveRecord)
+{
+	return TestMoveToAdjustToFloorInternal(MovingComps, Location, Rotation, MaxWalkSlopeCosine, /*bPushAwayOnly*/true, InOutCurrentFloor, InOutMoveRecord);
+}
+
+
 /* static */
 bool UGroundMovementUtils::TestMoveToStepOver(const FMovingComponentSet& MovingComps, const FVector& GravDir, float MaxStepHeight, float MaxWalkSlopeCosine, float FloorSweepDistance, const FVector& MoveDelta, const FQuat& Rotation, const FHitResult& MoveHitResult, const FFloorCheckResult& CurrentFloor, bool bIsFalling, FOptionalFloorCheckResult* OutFloorTestResult, FVector& OutFinalLocation, FMovementRecord& InOutMoveRecord)
 {
diff --git a/Engine/Plugins/Experimental/Mover/Source/Mover/Public/MoveLibrary/GroundMovementUtils.h b/Engine/Plugins/Experimental/Mover/Source/Mover/Public/MoveLibrary/GroundMovementUtils.h
--- a/Engine/Plugins/Experimental/Mover/Source/Mover/Public/MoveLibrary/GroundMovementUtils.h
+++ b/Engine/Plugins/Experimental/Mover/Source/Mover/Public/MoveLibrary/GroundMovementUtils.h
@@ -91,6 +91,9 @@ public:
 	/** Moves vertically to stay within range of the walkable floor. Does nothing and returns false if floor is unwalkable or if already in range. */ 
 	static UE_API bool TryMoveToAdjustHeightAboveFloor(const FMovingComponentSet& MovingComps, FFloorCheckResult& CurrentFloor, float MaxWalkSlopeCosine, FMovementRecord& MoveRecord);
 
+	/** Moves vertically to maintain the minimum distance above the walkable floor. Does nothing and returns false if floor is unwalkable or if already far enough away. */
+	static UE_API bool TryMoveToKeepMinHeightAboveFloor(const FMovingComponentSet& MovingComps, FFloorCheckResult& CurrentFloor, float MaxWalkSlopeCosine, FMovementRecord& MoveRecord);
+
 	/** Attempts to move a component along a surface in the walking mode. Returns the percent of time applied, with 0.0 meaning no movement occurred.
      *  Note: This modifies the normal and calls UMovementUtils::TryMoveToSlideAlongSurface
      */
@@ -123,6 +126,14 @@ public:
 	UFUNCTION(BlueprintCallable, Category = Mover)
 	static UE_API FVector TestMoveToAdjustToFloor(const FMovingComponentSet& MovingComps, const FVector& Location, const FQuat& Rotation, float MaxWalkSlopeCosine, UPARAM(ref) FFloorCheckResult& InOutCurrentFloor, UPARAM(ref) FMovementRecord& InOutMoveRecord);
 
+	/** Tests potential vertical movement of a component to float a bit above a walkable floor. Intended for use while moving along ground surfaces.
+	 * Modifies InOutCurrentFloor to account for adjustments, if any was made
+     * Appends to InOutMoveRecord with any movement substeps
+	 * Returns new location, including any adjustment if it was made
+	 */
+	UFUNCTION(BlueprintCallable, Category = Mover)
+	static UE_API FVector TestMoveToKeepMinHeightAboveFloor(const FMovingComponentSet& MovingComps, const FVector& Location, const FQuat& Rotation, float MaxWalkSlopeCosine, UPARAM(ref) FFloorCheckResult& InOutCurrentFloor, UPARAM(ref) FMovementRecord& InOutMoveRecord);
+
 	/** Tests potential up-and-over movement of a component to overcome a blocking obstacle hit, such as moving up steps or slope.
 	 * Sets OutFloorTestResult with final floor info (optional)
 	 * Sets OutFinalLocation with final component location
