Fix to walking mode

diff --git a/Engine/Plugins/Experimental/Mover/Source/Mover/Private/DefaultMovementSet/Modes/AsyncNavWalkingMode.cpp b/Engine/Plugins/Experimental/Mover/Source/Mover/Private/DefaultMovementSet/Modes/AsyncNavWalkingMode.cpp
--- a/Engine/Plugins/Experimental/Mover/Source/Mover/Private/DefaultMovementSet/Modes/AsyncNavWalkingMode.cpp
+++ b/Engine/Plugins/Experimental/Mover/Source/Mover/Private/DefaultMovementSet/Modes/AsyncNavWalkingMode.cpp
@@ -504,6 +504,7 @@ void UAsyncNavWalkingMode::FindBestNavMeshLocation(const FVector& TraceStart, co
 		if (NavMovementUtils::CalculateNavMeshNormal(CachedNavLocation, NavMeshNormal, NavDataInterface.Get(), GetMoverComponent()))
 		{
 			OutHitResult.ImpactNormal = NavMeshNormal;
+			OutHitResult.Normal = NavMeshNormal;
 			UE_VLOG_ARROW(GetMoverComponent(), "AsyncNavWalkingMode", Display, OutHitResult.ImpactPoint, (OutHitResult.ImpactPoint + DebugNormalLength * OutHitResult.ImpactNormal), FColor::Magenta, TEXT("NavMeshNormal"));
 		}
 	}
@@ -595,7 +596,7 @@ FVector UAsyncNavWalkingMode::ProjectLocationFromNavMesh(float DeltaSeconds, con
 		}
 		else
 		{
-			const FVector ProjectedPoint = FMath::LinePlaneIntersection(TraceStart, TraceEnd, CachedProjectedNavMeshHitResult.ImpactPoint, CachedProjectedNavMeshHitResult.Normal);
+			const FVector ProjectedPoint = FMath::LinePlaneIntersection(TraceStart, TraceEnd, CachedProjectedNavMeshHitResult.ImpactPoint, CachedProjectedNavMeshHitResult.ImpactNormal);
 			UE_VLOG_SPHERE(GetMoverComponent(), "AsyncNavWalkingMode", Display, ProjectedPoint, 2.5f, FColor::Orange, TEXT("ProjectedPoint"));
 			FVector::FReal ProjectedVertical = ProjectedPoint.Dot(UpDirection);
 
