[Mover] [BugFix] Translating the pawn inside PhysicsStanceModifier (and other supporting changes)

diff --git a/Engine/Plugins/Experimental/Mover/Source/Mover/Private/Backends/MoverNetworkPhysicsLiaison.cpp b/Engine/Plugins/Experimental/Mover/Source/Mover/Private/Backends/MoverNetworkPhysicsLiaison.cpp
--- a/Engine/Plugins/Experimental/Mover/Source/Mover/Private/Backends/MoverNetworkPhysicsLiaison.cpp
+++ b/Engine/Plugins/Experimental/Mover/Source/Mover/Private/Backends/MoverNetworkPhysicsLiaison.cpp
@@ -373,10 +373,6 @@ void UMoverNetworkPhysicsLiaisonComponent::PerformPreSimulate_Internal(const FPh
 		CharacterDefaultInputs.SetMoveInput(EMoveInputType::DirectionalIntent, FVector::ZeroVector);
 	}
 
-	// Cache the particle XR before calling the parent, because the parent will teleport if the target is too far away and change the particle's transform
-	const FVector CharacterParticleStartingX = CharacterParticle->GetX();
-	const FRotator CharacterParticleStartingR = FRotator(CharacterParticle->GetR());
-
 	Super::PerformPreSimulate_Internal(TickParams, Input, Output);
 
 	FMoverDefaultSyncState& OutputSyncState = Output.SyncState.SyncStateCollection.FindOrAddMutableDataByType<FMoverDefaultSyncState>();
@@ -385,7 +381,7 @@ void UMoverNetworkPhysicsLiaisonComponent::PerformPreSimulate_Internal(const FPh
 	FVector TargetVelocity = OutputSyncState.GetVelocity_WorldSpace() + LocalGroundVelocity;
 	CharacterParticle->SetV(TargetVelocity);
 
-	FRotator DeltaRotation = OutputSyncState.GetOrientation_WorldSpace() - CharacterParticleStartingR;
+	FRotator DeltaRotation = OutputSyncState.GetOrientation_WorldSpace() - FRotator(CharacterParticle->GetR());
 	FRotator Winding, Remainder;
 	DeltaRotation.GetWindingAndRemainder(Winding, Remainder);
 	float TargetDeltaFacing = FMath::DegreesToRadians(Remainder.Yaw);
@@ -443,7 +439,7 @@ void UMoverNetworkPhysicsLiaisonComponent::PerformPreSimulate_Internal(const FPh
 			WalkableSlopeCosine = 2.0f;
 		}
 
-		const FVector TargetDeltaPos = OutputSyncState.GetLocation_WorldSpace() - CharacterParticleStartingX;
+		const FVector TargetDeltaPos = OutputSyncState.GetLocation_WorldSpace() - CharacterParticle->GetX();
 		ConstraintHandle->SetData({
 			Output.FloorResult.HitResult.ImpactNormal,
 			TargetDeltaPos,
diff --git a/Engine/Plugins/Experimental/Mover/Source/Mover/Private/Backends/MoverNetworkPhysicsLiaisonBase.cpp b/Engine/Plugins/Experimental/Mover/Source/Mover/Private/Backends/MoverNetworkPhysicsLiaisonBase.cpp
--- a/Engine/Plugins/Experimental/Mover/Source/Mover/Private/Backends/MoverNetworkPhysicsLiaisonBase.cpp
+++ b/Engine/Plugins/Experimental/Mover/Source/Mover/Private/Backends/MoverNetworkPhysicsLiaisonBase.cpp
@@ -672,6 +672,14 @@ void UMoverNetworkPhysicsLiaisonComponentBase::ProcessInputs_Internal(int32 Phys
 					UMoverComponent& MoverComp = GetMoverComponent();
 					FMoverAuxStateContext UnusedInvalidAuxState;
 					FMoverAuxStateContext UnusedAuxState;
+
+					// Prepare Networked SyncState for rollback
+					if (Chaos::FPBDRigidParticleHandle* ParticleHandle = GetControlledParticle_Internal())
+					{
+						FMoverDefaultSyncState& NetSyncState = LatestSyncState.SyncStateCollection.FindOrAddMutableDataByType<FMoverDefaultSyncState>();
+						NetSyncState.SetTransforms_WorldSpace(ParticleHandle->GetX(), FRotator(ParticleHandle->GetR()), ParticleHandle->GetV());
+					}
+
 					MoverComp.OnSimulationPreRollback(&Input.SyncState, &LatestSyncState, &UnusedInvalidAuxState, &UnusedAuxState);
 					GetCurrentStateData(Input.SyncState);
 					MoverComp.OnSimulationRollback(&Input.SyncState, &UnusedAuxState);
diff --git a/Engine/Plugins/Experimental/Mover/Source/Mover/Private/PhysicsMover/MovementModifiers/PhysicsStanceModifier.cpp b/Engine/Plugins/Experimental/Mover/Source/Mover/Private/PhysicsMover/MovementModifiers/PhysicsStanceModifier.cpp
--- a/Engine/Plugins/Experimental/Mover/Source/Mover/Private/PhysicsMover/MovementModifiers/PhysicsStanceModifier.cpp
+++ b/Engine/Plugins/Experimental/Mover/Source/Mover/Private/PhysicsMover/MovementModifiers/PhysicsStanceModifier.cpp
@@ -2,13 +2,96 @@
 
 #include "PhysicsMover/MovementModifiers/PhysicsStanceModifier.h"
 
+#include "Chaos/ParticleHandle.h"
 #include "Components/CapsuleComponent.h"
 #include "DefaultMovementSet/CharacterMoverComponent.h"
 #include "DefaultMovementSet/Settings/StanceSettings.h"
 #include "MoveLibrary/MovementUtils.h"
 #include "PhysicsMover/Modes/PhysicsDrivenWalkingMode.h"
 #include "PhysicsMover/Modes/PhysicsDrivenFallingMode.h"
+#include "PhysicsMover/PhysicsMovementUtils.h"
 
+void FPhysicsStanceModifier::ApplyModifierMove(UMoverComponent* MoverComp, const FMoverSyncState& InSyncState)
+{
+	if (const UBaseMovementMode* MovementMode = MoverComp->FindMovementModeByName(InSyncState.MovementMode))
+	{
+		if (MovementMode->Implements<UPhysicsCharacterMovementModeInterface>())
+		{
+			if (const IPhysicsCharacterMovementModeInterface* PhysicsMode = CastChecked<IPhysicsCharacterMovementModeInterface>(MovementMode))
+			{
+				if (FMoverSyncState* MoverState = const_cast<FMoverSyncState*>(&InSyncState))
+				{
+					if (UPrimitiveComponent* PrimitiveComp = Cast<UPrimitiveComponent>(MoverComp->GetUpdatedComponent()))
+					{
+						FMoverDefaultSyncState& SyncState = MoverState->SyncStateCollection.FindOrAddMutableDataByType<FMoverDefaultSyncState>();
+
+						FVector TargetPosition = SyncState.GetLocation_WorldSpace();
+						const float TargetHeight = PhysicsMode->GetTargetHeight();
+
+						const FVector UpDir = MoverComp->GetUpDirection();
+						FFloorCheckResult FloorResult;
+						float FloorDist = UE_BIG_NUMBER;
+						UMoverBlackboard* SimBlackboard = MoverComp->GetSimBlackboard_Mutable();
+
+						if (SimBlackboard->TryGet(CommonBlackboard::LastFloorResult, FloorResult) && FloorResult.bBlockingHit)
+						{
+							// Last floor distance is ok in this case
+							FloorDist = FloorResult.FloorDist;
+						}
+						else
+						{
+							// Perform a raycast if no floor found (e.g., on rollbacks)
+							FCollisionQueryParams QueryParams(SCENE_QUERY_STAT(PhysicsFloorTest), false, PrimitiveComp->GetOwner());
+							QueryParams.bTraceIntoSubComponents = false;
+							const ECollisionChannel CollisionChannel = PrimitiveComp->GetCollisionObjectType();
+							FCollisionResponseParams ResponseParams(ECR_Overlap);
+							ResponseParams.CollisionResponse.SetResponse(ECC_WorldStatic, ECR_Block);
+							ResponseParams.CollisionResponse.SetResponse(ECC_WorldDynamic, ECR_Block);
+							ResponseParams.CollisionResponse.SetResponse(ECC_Vehicle, ECR_Block);
+							ResponseParams.CollisionResponse.SetResponse(ECC_Destructible, ECR_Block);
+							ResponseParams.CollisionResponse.SetResponse(ECC_PhysicsBody, ECR_Block);
+
+							PrimitiveComp->InitSweepCollisionParams(QueryParams, ResponseParams);
+
+							const float AdjustedQueryDistance = 2.f * TargetHeight;
+
+							FVector Start = TargetPosition;
+							FVector End = Start - AdjustedQueryDistance * UpDir;
+							FHitResult OutHit;
+							TArray<FHitResult> Hits;
+							if (const UWorld* World = PrimitiveComp->GetWorld())
+							{
+								if (Chaos::Private::FGenericPhysicsInterface_Internal::RaycastMulti(World, Hits, Start, End, CollisionChannel, QueryParams, ResponseParams))
+								{
+									OutHit = Hits.Last();
+								}
+							}
+
+							FloorDist = UpDir.Dot(TargetPosition - OutHit.ImpactPoint);
+						}
+
+						if (FloorDist <= (2.f * TargetHeight))
+						{
+							TargetPosition += UpDir * (TargetHeight - FloorDist);
+						}
+						SyncState.SetTransforms_WorldSpace(
+							TargetPosition,
+							SyncState.GetOrientation_WorldSpace(),
+							SyncState.GetVelocity_WorldSpace(),
+							nullptr);
+
+						// Force update the character particle transform
+						if (Chaos::FPBDRigidParticleHandle* ParticleHandle = UPhysicsMovementUtils::GetRigidParticleHandleFromComponent(PrimitiveComp))
+						{
+							ParticleHandle->SetX(TargetPosition);
+							ParticleHandle->SetV(SyncState.GetVelocity_WorldSpace());
+						}
+					}
+				}
+			}
+		}
+	}
+}
 
 void FPhysicsStanceModifier::OnStart(UMoverComponent* MoverComp, const FMoverTimeStep& TimeStep, const FMoverSyncState& SyncState, const FMoverAuxStateContext& AuxState)
 {
@@ -26,6 +109,8 @@ void FPhysicsStanceModifier::OnStart(UMoverComponent* MoverComp, const FMoverTim
 				UpdateTargetHeight<UPhysicsDrivenWalkingMode>(MoverComp, OriginalHalfHeight, CrouchedHalfHeight, true);
 				UpdateTargetHeight<UPhysicsDrivenFallingMode>(MoverComp, OriginalHalfHeight, CrouchedHalfHeight, true);
 
+				ApplyModifierMove(MoverComp, SyncState);
+
 				ApplyMovementSettings(MoverComp);
 			}
 		}
@@ -53,6 +138,8 @@ void FPhysicsStanceModifier::OnEnd(UMoverComponent* MoverComp, const FMoverTimeS
 				UpdateTargetHeight<UPhysicsDrivenWalkingMode>(MoverComp, OriginalHalfHeight, CrouchedHalfHeight, false);
 				UpdateTargetHeight<UPhysicsDrivenFallingMode>(MoverComp, OriginalHalfHeight, CrouchedHalfHeight, false);
 
+				ApplyModifierMove(MoverComp, SyncState);
+
 				RevertMovementSettings(MoverComp);
 			}
 		}
diff --git a/Engine/Plugins/Experimental/Mover/Source/Mover/Public/PhysicsMover/MovementModifiers/PhysicsStanceModifier.h b/Engine/Plugins/Experimental/Mover/Source/Mover/Public/PhysicsMover/MovementModifiers/PhysicsStanceModifier.h
--- a/Engine/Plugins/Experimental/Mover/Source/Mover/Public/PhysicsMover/MovementModifiers/PhysicsStanceModifier.h
+++ b/Engine/Plugins/Experimental/Mover/Source/Mover/Public/PhysicsMover/MovementModifiers/PhysicsStanceModifier.h
@@ -54,6 +54,8 @@ protected:
 			}
 		}
 	}
+
+	void ApplyModifierMove(UMoverComponent* MoverComp, const FMoverSyncState& InSyncState);
 };
 
 #undef UE_API
