[Mover] [BugFix] Skydiving and Crouching compatibility issues
Skydiving inverts Mover's GetUpDirection() which prevents proper movement calculations for crouching. Force use UpVector for now.

diff --git a/Engine/Plugins/Experimental/Mover/Source/Mover/Private/PhysicsMover/MovementModifiers/PhysicsStanceModifier.cpp b/Engine/Plugins/Experimental/Mover/Source/Mover/Private/PhysicsMover/MovementModifiers/PhysicsStanceModifier.cpp
--- a/Engine/Plugins/Experimental/Mover/Source/Mover/Private/PhysicsMover/MovementModifiers/PhysicsStanceModifier.cpp
+++ b/Engine/Plugins/Experimental/Mover/Source/Mover/Private/PhysicsMover/MovementModifiers/PhysicsStanceModifier.cpp
@@ -28,7 +28,7 @@ void FPhysicsStanceModifier::ApplyModifierMove(UMoverComponent* MoverComp, const
 						FVector TargetPosition = SyncState.GetLocation_WorldSpace();
 						const float TargetHeight = PhysicsMode->GetTargetHeight();
 
-						const FVector UpDir = MoverComp->GetUpDirection();
+						const FVector UpDir = FVector::UpVector;
 						FFloorCheckResult FloorResult;
 						float FloorDist = UE_BIG_NUMBER;
 						UMoverBlackboard* SimBlackboard = MoverComp->GetSimBlackboard_Mutable();
@@ -207,7 +207,7 @@ bool FPhysicsStanceModifier::CanExpand_Internal(UMoverComponent* MoverComponent,
 		else
 		{
 			// Expand while keeping base location the same.
-			const FVector StandingLocation = PawnLocation + (HalfHeightDifference + .01f) * CharMoverComp->GetUpDirection();
+			const FVector StandingLocation = PawnLocation + (HalfHeightDifference + .01f) * FVector::UpVector;
 			bEncroached = UMovementUtils::OverlapTest(UpdatedComponent, UpdatedCompAsPrimitive, StandingLocation, PawnRot, CollisionChannel, StandingCapsuleShape, CharMoverComp->GetOwner());
 		}
 	}
@@ -232,7 +232,7 @@ void FPhysicsStanceModifier::OnPostSimulationTick(const FStanceModifier* Modifie
 					const float ColliderOffset = -MoverComp->GetPrimaryVisualComponent()->GetRelativeLocation().Z - OldHalfHeight;
 	
 					CapsuleComponent->SetCapsuleSize(CapsuleComponent->GetUnscaledCapsuleRadius(), NewHalfHeight);
-					MoverComp->GetPrimaryVisualComponent()->SetRelativeLocation(-(NewHalfHeight + ColliderOffset) * MoverComp->GetUpDirection());
+					MoverComp->GetPrimaryVisualComponent()->SetRelativeLocation(-(NewHalfHeight + ColliderOffset) * FVector::UpVector);
 	
 					bStanceChanged = true;
 					bPostProcessed = true;
@@ -245,7 +245,7 @@ void FPhysicsStanceModifier::OnPostSimulationTick(const FStanceModifier* Modifie
 					const float ColliderOffset = -MoverComp->GetPrimaryVisualComponent()->GetRelativeLocation().Z - OldHalfHeight;
 	
 					CapsuleComponent->SetCapsuleSize(CapsuleComponent->GetUnscaledCapsuleRadius(), NewHalfHeight);
-					MoverComp->GetPrimaryVisualComponent()->SetRelativeLocation(-(NewHalfHeight + ColliderOffset) * MoverComp->GetUpDirection());
+					MoverComp->GetPrimaryVisualComponent()->SetRelativeLocation(-(NewHalfHeight + ColliderOffset) * FVector::UpVector);
 	
 					bStanceChanged = true;
 					bPostProcessed = false;
