Mover transitions broadcast when triggering

diff --git a/Engine/Plugins/Experimental/Mover/Source/Mover/Private/MovementModeStateMachine.cpp b/Engine/Plugins/Experimental/Mover/Source/Mover/Private/MovementModeStateMachine.cpp
--- a/Engine/Plugins/Experimental/Mover/Source/Mover/Private/MovementModeStateMachine.cpp
+++ b/Engine/Plugins/Experimental/Mover/Source/Mover/Private/MovementModeStateMachine.cpp
@@ -365,6 +365,8 @@ void UMovementModeStateMachine::OnSimulationTick(USceneComponent* UpdatedCompone
 					OutputState.MovementEndState.NextModeName = EvalResult.NextMode;
 					OutputState.MovementEndState.RemainingMs = WorkingSimTickParams.TimeStep.StepMs; 	// Pass all remaining time to next mode
 					TransitionToTrigger->Trigger(WorkingSimTickParams);
+
+					MoverComp->OnMovementTransitionTriggered.Broadcast(TransitionToTrigger);
 				}
 				else
 				{
diff --git a/Engine/Plugins/Experimental/Mover/Source/Mover/Public/MoverComponent.h b/Engine/Plugins/Experimental/Mover/Source/Mover/Public/MoverComponent.h
--- a/Engine/Plugins/Experimental/Mover/Source/Mover/Public/MoverComponent.h
+++ b/Engine/Plugins/Experimental/Mover/Source/Mover/Public/MoverComponent.h
@@ -51,6 +51,9 @@ DECLARE_DYNAMIC_MULTICAST_DELEGATE_TwoParams(FMover_OnPostSimRollback, const FMo
 // Fired after changing movement modes. First param is the name of the previous movement mode. Second is the name of the new movement mode. 
 DECLARE_DYNAMIC_MULTICAST_DELEGATE_TwoParams(FMover_OnMovementModeChanged, const FName&, PreviousMovementModeName, const FName&, NewMovementModeName);
 
+// Fired after a transition has been triggered.
+DECLARE_DYNAMIC_MULTICAST_DELEGATE_OneParam(FMover_OnTransitionTriggered, const UBaseMovementModeTransition*, ModeTransition);
+
 // Fired after a frame has been finalized. This may be a resimulation or not. No changes to state are possible. Guaranteed to be on the game thread.
 DECLARE_DYNAMIC_MULTICAST_DELEGATE_TwoParams(FMover_OnPostFinalize, const FMoverSyncState&, SyncState, const FMoverAuxStateContext&, AuxState);
 
@@ -101,6 +104,10 @@ public:
 	UPROPERTY(BlueprintAssignable, Category = Mover)
 	FMover_OnMovementModeChanged OnMovementModeChanged;
 
+	// Broadcast when a Transition has been triggered.
+	UPROPERTY(BlueprintAssignable, Category = Mover)
+	FMover_OnTransitionTriggered OnMovementTransitionTriggered;
+
 	// Broadcast after each finalized simulation frame, after the state is finalized. (Game thread only)
 	UPROPERTY(BlueprintAssignable, Category = Mover)
 	FMover_OnPostFinalize OnPostFinalize;
