Mover - Add support for async movement modifiers

diff --git a/Engine/Plugins/Experimental/Mover/Source/Mover/Private/MovementModifier.cpp b/Engine/Plugins/Experimental/Mover/Source/Mover/Private/MovementModifier.cpp
--- a/Engine/Plugins/Experimental/Mover/Source/Mover/Private/MovementModifier.cpp
+++ b/Engine/Plugins/Experimental/Mover/Source/Mover/Private/MovementModifier.cpp
@@ -37,6 +37,18 @@ void FMovementModifierBase::EndModifier(UMoverComponent* MoverComp, const FMover
 	OnEnd(MoverComp, TimeStep, SyncState, AuxState);
 }
 
+void FMovementModifierBase::StartModifier_Async(const FMovementModifierParams_Async& Params)
+{
+	check(Params.TimeStep);
+	StartSimTimeMs = Params.TimeStep->BaseSimTimeMs;
+	OnStart_Async(Params);
+}
+
+void FMovementModifierBase::EndModifier_Async(const FMovementModifierParams_Async& Params)
+{
+	OnEnd_Async(Params);
+}
+
 bool FMovementModifierBase::IsFinished(float CurrentSimTimeMs) const
 {
 	const bool bHasStarted = (StartSimTimeMs >= 0.f);
@@ -136,6 +148,12 @@ TArray<TSharedPtr<FMovementModifierBase>> FMovementModifierGroup::GenerateActive
 	return ActiveModifiers;
 }
 
+TArray<TSharedPtr<FMovementModifierBase>> FMovementModifierGroup::GenerateActiveModifiers_Async(const FMovementModifierParams_Async& Params)
+{
+	FlushModifierArrays_Async(Params);
+	return ActiveModifiers;
+}
+
 static void CopyModifierArray(TArray<TSharedPtr<FMovementModifierBase>>& Dest, const TArray<TSharedPtr<FMovementModifierBase>>& Src)
 {
 	bool bCanCopyInPlace = (UE::Mover::DisableDataCopyInPlace == 0 && Dest.Num() == Src.Num());
@@ -339,7 +357,7 @@ void FMovementModifierGroup::FlushModifierArrays(UMoverComponent* MoverComp, con
 	// Make any queued moves active
 	for (TSharedPtr<FMovementModifierBase>& QueuedModifier : QueuedModifiers)
 	{
-		bool bModidiferAlreadyActive = false;
+		bool bModifierAlreadyActive = false;
 		for (TSharedPtr<FMovementModifierBase>& ActiveModifier : ActiveModifiers)
 		{
 			// We don't really need to assign the QueuedModifier a start time but it would help if modifiers are compared based off of start time as well
@@ -349,12 +367,12 @@ void FMovementModifierGroup::FlushModifierArrays(UMoverComponent* MoverComp, con
 			if (QueuedModifier->Matches(ActiveModifier.Get()))
 			{
 				ActiveModifier->OverwriteHandleIfInvalid(QueuedModifier->GetHandle());
-				bModidiferAlreadyActive = true;
+				bModifierAlreadyActive = true;
 				break;
 			}
 		}
 
-		if (!bModidiferAlreadyActive)
+		if (!bModifierAlreadyActive)
 		{
 			ActiveModifiers.Add(QueuedModifier);
 			QueuedModifier->StartModifier(MoverComp, TimeStep, SyncState, AuxState);
@@ -364,6 +382,58 @@ void FMovementModifierGroup::FlushModifierArrays(UMoverComponent* MoverComp, con
 	QueuedModifiers.Empty();
 }
 
+void FMovementModifierGroup::FlushModifierArrays_Async(const FMovementModifierParams_Async& Params)
+{
+	check(Params.TimeStep);
+
+	// Remove any finished moves
+	ActiveModifiers.RemoveAll([Params, this]
+	(const TSharedPtr<FMovementModifierBase>& Modifier)
+		{
+			if (Modifier.IsValid())
+			{
+				if (Modifier->IsFinished(Params.TimeStep->BaseSimTimeMs))
+				{
+					Modifier->EndModifier_Async(Params);
+					return true;
+				}
+			}
+			else
+			{
+				return true;
+			}
+
+			return false;
+		});
+
+	// Make any queued moves active
+	for (TSharedPtr<FMovementModifierBase>& QueuedModifier : QueuedModifiers)
+	{
+		bool bModifierAlreadyActive = false;
+		for (TSharedPtr<FMovementModifierBase>& ActiveModifier : ActiveModifiers)
+		{
+			// We don't really need to assign the QueuedModifier a start time but it would help if modifiers are compared based off of start time as well
+			QueuedModifier->StartSimTimeMs = Params.TimeStep->BaseSimTimeMs;
+
+			// We only want to queue this queued modifier if it wasn't already added from state received from authority. If we already have the modifier just assign it a handle since it's already been activated.
+			if (QueuedModifier->Matches(ActiveModifier.Get()))
+			{
+				ActiveModifier->OverwriteHandleIfInvalid(QueuedModifier->GetHandle());
+				bModifierAlreadyActive = true;
+				break;
+			}
+		}
+
+		if (!bModifierAlreadyActive)
+		{
+			ActiveModifiers.Add(QueuedModifier);
+			QueuedModifier->StartModifier_Async(Params);
+		}
+	}
+
+	QueuedModifiers.Empty();
+}
+
 struct FMovementModifierDeleter
 {
 	FORCEINLINE void operator()(FMovementModifierBase* Object) const
@@ -463,4 +533,3 @@ void FMovementModifierGroup::Reset()
 	QueuedModifiers.Empty();
 	ActiveModifiers.Empty();
 }
-
diff --git a/Engine/Plugins/Experimental/Mover/Source/Mover/Public/MovementModifier.h b/Engine/Plugins/Experimental/Mover/Source/Mover/Public/MovementModifier.h
--- a/Engine/Plugins/Experimental/Mover/Source/Mover/Public/MovementModifier.h
+++ b/Engine/Plugins/Experimental/Mover/Source/Mover/Public/MovementModifier.h
@@ -9,6 +9,7 @@
 #define UE_API MOVER_API
 
 class UMoverComponent;
+class UMoverSimulation;
 struct FMoverAuxStateContext;
 struct FMoverTimeStep;
 struct FMoverSyncState;
@@ -68,6 +69,27 @@ private:
 	MODIFIER_HANDLE_TYPE Handle;
 };
 
+struct FMovementModifierParams_Async
+{
+	FMovementModifierParams_Async(UMoverSimulation* InSimulation, const FMoverSyncState* InSyncState, const FMoverTimeStep* InTimeStep)
+		: Simulation(InSimulation)
+		, SyncState(InSyncState)
+		, TimeStep(InTimeStep)
+	{
+		ensureMsgf(IsValid(), TEXT("Invalid initialization data for FMovementModifierParams_Async"));
+	}
+	FMovementModifierParams_Async() = delete;
+
+	bool IsValid() const
+	{
+		return Simulation && SyncState && TimeStep;
+	}
+
+	UMoverSimulation* Simulation = nullptr;
+	const FMoverSyncState* SyncState = nullptr;
+	const FMoverTimeStep* TimeStep = nullptr;
+};
+
 /**
  * Movement Modifier: Used to apply changes that indirectly influence the movement simulation,
  * without proposing or executing any movement, but still in sync with the sim.
@@ -116,6 +138,23 @@ struct FMovementModifierBase
 	/** Runtime query whether this modifier is finished and can be destroyed. The default implementation is based on DurationMs. */
 	UE_API virtual bool IsFinished(float CurrentSimTimeMs) const;
 	
+
+	// Begin Async API
+	// The async API is used by the async friendly mover simulation class
+	// which is currently implemented in ChaosMover for async physics support.
+	// The standard Mover update does not currently require these async functions
+	// to be implemented
+	virtual void OnStart_Async(const FMovementModifierParams_Async& Params) {}
+	virtual void OnEnd_Async(const FMovementModifierParams_Async& Params) {}
+	virtual void OnPreMovement_Async(const FMovementModifierParams_Async& Params) {}
+	virtual void OnPostMovement_Async(const FMovementModifierParams_Async& Params) {}
+
+	UE_API void StartModifier_Async(const FMovementModifierParams_Async& Params);
+	UE_API void EndModifier_Async(const FMovementModifierParams_Async& Params);
+	// End Async API
+
+
+
 	// @return newly allocated copy of this FMovementModifier. Must be overridden by child classes
 	UE_API virtual FMovementModifierBase* Clone() const;
 
@@ -197,6 +236,8 @@ struct FMovementModifierGroup
 	// Generates active modifier list (by calling FlushModifierArrays) and returns the an array of all currently active modifiers
 	UE_API TArray<TSharedPtr<FMovementModifierBase>> GenerateActiveModifiers(UMoverComponent* MoverComp, const FMoverTimeStep& TimeStep, const FMoverSyncState& SyncState, const FMoverAuxStateContext& AuxState);
 	
+	UE_API TArray<TSharedPtr<FMovementModifierBase>> GenerateActiveModifiers_Async(const FMovementModifierParams_Async& Params);
+
 	/** Copy operator - deep copy so it can be used for archiving/saving off moves */
 	UE_API FMovementModifierGroup& operator=(const FMovementModifierGroup& Other);
 
@@ -234,6 +275,8 @@ protected:
 	// Clears out any finished or invalid modifiers and adds any queued modifiers to the active modifiers
 	UE_API void FlushModifierArrays(UMoverComponent* MoverComp, const FMoverTimeStep& TimeStep, const FMoverSyncState& SyncState, const FMoverAuxStateContext& AuxState);
 	
+	UE_API void FlushModifierArrays_Async(const FMovementModifierParams_Async& Params);
+
 	/** Helper function for serializing array of movement modifiers */
 	static UE_API void NetSerializeMovementModifierArray(FArchive& Ar, TArray< TSharedPtr<FMovementModifierBase> >& ModifierArray, uint8 MaxNumModifiersToSerialize = MAX_uint8);
 	
