Mover - Add callback to simulation event

diff --git a/Engine/Plugins/Experimental/Mover/Source/Mover/Private/MoverComponent.cpp b/Engine/Plugins/Experimental/Mover/Source/Mover/Private/MoverComponent.cpp
--- a/Engine/Plugins/Experimental/Mover/Source/Mover/Private/MoverComponent.cpp
+++ b/Engine/Plugins/Experimental/Mover/Source/Mover/Private/MoverComponent.cpp
@@ -2251,6 +2251,9 @@ void UMoverComponent::SetSimulationOutput(const FMoverTimeStep& TimeStep, const
 	{
 		if (const FMoverSimulationEventData* Data = EventData.Get())
 		{
+			// This gives the event a callback when it is processed on the game thread
+			Data->OnEventProcessed();
+
 			// Process the simulation event at the mover component (or derived) level
 			ProcessSimulationEvent(*Data);
 
diff --git a/Engine/Plugins/Experimental/Mover/Source/Mover/Public/MoverSimulationTypes.h b/Engine/Plugins/Experimental/Mover/Source/Mover/Public/MoverSimulationTypes.h
--- a/Engine/Plugins/Experimental/Mover/Source/Mover/Public/MoverSimulationTypes.h
+++ b/Engine/Plugins/Experimental/Mover/Source/Mover/Public/MoverSimulationTypes.h
@@ -11,6 +11,8 @@
 #include "MovementModifier.h"
 #include "MoverDataModelTypes.h"
 #include "UObject/Interface.h"
+#include <functional>
+
 #include "MoverSimulationTypes.generated.h"
 
 // Names for our default modes
@@ -169,7 +171,7 @@ public:
 	bool ShouldReconcile(const FMoverSyncState& AuthorityState) const
 	{
 		return (MovementMode != AuthorityState.MovementMode) || 
-		       SyncStateCollection.ShouldReconcile(AuthorityState.SyncStateCollection);
+			   SyncStateCollection.ShouldReconcile(AuthorityState.SyncStateCollection);
 	}
 
 	void Interpolate(const FMoverSyncState* From, const FMoverSyncState* To, float Pct)
@@ -417,11 +419,15 @@ struct FMoverSimulationEventData
 {
 	GENERATED_BODY()
 
-	FMoverSimulationEventData(float InEventTimeMs)
-		: EventTimeMs(InEventTimeMs)
+	using FEventProcessedCallbackPtr = std::function<void(const FMoverSimulationEventData& Data)>;
+
+	FMoverSimulationEventData(float InEventTimeMs, FEventProcessedCallbackPtr InEventProcessedCallback = nullptr)
+		: EventProcessedCallback(InEventProcessedCallback)
+		, EventTimeMs(InEventTimeMs)
 	{
 	}
 	FMoverSimulationEventData() {}
+	virtual ~FMoverSimulationEventData() {}
 
 	// User must override
 	MOVER_API virtual UScriptStruct* GetScriptStruct() const;
@@ -438,6 +444,25 @@ struct FMoverSimulationEventData
 		return const_cast<const T*>(const_cast<FMoverSimulationEventData*>(this)->CastTo_Mutable<T>());
 	}
 
+	void OnEventProcessed() const
+	{
+		if (EventProcessedCallback)
+		{
+			EventProcessedCallback(*this);
+		}
+	}
+
+	void SetEventProcessedCallback(FEventProcessedCallbackPtr Callback)
+	{
+		EventProcessedCallback = Callback;
+	}
+
+private:
+	// This callback is fired when the event is processed on the game thread
+	// This is called before and in addition to any type based handling
+	FEventProcessedCallbackPtr EventProcessedCallback = nullptr;
+
+public:
 	float EventTimeMs = 0.0f;
 };
 
@@ -446,8 +471,8 @@ struct FMovementModeChangedEventData : public FMoverSimulationEventData
 {
 	GENERATED_BODY()
 
-	FMovementModeChangedEventData(float InEventTimeMs, const FName InPreviousModeName, const FName InNewModeName)
-		: FMoverSimulationEventData(InEventTimeMs)
+	FMovementModeChangedEventData(float InEventTimeMs, const FName InPreviousModeName, const FName InNewModeName, FEventProcessedCallbackPtr InEventProcessedCallback = nullptr)
+		: FMoverSimulationEventData(InEventTimeMs, InEventProcessedCallback)
 		, PreviousModeName(InPreviousModeName)
 		, NewModeName(InNewModeName)
 	{
