Mover - Add External Activate/Deactivate functions to be called when the mode is set on the game thread

diff --git a/Engine/Plugins/Experimental/Mover/Source/Mover/Private/MovementMode.cpp b/Engine/Plugins/Experimental/Mover/Source/Mover/Private/MovementMode.cpp
--- a/Engine/Plugins/Experimental/Mover/Source/Mover/Private/MovementMode.cpp
+++ b/Engine/Plugins/Experimental/Mover/Source/Mover/Private/MovementMode.cpp
@@ -63,12 +63,34 @@ void UBaseMovementMode::OnUnregistered()
 
 void UBaseMovementMode::Activate()
 {
-	K2_OnActivated();
+	if (!bSupportsAsync)
+	{
+		K2_OnActivated();
+	}
 }
 
 void UBaseMovementMode::Deactivate()
 {
-	K2_OnDeactivated();
+	if (!bSupportsAsync)
+	{
+		K2_OnDeactivated();
+	}
+}
+
+void UBaseMovementMode::Activate_External()
+{
+	if (bSupportsAsync)
+	{
+		K2_OnActivated();
+	}
+}
+
+void UBaseMovementMode::Deactivate_External()
+{
+	if (bSupportsAsync)
+	{
+		K2_OnDeactivated();
+	}
 }
 
 void UBaseMovementMode::GenerateMove_Implementation(const FMoverTickStartData& StartState, const FMoverTimeStep& TimeStep, FProposedMove& OutProposedMove) const
diff --git a/Engine/Plugins/Experimental/Mover/Source/Mover/Private/MoverComponent.cpp b/Engine/Plugins/Experimental/Mover/Source/Mover/Private/MoverComponent.cpp
--- a/Engine/Plugins/Experimental/Mover/Source/Mover/Private/MoverComponent.cpp
+++ b/Engine/Plugins/Experimental/Mover/Source/Mover/Private/MoverComponent.cpp
@@ -2272,8 +2272,34 @@ void UMoverComponent::SetSimulationOutput(const FMoverTimeStep& TimeStep, const
 
 void UMoverComponent::ProcessSimulationEvent(const FMoverSimulationEventData& EventData)
 {
+	// On a mode change call deactivate on the previous mode and activate on the new mode,
+	// then broadcast the mode changed event
 	if (const FMovementModeChangedEventData* ModeChangedData = EventData.CastTo<FMovementModeChangedEventData>())
 	{
+		if (ModeChangedData->PreviousModeName != NAME_None)
+		{
+			if (TObjectPtr<UBaseMovementMode>* PrevModePtr = MovementModes.Find(ModeChangedData->PreviousModeName))
+			{
+				UBaseMovementMode* PrevMode = PrevModePtr->Get();
+				if (PrevMode && PrevMode->bSupportsAsync)
+				{
+					PrevMode->Deactivate_External();
+				}
+			}
+		}
+
+		if (ModeChangedData->NewModeName != NAME_None)
+		{
+			if (TObjectPtr<UBaseMovementMode>* NewModePtr = MovementModes.Find(ModeChangedData->NewModeName))
+			{
+				UBaseMovementMode* NewMode = NewModePtr->Get();
+				if (NewMode && NewMode->bSupportsAsync)
+				{
+					NewMode->Activate_External();
+				}
+			}
+		}
+
 		OnMovementModeChanged.Broadcast(ModeChangedData->PreviousModeName, ModeChangedData->NewModeName);
 	}
 }
diff --git a/Engine/Plugins/Experimental/Mover/Source/Mover/Public/MovementMode.h b/Engine/Plugins/Experimental/Mover/Source/Mover/Public/MovementMode.h
--- a/Engine/Plugins/Experimental/Mover/Source/Mover/Public/MovementMode.h
+++ b/Engine/Plugins/Experimental/Mover/Source/Mover/Public/MovementMode.h
@@ -46,8 +46,14 @@ public:
 	UE_API virtual void OnRegistered(const FName ModeName);
 	UE_API virtual void OnUnregistered();
 	
+	// These functions are called immediately when the state machine switches modes
 	UE_API virtual void Activate();
 	UE_API virtual void Deactivate();
+
+	// These functions are called when the sync state is changed on the game thread
+	// and a new mode is activated/deactivated
+	UE_API virtual void Activate_External();
+	UE_API virtual void Deactivate_External();
 	
 	UFUNCTION(BlueprintNativeEvent, meta = (DisplayName = "Generate Move"))
 	UE_API void GenerateMove(const FMoverTickStartData& StartState, const FMoverTimeStep& TimeStep, FProposedMove& OutProposedMove) const;
