Mover: fixing cases where strict equality checks on orientations should have a small tolerance

diff --git a/Engine/Plugins/Experimental/Mover/Source/Mover/Private/DefaultMovementSet/Modes/AsyncFlyingMode.cpp b/Engine/Plugins/Experimental/Mover/Source/Mover/Private/DefaultMovementSet/Modes/AsyncFlyingMode.cpp
--- a/Engine/Plugins/Experimental/Mover/Source/Mover/Private/DefaultMovementSet/Modes/AsyncFlyingMode.cpp
+++ b/Engine/Plugins/Experimental/Mover/Source/Mover/Private/DefaultMovementSet/Modes/AsyncFlyingMode.cpp
@@ -96,9 +96,7 @@ void UAsyncFlyingMode::SimulationTick_Implementation(const FSimulationTickParams
 
 	// Use the orientation intent directly. If no intent is provided, use last frame's orientation. Note that we are assuming rotation changes can't fail. 
 	const FRotator StartingOrient = StartingSyncState->GetOrientation_WorldSpace();
-
 	const FRotator TargetOrient = UMovementUtils::ApplyAngularVelocity(StartingOrient, ProposedMove.AngularVelocity, DeltaSeconds);
-	const bool bIsOrientationChanging = (TargetOrient != StartingOrient);
 	
 	const FVector StartLocation = StartingSyncState->GetLocation_WorldSpace();
 	const FVector TargetLocation = StartLocation + (ProposedMove.LinearVelocity * DeltaSeconds);
diff --git a/Engine/Plugins/Experimental/Mover/Source/Mover/Private/DefaultMovementSet/Modes/AsyncWalkingMode.cpp b/Engine/Plugins/Experimental/Mover/Source/Mover/Private/DefaultMovementSet/Modes/AsyncWalkingMode.cpp
--- a/Engine/Plugins/Experimental/Mover/Source/Mover/Private/DefaultMovementSet/Modes/AsyncWalkingMode.cpp
+++ b/Engine/Plugins/Experimental/Mover/Source/Mover/Private/DefaultMovementSet/Modes/AsyncWalkingMode.cpp
@@ -159,7 +159,7 @@ void UAsyncWalkingMode::SimulationTick_Implementation(const FSimulationTickParam
 
 	const FRotator StartingOrient = StartingSyncState->GetOrientation_WorldSpace();
 	const FRotator TargetOrient = UMovementUtils::ApplyAngularVelocity(StartingOrient, ProposedMove.AngularVelocity, DeltaSeconds);
-	const bool bIsOrientationChanging = (TargetOrient != StartingOrient);
+	const bool bIsOrientationChanging = !StartingOrient.Equals(TargetOrient);
 
 	const FQuat StartRotation = StartingOrient.Quaternion();
 
diff --git a/Engine/Plugins/Experimental/Mover/Source/Mover/Private/DefaultMovementSet/Modes/FlyingMode.cpp b/Engine/Plugins/Experimental/Mover/Source/Mover/Private/DefaultMovementSet/Modes/FlyingMode.cpp
--- a/Engine/Plugins/Experimental/Mover/Source/Mover/Private/DefaultMovementSet/Modes/FlyingMode.cpp
+++ b/Engine/Plugins/Experimental/Mover/Source/Mover/Private/DefaultMovementSet/Modes/FlyingMode.cpp
@@ -102,7 +102,7 @@ void UFlyingMode::SimulationTick_Implementation(const FSimulationTickParams& Par
 	const FRotator StartingOrient = StartingSyncState->GetOrientation_WorldSpace();
 
 	const FRotator TargetOrient = UMovementUtils::ApplyAngularVelocity(StartingOrient, ProposedMove.AngularVelocity, DeltaSeconds);
-	const bool bIsOrientationChanging = (TargetOrient != StartingOrient);
+	const bool bIsOrientationChanging = !StartingOrient.Equals(TargetOrient);
 	
 	FVector MoveDelta = ProposedMove.LinearVelocity * DeltaSeconds;
 
diff --git a/Engine/Plugins/Experimental/Mover/Source/Mover/Private/DefaultMovementSet/Modes/NavWalkingMode.cpp b/Engine/Plugins/Experimental/Mover/Source/Mover/Private/DefaultMovementSet/Modes/NavWalkingMode.cpp
--- a/Engine/Plugins/Experimental/Mover/Source/Mover/Private/DefaultMovementSet/Modes/NavWalkingMode.cpp
+++ b/Engine/Plugins/Experimental/Mover/Source/Mover/Private/DefaultMovementSet/Modes/NavWalkingMode.cpp
@@ -151,7 +151,7 @@ void UNavWalkingMode::SimulationTick_Implementation(const FSimulationTickParams&
 
 	const FRotator StartingOrient = StartingSyncState->GetOrientation_WorldSpace();
 	const FRotator TargetOrient = UMovementUtils::ApplyAngularVelocity(StartingOrient, ProposedMove.AngularVelocity, DeltaSeconds);
-	const bool bIsOrientationChanging = (TargetOrient != StartingOrient);
+	const bool bIsOrientationChanging = !StartingOrient.Equals(TargetOrient);
 	
 	FQuat TargetOrientQuat = TargetOrient.Quaternion();
 	if (CommonLegacySettings->bShouldRemainVertical)
diff --git a/Engine/Plugins/Experimental/Mover/Source/Mover/Private/DefaultMovementSet/Modes/WalkingMode.cpp b/Engine/Plugins/Experimental/Mover/Source/Mover/Private/DefaultMovementSet/Modes/WalkingMode.cpp
--- a/Engine/Plugins/Experimental/Mover/Source/Mover/Private/DefaultMovementSet/Modes/WalkingMode.cpp
+++ b/Engine/Plugins/Experimental/Mover/Source/Mover/Private/DefaultMovementSet/Modes/WalkingMode.cpp
@@ -151,7 +151,7 @@ void UWalkingMode::SimulationTick_Implementation(const FSimulationTickParams& Pa
 
 	const FRotator StartingOrient = StartingSyncState->GetOrientation_WorldSpace();
 	const FRotator TargetOrient = UMovementUtils::ApplyAngularVelocity(StartingOrient, ProposedMove.AngularVelocity, DeltaSeconds);
-	const bool bIsOrientationChanging = (TargetOrient != StartingOrient);
+	const bool bIsOrientationChanging = !StartingOrient.Equals(TargetOrient);
 
 	FQuat TargetOrientQuat = TargetOrient.Quaternion();
 	if (CommonLegacySettings->bShouldRemainVertical)
