Mover: fixed issue where a project with a specified InputProducer could fail to have ProduceInput called on it.

diff --git a/Engine/Plugins/Experimental/Mover/Source/Mover/Private/MoverComponent.cpp b/Engine/Plugins/Experimental/Mover/Source/Mover/Private/MoverComponent.cpp
--- a/Engine/Plugins/Experimental/Mover/Source/Mover/Private/MoverComponent.cpp
+++ b/Engine/Plugins/Experimental/Mover/Source/Mover/Private/MoverComponent.cpp
@@ -277,28 +277,24 @@ void UMoverComponent::BeginPlay()
 		}
 
 		// If an InputProducer isn't already set, check if the actor is one
-		if (!InputProducer)
+		if (!InputProducer &&
+			MyActor->GetClass()->ImplementsInterface(UMoverInputProducerInterface::StaticClass()))
 		{
-			if (MyActor->GetClass()->ImplementsInterface(UMoverInputProducerInterface::StaticClass()))
-			{
-				InputProducer = MyActor;
-				InputProducers.Add(MyActor);
-			}
+			InputProducer = MyActor;
+		}
+
+		if (InputProducer)
+		{
+			InputProducers.AddUnique(InputProducer);
 		}
 
 		TSet<UActorComponent*> Components = MyActor->GetComponents();
 		for (UActorComponent* Component : Components)
 		{
-			if (IsValid(Component))
+			if (IsValid(Component) &&
+				Component->GetClass()->ImplementsInterface(UMoverInputProducerInterface::StaticClass()))
 			{
-				if (Component->GetClass()->ImplementsInterface(UMoverInputProducerInterface::StaticClass()))
-				{
-					// Let's not add that component if a derived class made it the InputProducer
-					if (InputProducer != Component)
-					{
-						InputProducers.Add(Component);
-					}
-				}
+				InputProducers.AddUnique(Component);
 			}
 		}
 	}
