Add a spring walking movement mode to Mover.
* This uses a critically damped velocity spring model to translate and rotate the character.
* Add utility spring smoothing method (internal for now), and tweak the parameterization of existing internal spring method to better match others already existing in the engine.
* Add level to anim sandbox to test new mode

#tests animsandbox

#virtualized

diff --git a/Engine/Plugins/Animation/AnimationWarping/Source/Runtime/Private/BoneControllers/AnimNode_Steering.cpp b/Engine/Plugins/Animation/AnimationWarping/Source/Runtime/Private/BoneControllers/AnimNode_Steering.cpp
--- a/Engine/Plugins/Animation/AnimationWarping/Source/Runtime/Private/BoneControllers/AnimNode_Steering.cpp
+++ b/Engine/Plugins/Animation/AnimationWarping/Source/Runtime/Private/BoneControllers/AnimNode_Steering.cpp
@@ -148,7 +148,7 @@ void FAnimNode_Steering::EvaluateSkeletalControl_AnyThread(FComponentSpacePoseCo
 						{
 							// Apply linear correction
 							FQuat LinearCorrection = FQuat::Identity;
-							AnimationMath::SimpleSpringDamperQuat(IN OUT LinearCorrection, IN OUT AngularVelocity, DeltaToTargetOrientation, ProceduralTargetTime, DeltaSeconds);
+							AnimationMath::SimpleSpringDamperQuat(IN OUT LinearCorrection, IN OUT AngularVelocity, DeltaToTargetOrientation, AnimationMath::SpringHalfLifeToSmoothingTime(ProceduralTargetTime), DeltaSeconds);
 
 							UE_VLOG_ARROW(Output.AnimInstanceProxy->GetAnimInstanceObject(), "Steering", Display,
 								RootBoneTransform.GetLocation(),
diff --git a/Engine/Plugins/Experimental/AnimNextWarping/Source/AnimNextWarping/Private/SteeringTrait.cpp b/Engine/Plugins/Experimental/AnimNextWarping/Source/AnimNextWarping/Private/SteeringTrait.cpp
--- a/Engine/Plugins/Experimental/AnimNextWarping/Source/AnimNextWarping/Private/SteeringTrait.cpp
+++ b/Engine/Plugins/Experimental/AnimNextWarping/Source/AnimNextWarping/Private/SteeringTrait.cpp
@@ -198,7 +198,7 @@ void FAnimNextSteeringTask::Execute(UE::AnimNext::FEvaluationVM& VM) const
 						{
 							// Apply linear correction
 							FQuat LinearCorrection = FQuat::Identity;
-							AnimationMath::SimpleSpringDamperQuat(LinearCorrection, InstanceData->AngularVelocity, DeltaToTargetOrientation, SharedData->ProceduralTargetTime, InstanceData->DeltaTime);
+							AnimationMath::SimpleSpringDamperQuat(LinearCorrection, InstanceData->AngularVelocity, DeltaToTargetOrientation, AnimationMath::SpringHalfLifeToSmoothingTime(SharedData->ProceduralTargetTime), InstanceData->DeltaTime);
 
 #if ENABLE_ANIM_DEBUG
 							UE_VLOG_ARROW(InstanceData->HostObject, "Steering", Display,
diff --git a/Engine/Plugins/Experimental/Mover/Source/Mover/Private/DefaultMovementSet/Modes/SpringWalkingMode.cpp b/Engine/Plugins/Experimental/Mover/Source/Mover/Private/DefaultMovementSet/Modes/SpringWalkingMode.cpp
new file mode 100644
--- /dev/null
+++ b/Engine/Plugins/Experimental/Mover/Source/Mover/Private/DefaultMovementSet/Modes/SpringWalkingMode.cpp
@@ -0,0 +1,115 @@
+// Copyright Epic Games, Inc. All Rights Reserved.
+
+#include "DefaultMovementSet/Modes/SpringWalkingMode.h"
+#include "DefaultMovementSet/Modes/SpringWalkingState.h"
+
+#include "MoverComponent.h"
+#include "DefaultMovementSet/Settings/CommonLegacyMovementSettings.h"
+#include "Runtime/Engine/Internal/Animation/AnimationMath.h"
+
+void USpringWalkingMode::GenerateMove_Implementation(const FMoverTickStartData& StartState, const FMoverTimeStep& TimeStep,
+                                                   FProposedMove& OutProposedMove) const
+{
+	if (CommonLegacySettings.Get() == nullptr)
+	{
+		return;
+	}
+
+	const float DeltaSeconds = TimeStep.StepMs * 0.001f;
+	if (DeltaSeconds <= 0.0f)
+	{
+		return;
+	}
+	
+	const FMoverDefaultSyncState* StartingSyncState = StartState.SyncState.SyncStateCollection.FindDataByType<FMoverDefaultSyncState>();
+
+	// Get input
+	const FCharacterDefaultInputs* CharacterInputs = StartState.InputCmd.InputCollection.FindDataByType<FCharacterDefaultInputs>();
+
+	// Need to hack around the fact that GenerateMove is const, but we want to mutate state data here for our spring
+	FMoverTickStartData* MutableStartData = const_cast<FMoverTickStartData*>(&StartState);
+	FMoverSpringState& SpringState = MutableStartData->SyncState.SyncStateCollection.FindOrAddMutableDataByType<FMoverSpringState>();
+	
+	FVector DesiredVelocity = CharacterInputs->GetMoveInput_WorldSpace();
+	float MaxMoveSpeed = CommonLegacySettings->MaxSpeed;
+
+	// Subtract vertical component but keep same magnitude
+	float DesiredVelMag = DesiredVelocity.Length();
+	DesiredVelocity -= DesiredVelocity.ProjectOnTo(GetMoverComponent()->GetUpDirection());
+	float DesiredVel2DSquaredLength = DesiredVelocity.SquaredLength(); 
+	if (DesiredVel2DSquaredLength > 0.0f)
+	{
+		DesiredVelocity *= DesiredVelMag / FMath::Sqrt(DesiredVel2DSquaredLength); 
+	}
+
+	switch (CharacterInputs->GetMoveInputType())
+	{
+	case EMoveInputType::DirectionalIntent:
+		{
+			DesiredVelocity *= MaxMoveSpeed;
+		}
+		break;
+	case EMoveInputType::Velocity:
+		{
+			// Clamp to max move speed
+			DesiredVelocity = DesiredVelocity.GetClampedToMaxSize(MaxMoveSpeed);
+		}
+		break;
+	case EMoveInputType::None:
+	case EMoveInputType::Invalid:
+	default:
+		{
+			UE_LOG(LogMover, Warning, TEXT("Unhandled MoveInputType %i in USpringWalkingMode"), CharacterInputs->GetMoveInputType());
+			DesiredVelocity = FVector::ZeroVector;
+		}
+	break;
+	}
+
+	// Linear //
+	FVector CurrentPosition = StartingSyncState->GetLocation_WorldSpace();
+	FVector PrevPosition = CurrentPosition;
+
+	AnimationMath::SpringCharacterUpdate(CurrentPosition, SpringState.CurrentVelocity, SpringState.CurrentAccel, DesiredVelocity, VelocitySmoothingTime, DeltaSeconds, VelocityDeadzoneThreshold);
+	
+	OutProposedMove.LinearVelocity = (CurrentPosition - PrevPosition) / DeltaSeconds;
+
+	// Angular //
+	FVector DesiredFacingDir = CharacterInputs->GetOrientationIntentDir_WorldSpace();
+	DesiredFacingDir -= DesiredFacingDir.ProjectOnTo(GetMoverComponent()->GetUpDirection());
+	if (DesiredFacingDir.IsNearlyZero())
+	{
+		OutProposedMove.AngularVelocity = FRotator::ZeroRotator;
+		return;
+	}
+
+	DesiredFacingDir.Normalize();
+	
+	FQuat CurrentFacing = StartingSyncState->GetOrientation_WorldSpace().Quaternion();
+	FQuat PrevFacing = CurrentFacing;
+	
+	FQuat DesiredFacing = FQuat::FindBetween(FVector::ForwardVector, DesiredFacingDir);
+	AnimationMath::SimpleSpringDamperQuat(CurrentFacing, SpringState.CurrentAngularVelocity, DesiredFacing, FacingSmoothingTime, DeltaSeconds);
+
+	// Calc angular velocity
+	FQuat Diff = (PrevFacing.Inverse() * CurrentFacing);
+	Diff.Normalize();
+	Diff.EnforceShortestArcWith(FQuat::Identity);
+
+	FRotator AngularVelocity = Diff.Rotator();
+	AngularVelocity *= (1.0f / DeltaSeconds);
+	OutProposedMove.AngularVelocity = AngularVelocity;
+}
+
+
+void USpringWalkingMode::SimulationTick_Implementation(const FSimulationTickParams& Params, FMoverTickEndData& OutputState)
+{
+	Super::SimulationTick_Implementation(Params, OutputState);
+
+	// We've already updated the spring state during GenerateMove, and just need to copy it into the output simulation state
+	if (const FMoverSpringState* InSpringState = Params.StartState.SyncState.SyncStateCollection.FindDataByType<FMoverSpringState>())
+	{
+		FMoverSpringState& OutputSpringState = OutputState.SyncState.SyncStateCollection.FindOrAddMutableDataByType<FMoverSpringState>();
+		OutputSpringState = *InSpringState;
+	}
+}
+
diff --git a/Engine/Plugins/Experimental/Mover/Source/Mover/Private/DefaultMovementSet/Modes/SpringWalkingState.h b/Engine/Plugins/Experimental/Mover/Source/Mover/Private/DefaultMovementSet/Modes/SpringWalkingState.h
new file mode 100644
--- /dev/null
+++ b/Engine/Plugins/Experimental/Mover/Source/Mover/Private/DefaultMovementSet/Modes/SpringWalkingState.h
@@ -0,0 +1,30 @@
+// Copyright Epic Games, Inc. All Rights Reserved.
+
+#include "CoreMinimal.h"
+#include "MovementMode.h"
+#include "MoverTypes.h"
+#include "SpringWalkingState.generated.h"
+
+/**
+* Internal state data for the SpringWalkingState
+*/
+USTRUCT()
+struct FMoverSpringState : public FMoverDataStructBase
+{
+	GENERATED_BODY()
+	
+	virtual UScriptStruct* GetScriptStruct() const override { return StaticStruct(); }
+	virtual FMoverDataStructBase* Clone() const override { return new FMoverSpringState(*this); }
+
+	// Velocity of internal spring model. Not necessarily the velocity of the character
+	UPROPERTY(BlueprintReadOnly, Category = "Mover|Experimental")
+	FVector CurrentVelocity = FVector::ZeroVector;
+
+	// Acceleration of internal spring model
+	UPROPERTY(BlueprintReadOnly, Category = "Mover|Experimental")
+	FVector CurrentAccel = FVector::ZeroVector;
+
+	// Angular velocity of internal spring model in radians per second
+	UPROPERTY(BlueprintReadOnly, Category = "Mover|Experimental")
+	FVector CurrentAngularVelocity = FVector::ZeroVector;
+};
diff --git a/Engine/Plugins/Experimental/Mover/Source/Mover/Public/DefaultMovementSet/Modes/SpringWalkingMode.h b/Engine/Plugins/Experimental/Mover/Source/Mover/Public/DefaultMovementSet/Modes/SpringWalkingMode.h
new file mode 100644
--- /dev/null
+++ b/Engine/Plugins/Experimental/Mover/Source/Mover/Public/DefaultMovementSet/Modes/SpringWalkingMode.h
@@ -0,0 +1,36 @@
+// Copyright Epic Games, Inc. All Rights Reserved.
+
+#pragma once
+
+#include "CoreMinimal.h"
+#include "DefaultMovementSet/Modes/WalkingMode.h"
+#include "SpringWalkingMode.generated.h"
+
+#define UE_API MOVER_API
+
+/**
+ * A walking mode that uses a critically damped spring for translation and rotation.
+ * The strength of the critically damped spring is set via smoothing times (separate for translation and rotation)
+ */
+UCLASS(BlueprintType)
+class USpringWalkingMode : public UWalkingMode
+{
+	GENERATED_BODY()
+
+public:
+	UE_API virtual void GenerateMove_Implementation(const FMoverTickStartData& StartState, const FMoverTimeStep& TimeStep, FProposedMove& OutProposedMove) const override;
+	UE_API virtual void SimulationTick_Implementation(const FSimulationTickParams& Params, FMoverTickEndData& OutputState) override;
+
+protected:
+	UPROPERTY(EditAnywhere, BlueprintReadWrite, Category = "Mover|Spring Settings")
+	float VelocitySmoothingTime = 0.2f;
+
+	UPROPERTY(EditAnywhere, BlueprintReadWrite, Category = "Mover|Spring Settings")
+	float FacingSmoothingTime = 0.2f;
+
+	UPROPERTY(EditAnywhere, BlueprintReadWrite, Category = "Mover|Spring Settings")
+	// Below this speed we set velocity to 0
+	float VelocityDeadzoneThreshold = 0.1f;
+};
+
+#undef UE_API
diff --git a/Engine/Source/Runtime/Core/Public/Math/UnrealMathUtility.h b/Engine/Source/Runtime/Core/Public/Math/UnrealMathUtility.h
--- a/Engine/Source/Runtime/Core/Public/Math/UnrealMathUtility.h
+++ b/Engine/Source/Runtime/Core/Public/Math/UnrealMathUtility.h
@@ -157,6 +157,8 @@ class TRange;
 #define UE_DOUBLE_TWO_PI		(6.28318530717958647692)
 #define UE_DOUBLE_PI_SQUARED	(9.86960440108935861883)
 
+#define UE_LN2		(0.69314718056f)
+
 // Common square roots
 #define UE_SQRT_2		(1.4142135623730950488016887242097f)
 #define UE_SQRT_3		(1.7320508075688772935274463415059f)
diff --git a/Engine/Source/Runtime/Engine/Internal/Animation/AnimationMath.h b/Engine/Source/Runtime/Engine/Internal/Animation/AnimationMath.h
--- a/Engine/Source/Runtime/Engine/Internal/Animation/AnimationMath.h
+++ b/Engine/Source/Runtime/Engine/Internal/Animation/AnimationMath.h
@@ -12,7 +12,32 @@ namespace AnimationMath
 
 	FORCEINLINE float SpringHalfLifeToDamping(float HalfLife, float eps = 1e-5f)
 	{
-		return (4.0f * 0.69314718056f) / (HalfLife + eps);
+		return (4.0f * UE_LN2) / (HalfLife + eps);
+	}
+
+	FORCEINLINE float SpringDampingToHalfLife(float Damping)
+	{
+		return Damping / (4.0f * UE_LN2);
+	}
+
+	FORCEINLINE float SpringSmoothingTimeToHalfLife(float SmoothingTime)
+	{
+		return SmoothingTime / UE_LN2;
+	}
+
+	FORCEINLINE float SpringHalfLifeToSmoothingTime(float HalfLife)
+	{
+		return HalfLife * UE_LN2;
+	}
+
+	FORCEINLINE float SpringSmoothingTimeToDamping(float SmoothTime, float eps = 1e-5f)
+	{
+		return 4.0f / (SmoothTime + eps);
+	}
+
+	FORCEINLINE float SpringSmoothingTimeToStrength(float SmoothingTime)
+	{
+		return 2.0f / SmoothingTime;
 	}
 
 	FORCEINLINE FVector QuatToScaledAngleAxis(const FQuat& Q)
@@ -30,14 +55,14 @@ namespace AnimationMath
 
 	// Simplified version of FMath::CriticallyDampedSmoothing where v_goal is assumed to be 0
 	template<typename T>
-	FORCEINLINE void SimpleSpringDamper(
+	void SimpleSpringDamper(
 		T& x,
 		T& v,
 		T x_goal,
-		float halflife,
+		float SmoothingTime,
 		float dt)
 	{
-		float y = SpringHalfLifeToDamping(halflife) / 2.0f;
+		float y = SpringSmoothingTimeToDamping(SmoothingTime) / 2.0f;
 		T j0 = x - x_goal;
 		T j1 = v + j0 * y;
 		float eydt = FMath::InvExpApprox(y * dt);
@@ -47,32 +72,64 @@ namespace AnimationMath
 	}
 
 	// Specialized quaternion version, uses FVector for velocity
-	FORCEINLINE void SimpleSpringDamperQuat(
+	inline void SimpleSpringDamperQuat(
 		FQuat& InOutRotation,
-		FVector& InOutAngularVelocity,
+		FVector& InOutAngularVelocityRadians,
 		const FQuat& TargetRotation,
-		float HalfLife,
+		float SmoothingTime,
 		float DeltaTime)
 	{
-		float y = SpringHalfLifeToDamping(HalfLife) / 2.0f;
+		float y = SpringSmoothingTimeToDamping(SmoothingTime) / 2.0f;
 
 		FQuat diff = InOutRotation * TargetRotation.Inverse();
 		diff.EnforceShortestArcWith(FQuat::Identity);
 		FVector j0 = QuatToScaledAngleAxis(diff);
-		FVector j1 = InOutAngularVelocity + j0 * y;
+		FVector j1 = InOutAngularVelocityRadians + j0 * y;
 
 		float eydt = FMath::InvExpApprox(y * DeltaTime);
 
 		InOutRotation = QuatFromScaledAngleAxis(eydt * (j0 + j1 * DeltaTime)) * TargetRotation;
-		InOutAngularVelocity = eydt * (InOutAngularVelocity - j1 * y * DeltaTime);
+		InOutAngularVelocityRadians = eydt * (InOutAngularVelocityRadians - j1 * y * DeltaTime);
+	}
+
+	inline void SpringCharacterUpdate(
+		FVector& InOutPosition,
+		FVector& InOutVelocity,
+		FVector& InOutAcceleration,
+		const FVector& TargetVelocity,
+		float SmoothingTime,
+		float DeltaTime,
+		float VDeadzone = 1e-2,
+		float ADeadzone = 1e-4)
+	{
+		float y = SpringSmoothingTimeToDamping(SmoothingTime) / 2.0f;
+		FVector j0 = InOutVelocity - TargetVelocity;
+		FVector j1 = InOutAcceleration + j0 * y;
+		float eydt = FMath::InvExpApprox(y * DeltaTime);
+
+		InOutPosition = eydt * (((-j1) / (y * y)) + ((-j0 - j1 * DeltaTime) / y)) +
+			(j1 / (y * y)) + j0 / y + TargetVelocity * DeltaTime + InOutPosition;
+		InOutVelocity = eydt * (j0 + j1 * DeltaTime) + TargetVelocity;
+		InOutAcceleration = eydt * (InOutAcceleration - j1 * y * DeltaTime);
+
+		if ((TargetVelocity - InOutVelocity).SquaredLength() < FMath::Square(VDeadzone))
+		{
+			// We reached our target
+			InOutVelocity = TargetVelocity;
+
+			if (InOutAcceleration.SquaredLength() < FMath::Square(ADeadzone))
+			{
+				InOutAcceleration = FVector::ZeroVector;
+			}
+		}
 	}
 
 	// Specialized quaternion damper, similar to FMath::ExponentialSmoothingApprox but for quats
-	FORCEINLINE FQuat Damper(const FQuat& CurrentRotation, const FQuat& TargetRotation, float HalfLife, float DeltaTime, float Eps = 1e-5f)
+	FORCEINLINE FQuat Damper(const FQuat& CurrentRotation, const FQuat& TargetRotation, float SmoothingTime, float DeltaTime, float Eps = 1e-5f)
 	{
-		if (HalfLife < Eps)
+		if (SmoothingTime < Eps)
 			return TargetRotation;
 
-		return FQuat::Slerp(CurrentRotation, TargetRotation, 1.0f - FMath::InvExpApprox((0.69314718056f * DeltaTime) / (HalfLife)));
+		return FQuat::Slerp(CurrentRotation, TargetRotation, 1.0f - FMath::InvExpApprox((DeltaTime) / (SmoothingTime)));
 	}
 }
\ No newline at end of file
