Mover: Fixing issue where pawns wouldn't follow a navlink off the navmesh if bSlideAlongNavMeshEdge was true

diff --git a/Engine/Plugins/Experimental/Mover/Source/Mover/Private/DefaultMovementSet/Modes/AsyncNavWalkingMode.cpp b/Engine/Plugins/Experimental/Mover/Source/Mover/Private/DefaultMovementSet/Modes/AsyncNavWalkingMode.cpp
--- a/Engine/Plugins/Experimental/Mover/Source/Mover/Private/DefaultMovementSet/Modes/AsyncNavWalkingMode.cpp
+++ b/Engine/Plugins/Experimental/Mover/Source/Mover/Private/DefaultMovementSet/Modes/AsyncNavWalkingMode.cpp
@@ -16,6 +16,7 @@
 #include "MoveLibrary/MovementUtils.h"
 #include "DrawDebugHelpers.h"
 #include "NavigationSystem.h"
+#include "AI/Navigation/PathFollowingAgentInterface.h"
 #include "MoveLibrary/NavMovementUtils.h"
 #include "NavMesh/RecastNavMesh.h"
 #include "VisualLogger/VisualLogger.h"
@@ -249,23 +250,33 @@ void UAsyncNavWalkingMode::SimulationTick_Implementation(const FSimulationTickPa
 		}
 		
 		// Find the point on the NavMesh
-		bool bHasNavigationData = false;
+		bool bFoundPointOnNavMesh = false;
 
 		if (NavDataInterface.IsValid())
 		{
-			if (bSlideAlongNavMeshEdge && CachedNavLocation.HasNodeRef())
+			const IPathFollowingAgentInterface* PathFollowingAgent = NavMoverComponent->GetPathFollowingAgent();
+			const bool bIsOnNavLink = PathFollowingAgent && PathFollowingAgent->IsFollowingNavLink();
+			
+			if (bSlideAlongNavMeshEdge && !bIsOnNavLink)
 			{
+				FNavLocation StartingNavFloorLocation;
 				bool bHasValidCachedNavLocation = NavDataInterface->IsNodeRefValid(CachedNavLocation.NodeRef);
+
+				// If we don't have a valid CachedNavLocation lets try finding the NavFloor where we're currently at and use that otherwise we can just use our CachedNavLocation
 				if (!bHasValidCachedNavLocation)
 				{
-					bHasValidCachedNavLocation = FindNavFloor(TargetFeetLocation, OUT CachedNavLocation, NavDataInterface.Get());
+					bHasValidCachedNavLocation = FindNavFloor(StartingFeetLocation, OUT StartingNavFloorLocation, NavDataInterface.Get());
+				}
+				else
+				{
+					StartingNavFloorLocation = CachedNavLocation;
 				}
 
 				if (bHasValidCachedNavLocation)
 				{
-					bHasNavigationData = NavDataInterface->FindMoveAlongSurface(CachedNavLocation, TargetFeetLocation, OUT DestNavLocation);
+					bFoundPointOnNavMesh = NavDataInterface->FindMoveAlongSurface(StartingNavFloorLocation, TargetFeetLocation, OUT DestNavLocation);
 
-					if (bHasNavigationData)
+					if (bFoundPointOnNavMesh)
 					{
 						TargetFeetLocation = UMovementUtils::ProjectToGravityFloor(DestNavLocation.Location, UpDirection) + UMovementUtils::GetGravityVerticalComponent(TargetFeetLocation, UpDirection);
 					}
@@ -273,11 +284,11 @@ void UAsyncNavWalkingMode::SimulationTick_Implementation(const FSimulationTickPa
 			}
 			else
 			{
-				bHasNavigationData = FindNavFloor(TargetFeetLocation, OUT DestNavLocation, NavDataInterface.Get());
+				bFoundPointOnNavMesh = FindNavFloor(TargetFeetLocation, OUT DestNavLocation, NavDataInterface.Get());
 			}
 		}
 
-		if (!bHasNavigationData)
+		if (!bFoundPointOnNavMesh)
 		{
 			// Can't find nav mesh at this location, so we need to do something else
 			switch (BehaviorOffNavMesh)
diff --git a/Engine/Plugins/Experimental/Mover/Source/Mover/Private/DefaultMovementSet/Modes/NavWalkingMode.cpp b/Engine/Plugins/Experimental/Mover/Source/Mover/Private/DefaultMovementSet/Modes/NavWalkingMode.cpp
--- a/Engine/Plugins/Experimental/Mover/Source/Mover/Private/DefaultMovementSet/Modes/NavWalkingMode.cpp
+++ b/Engine/Plugins/Experimental/Mover/Source/Mover/Private/DefaultMovementSet/Modes/NavWalkingMode.cpp
@@ -6,6 +6,7 @@
 #include "NavigationSystem.h"
 #include "AI/NavigationSystemBase.h"
 #include "AI/Navigation/NavigationDataInterface.h"
+#include "AI/Navigation/PathFollowingAgentInterface.h"
 #include "Components/ShapeComponent.h"
 #include "DefaultMovementSet/NavMoverComponent.h"
 #include "DefaultMovementSet/Settings/CommonLegacyMovementSettings.h"
@@ -222,23 +223,33 @@ void UNavWalkingMode::SimulationTick_Implementation(const FSimulationTickParams&
 		}
 		
 		// Find the point on the NavMesh
-		bool bHasNavigationData = false;
+		bool bFoundPointOnNavMesh = false;
 
 		if (NavDataInterface.IsValid())
 		{
-			if (bSlideAlongNavMeshEdge && CachedNavLocation.HasNodeRef())
+			const IPathFollowingAgentInterface* PathFollowingAgent = NavMoverComponent->GetPathFollowingAgent();
+			const bool bIsOnNavLink = PathFollowingAgent && PathFollowingAgent->IsFollowingNavLink();
+			
+			if (bSlideAlongNavMeshEdge && !bIsOnNavLink)
 			{
+				FNavLocation StartingNavFloorLocation;
 				bool bHasValidCachedNavLocation = NavDataInterface->IsNodeRefValid(CachedNavLocation.NodeRef);
+
+				// If we don't have a valid CachedNavLocation lets try finding the NavFloor where we're currently at and use that otherwise we can just use our CachedNavLocation
 				if (!bHasValidCachedNavLocation)
 				{
-					bHasValidCachedNavLocation = FindNavFloor(AdjustedDest, OUT CachedNavLocation, NavDataInterface.Get());
+					bHasValidCachedNavLocation = FindNavFloor(StartingLocation, OUT StartingNavFloorLocation, NavDataInterface.Get());
+				}
+				else
+				{
+					StartingNavFloorLocation = CachedNavLocation;
 				}
 
 				if (bHasValidCachedNavLocation)
 				{
-					bHasNavigationData = NavDataInterface->FindMoveAlongSurface(CachedNavLocation, AdjustedDest, OUT DestNavLocation);
+					bFoundPointOnNavMesh = NavDataInterface->FindMoveAlongSurface(StartingNavFloorLocation, AdjustedDest, OUT DestNavLocation);
 
-					if (bHasNavigationData)
+					if (bFoundPointOnNavMesh)
 					{
 						AdjustedDest = UMovementUtils::ProjectToGravityFloor(DestNavLocation.Location, UpDirection) + UMovementUtils::GetGravityVerticalComponent(AdjustedDest, UpDirection);
 					}
@@ -246,12 +257,12 @@ void UNavWalkingMode::SimulationTick_Implementation(const FSimulationTickParams&
 			}
 			else
 			{
-				bHasNavigationData = FindNavFloor(AdjustedDest, DestNavLocation, NavDataInterface.Get());
+				bFoundPointOnNavMesh = FindNavFloor(AdjustedDest, DestNavLocation, NavDataInterface.Get());
 			}
 		}
 
 
-		if (!bHasNavigationData)
+		if (!bFoundPointOnNavMesh)
 		{
 			// Can't find nav mesh at this location, so we need to do something else
 			switch (BehaviorOffNavMesh)
diff --git a/Engine/Plugins/Experimental/Mover/Source/Mover/Private/DefaultMovementSet/NavMoverComponent.cpp b/Engine/Plugins/Experimental/Mover/Source/Mover/Private/DefaultMovementSet/NavMoverComponent.cpp
--- a/Engine/Plugins/Experimental/Mover/Source/Mover/Private/DefaultMovementSet/NavMoverComponent.cpp
+++ b/Engine/Plugins/Experimental/Mover/Source/Mover/Private/DefaultMovementSet/NavMoverComponent.cpp
@@ -231,17 +231,17 @@ bool UNavMoverComponent::CanStopPathFollowing() const
 
 void UNavMoverComponent::SetPathFollowingAgent(IPathFollowingAgentInterface* InPathFollowingAgent)
 {
-	PathFollowingComp = Cast<UObject>(InPathFollowingAgent);
+	PathFollowingComp = InPathFollowingAgent;
 }
 
 IPathFollowingAgentInterface* UNavMoverComponent::GetPathFollowingAgent()
 {
-	return Cast<IPathFollowingAgentInterface>(PathFollowingComp);
+	return PathFollowingComp.Get();
 }
 
 const IPathFollowingAgentInterface* UNavMoverComponent::GetPathFollowingAgent() const
 {
-	return Cast<const IPathFollowingAgentInterface>(PathFollowingComp);
+	return PathFollowingComp.Get();
 }
 
 const FNavAgentProperties& UNavMoverComponent::GetNavAgentPropertiesRef() const
diff --git a/Engine/Plugins/Experimental/Mover/Source/Mover/Public/DefaultMovementSet/NavMoverComponent.h b/Engine/Plugins/Experimental/Mover/Source/Mover/Public/DefaultMovementSet/NavMoverComponent.h
--- a/Engine/Plugins/Experimental/Mover/Source/Mover/Public/DefaultMovementSet/NavMoverComponent.h
+++ b/Engine/Plugins/Experimental/Mover/Source/Mover/Public/DefaultMovementSet/NavMoverComponent.h
@@ -7,6 +7,7 @@
 #include "Components/ActorComponent.h"
 #include "MoverComponent.h"
 #include "GameFramework/NavMovementInterface.h"
+#include "UObject/WeakInterfacePtr.h"
 
 #include "NavMoverComponent.generated.h"
 
@@ -53,8 +54,7 @@ protected:
 private:
 	/** object implementing IPathFollowingAgentInterface. Is private to control access to it.
 	 *	@see SetPathFollowingAgent, GetPathFollowingAgent */
-	UPROPERTY()
-	TObjectPtr<UObject> PathFollowingComp;
+	TWeakInterfacePtr<IPathFollowingAgentInterface> PathFollowingComp;
 
 	/** Associated Movement component that will actually move the actor */ 
 	UPROPERTY()
