Mover: fix layered move timestamp conversions, including fix for broken multi-jump due to rounding

#rnx

diff --git a/Engine/Plugins/Experimental/Mover/Source/Mover/Private/DefaultMovementSet/LayeredMoves/MultiJumpLayeredMove.cpp b/Engine/Plugins/Experimental/Mover/Source/Mover/Private/DefaultMovementSet/LayeredMoves/MultiJumpLayeredMove.cpp
--- a/Engine/Plugins/Experimental/Mover/Source/Mover/Private/DefaultMovementSet/LayeredMoves/MultiJumpLayeredMove.cpp
+++ b/Engine/Plugins/Experimental/Mover/Source/Mover/Private/DefaultMovementSet/LayeredMoves/MultiJumpLayeredMove.cpp
@@ -39,7 +39,7 @@ bool FLayeredMove_MultiJump::GenerateMove(const FMoverTickStartData& StartState,
 	FFloorCheckResult FloorHitResult;
 	bool bValidBlackboard = SimBlackboard->TryGet(CommonBlackboard::LastFloorResult, OUT FloorHitResult);
 
-	if (StartSimTimeMs == TimeStep.BaseSimTimeMs)
+	if (FMath::IsNearlyEqual(StartSimTimeMs, TimeStep.BaseSimTimeMs))
 	{
 		JumpsInAirRemaining = MaximumInAirJumps;
 	}
diff --git a/Engine/Plugins/Experimental/Mover/Source/Mover/Private/LayeredMove.cpp b/Engine/Plugins/Experimental/Mover/Source/Mover/Private/LayeredMove.cpp
--- a/Engine/Plugins/Experimental/Mover/Source/Mover/Private/LayeredMove.cpp
+++ b/Engine/Plugins/Experimental/Mover/Source/Mover/Private/LayeredMove.cpp
@@ -8,7 +8,7 @@
 
 #include UE_INLINE_GENERATED_CPP_BY_NAME(LayeredMove)
 
-const float LayeredMove_InvalidTime = -UE_BIG_NUMBER;
+const double LayeredMove_InvalidTime = -UE_BIG_NUMBER;
 
 void FLayeredMoveFinishVelocitySettings::NetSerialize(FArchive& Ar)
 {
@@ -41,13 +41,13 @@ FLayeredMoveBase::FLayeredMoveBase() :
 }
 
 
-void FLayeredMoveBase::StartMove(const UMoverComponent* MoverComp, UMoverBlackboard* SimBlackboard, float CurrentSimTimeMs)
+void FLayeredMoveBase::StartMove(const UMoverComponent* MoverComp, UMoverBlackboard* SimBlackboard, double CurrentSimTimeMs)
 {
 	StartSimTimeMs = CurrentSimTimeMs;
 	OnStart(MoverComp, SimBlackboard);
 }
 
-void FLayeredMoveBase::StartMove_Async(UMoverBlackboard* SimBlackboard, float CurrentSimTimeMs)
+void FLayeredMoveBase::StartMove_Async(UMoverBlackboard* SimBlackboard, double CurrentSimTimeMs)
 {
 	StartSimTimeMs = CurrentSimTimeMs;
 	OnStart_Async(SimBlackboard);
@@ -59,21 +59,21 @@ bool FLayeredMoveBase::GenerateMove_Async(const FMoverTickStartData& StartState,
 	return false;
 }
 
-bool FLayeredMoveBase::IsFinished(float CurrentSimTimeMs) const
+bool FLayeredMoveBase::IsFinished(double CurrentSimTimeMs) const
 {
-	const bool bHasStarted = (StartSimTimeMs >= 0.f);
+	const bool bHasStarted = (StartSimTimeMs >= 0.0);
 	const bool bTimeExpired = bHasStarted && (DurationMs > 0.f) && (StartSimTimeMs + DurationMs <= CurrentSimTimeMs);
 	const bool bDidTickOnceAndExpire = bHasStarted && (DurationMs == 0.f);
 
 	return bTimeExpired || bDidTickOnceAndExpire;
 }
 
-void FLayeredMoveBase::EndMove(const UMoverComponent* MoverComp, UMoverBlackboard* SimBlackboard, float CurrentSimTimeMs)
+void FLayeredMoveBase::EndMove(const UMoverComponent* MoverComp, UMoverBlackboard* SimBlackboard, double CurrentSimTimeMs)
 {
 	OnEnd(MoverComp, SimBlackboard, CurrentSimTimeMs);
 }
 
-void FLayeredMoveBase::EndMove_Async(UMoverBlackboard* SimBlackboard, float CurrentSimTimeMs)
+void FLayeredMoveBase::EndMove_Async(UMoverBlackboard* SimBlackboard, double CurrentSimTimeMs)
 {
 	OnEnd_Async(SimBlackboard, CurrentSimTimeMs);
 }
diff --git a/Engine/Plugins/Experimental/Mover/Source/Mover/Private/MoveLibrary/MovementMixer.cpp b/Engine/Plugins/Experimental/Mover/Source/Mover/Private/MoveLibrary/MovementMixer.cpp
--- a/Engine/Plugins/Experimental/Mover/Source/Mover/Private/MoveLibrary/MovementMixer.cpp
+++ b/Engine/Plugins/Experimental/Mover/Source/Mover/Private/MoveLibrary/MovementMixer.cpp
@@ -136,10 +136,10 @@ void UMovementMixer::MixProposedMoves(const FProposedMove& MoveToMix, FVector Up
 void UMovementMixer::ResetMixerState()
 {
 	CurrentHighestPriority = 0;
-	CurrentLayeredMoveStartTime = TNumericLimits<float>::Max();
+	CurrentLayeredMoveStartTime = TNumericLimits<double>::Max();
 }
 
-bool UMovementMixer::CheckPriority(const FLayeredMoveBase* LayeredMove, uint8& InOutHighestPriority, float& InOutCurrentLayeredMoveStartTimeMs)
+bool UMovementMixer::CheckPriority(const FLayeredMoveBase* LayeredMove, uint8& InOutHighestPriority, double& InOutCurrentLayeredMoveStartTimeMs)
 {
 	if (LayeredMove->Priority > InOutHighestPriority)
 	{
diff --git a/Engine/Plugins/Experimental/Mover/Source/Mover/Public/LayeredMove.h b/Engine/Plugins/Experimental/Mover/Source/Mover/Public/LayeredMove.h
--- a/Engine/Plugins/Experimental/Mover/Source/Mover/Public/LayeredMove.h
+++ b/Engine/Plugins/Experimental/Mover/Source/Mover/Public/LayeredMove.h
@@ -113,9 +113,9 @@ struct FLayeredMoveBase
 	}
 	
 	// Kicks off this move, allowing any initialization to occur.
-	UE_API void StartMove(const UMoverComponent* MoverComp, UMoverBlackboard* SimBlackboard, float CurrentSimTimeMs);
+	UE_API void StartMove(const UMoverComponent* MoverComp, UMoverBlackboard* SimBlackboard, double CurrentSimTimeMs);
 	// Async version of the above
-	UE_API void StartMove_Async(UMoverBlackboard* SimBlackboard, float CurrentSimTimeMs);
+	UE_API void StartMove_Async(UMoverBlackboard* SimBlackboard, double CurrentSimTimeMs);
 
 	// Called when this layered move starts.
 	virtual void OnStart(const UMoverComponent* MoverComp, UMoverBlackboard* SimBlackboard) {}
@@ -129,17 +129,17 @@ struct FLayeredMoveBase
 	UE_API virtual bool GenerateMove_Async(const FMoverTickStartData& StartState, const FMoverTimeStep& TimeStep, UMoverBlackboard* SimBlackboard, FProposedMove& OutProposedMove);
 
 	// Runtime query whether this move is finished and can be destroyed. The default implementation is based on DurationMs.
-	UE_API virtual bool IsFinished(float CurrentSimTimeMs) const;
+	UE_API virtual bool IsFinished(double CurrentSimTimeMs) const;
 
 	// Ends this move, allowing any cleanup to occur.
-	UE_API void EndMove(const UMoverComponent* MoverComp, UMoverBlackboard* SimBlackboard, float CurrentSimTimeMs);
+	UE_API void EndMove(const UMoverComponent* MoverComp, UMoverBlackboard* SimBlackboard, double CurrentSimTimeMs);
 	// Async version of the above
-	UE_API void EndMove_Async(UMoverBlackboard* SimBlackboard, float CurrentSimTimeMs);
+	UE_API void EndMove_Async(UMoverBlackboard* SimBlackboard, double CurrentSimTimeMs);
 	
 	// Called when this layered move ends.
-	virtual void OnEnd(const UMoverComponent* MoverComp, UMoverBlackboard* SimBlackboard, float CurrentSimTimeMs) {}
+	virtual void OnEnd(const UMoverComponent* MoverComp, UMoverBlackboard* SimBlackboard, double CurrentSimTimeMs) {}
 	// Async version of the above
-	virtual void OnEnd_Async(UMoverBlackboard* SimBlackboard, float CurrentSimTimeMs) {}
+	virtual void OnEnd_Async(UMoverBlackboard* SimBlackboard, double CurrentSimTimeMs) {}
 
 	// @return newly allocated copy of this FLayeredMoveBase. Must be overridden by child classes
 	UE_API virtual FLayeredMoveBase* Clone() const;
diff --git a/Engine/Plugins/Experimental/Mover/Source/Mover/Public/MoveLibrary/MovementMixer.h b/Engine/Plugins/Experimental/Mover/Source/Mover/Public/MoveLibrary/MovementMixer.h
--- a/Engine/Plugins/Experimental/Mover/Source/Mover/Public/MoveLibrary/MovementMixer.h
+++ b/Engine/Plugins/Experimental/Mover/Source/Mover/Public/MoveLibrary/MovementMixer.h
@@ -35,13 +35,13 @@ protected:
 	uint8 CurrentHighestPriority;
 
 	// Earliest start time of the layered move with highest priority. Used to help break ties of moves with same priority. Note: Currently only used for mixing layered moves
-	float CurrentLayeredMoveStartTime;
+	double CurrentLayeredMoveStartTime;
 	
 	/**
 	 * Helper function for layered move mixing to check priority and start time if priority is the same.
 	 * Returns true if this layered move should take priority given current HighestPriority and CurrentLayeredMoveStartTimeMs
 	 */
-	static UE_API bool CheckPriority(const FLayeredMoveBase* LayeredMove, uint8& InOutHighestPriority, float& InOutCurrentLayeredMoveStartTimeMs);
+	static UE_API bool CheckPriority(const FLayeredMoveBase* LayeredMove, uint8& InOutHighestPriority, double& InOutCurrentLayeredMoveStartTimeMs);
 };
 
 #undef UE_API
diff --git a/Engine/Plugins/Experimental/MoverTests/Source/MoverTests/Private/TestCustomLayeredMoves.cpp b/Engine/Plugins/Experimental/MoverTests/Source/MoverTests/Private/TestCustomLayeredMoves.cpp
--- a/Engine/Plugins/Experimental/MoverTests/Source/MoverTests/Private/TestCustomLayeredMoves.cpp
+++ b/Engine/Plugins/Experimental/MoverTests/Source/MoverTests/Private/TestCustomLayeredMoves.cpp
@@ -24,7 +24,7 @@ void FTestCustomLayeredMove::OnStart(const UMoverComponent* MoverComp, UMoverBla
 	UE_LOG(LogMover, Log, TEXT("Custom test layered move started!"));
 }
 
-void FTestCustomLayeredMove::OnEnd(const UMoverComponent* MoverComp, UMoverBlackboard* SimBlackboard, float CurrentSimTimeMs)
+void FTestCustomLayeredMove::OnEnd(const UMoverComponent* MoverComp, UMoverBlackboard* SimBlackboard, double CurrentSimTimeMs)
 {
 	UE_LOG(LogMover, Log, TEXT("Custom test layered move ended!"));
 }
diff --git a/Engine/Plugins/Experimental/MoverTests/Source/MoverTests/Public/TestCustomLayeredMoves.h b/Engine/Plugins/Experimental/MoverTests/Source/MoverTests/Public/TestCustomLayeredMoves.h
--- a/Engine/Plugins/Experimental/MoverTests/Source/MoverTests/Public/TestCustomLayeredMoves.h
+++ b/Engine/Plugins/Experimental/MoverTests/Source/MoverTests/Public/TestCustomLayeredMoves.h
@@ -26,7 +26,7 @@ struct MOVERTESTS_API FTestCustomLayeredMove : public FLayeredMoveBase
 
 	virtual void OnStart(const UMoverComponent* MoverComp, UMoverBlackboard* SimBlackboard) override;
 	
-	virtual void OnEnd(const UMoverComponent* MoverComp, UMoverBlackboard* SimBlackboard, float CurrentSimTimeMs) override;
+	virtual void OnEnd(const UMoverComponent* MoverComp, UMoverBlackboard* SimBlackboard, double CurrentSimTimeMs) override;
 	
 	// Generate a movement 
 	virtual bool GenerateMove(const FMoverTickStartData& StartState, const FMoverTimeStep& TimeStep, const UMoverComponent* MoverComp, UMoverBlackboard* SimBlackboard, FProposedMove& OutProposedMove) override;
