[Mover] Adding MovementModifiers to ShouldReconcile in FMoverSyncState

diff --git a/Engine/Plugins/Experimental/Mover/Source/Mover/Private/MovementModifier.cpp b/Engine/Plugins/Experimental/Mover/Source/Mover/Private/MovementModifier.cpp
--- a/Engine/Plugins/Experimental/Mover/Source/Mover/Private/MovementModifier.cpp
+++ b/Engine/Plugins/Experimental/Mover/Source/Mover/Private/MovementModifier.cpp
@@ -535,3 +535,32 @@ void FMovementModifierGroup::Reset()
 	QueuedModifiers.Empty();
 	ActiveModifiers.Empty();
 }
+
+bool FMovementModifierGroup::ShouldReconcile(const FMovementModifierGroup& Other) const
+{
+	// Only compare the types of modifiers contained, not the state
+	if (ActiveModifiers.Num() != Other.ActiveModifiers.Num() ||
+		QueuedModifiers.Num() != Other.QueuedModifiers.Num())
+	{
+		return false;
+	}
+
+
+	for (int32 i = 0; i < ActiveModifiers.Num(); ++i)
+	{
+		if (ActiveModifiers[i]->GetScriptStruct() != Other.ActiveModifiers[i]->GetScriptStruct())
+		{
+			return false;
+		}
+	}
+
+	for (int32 i = 0; i < QueuedModifiers.Num(); ++i)
+	{
+		if (QueuedModifiers[i]->GetScriptStruct() != Other.QueuedModifiers[i]->GetScriptStruct())
+		{
+			return false;
+		}
+	}
+
+	return true;
+}
\ No newline at end of file
diff --git a/Engine/Plugins/Experimental/Mover/Source/Mover/Public/MovementModifier.h b/Engine/Plugins/Experimental/Mover/Source/Mover/Public/MovementModifier.h
--- a/Engine/Plugins/Experimental/Mover/Source/Mover/Public/MovementModifier.h
+++ b/Engine/Plugins/Experimental/Mover/Source/Mover/Public/MovementModifier.h
@@ -264,6 +264,9 @@ struct FMovementModifierGroup
 
 	// Clears active and queued modifiers
 	UE_API void Reset();
+
+	// Whether this modifier should reconcile with the authority state
+	UE_API bool ShouldReconcile(const FMovementModifierGroup& Other) const;
 	
 protected:
 	/** Movement modifiers that are currently active in this group */
diff --git a/Engine/Plugins/Experimental/Mover/Source/Mover/Public/MoverSimulationTypes.h b/Engine/Plugins/Experimental/Mover/Source/Mover/Public/MoverSimulationTypes.h
--- a/Engine/Plugins/Experimental/Mover/Source/Mover/Public/MoverSimulationTypes.h
+++ b/Engine/Plugins/Experimental/Mover/Source/Mover/Public/MoverSimulationTypes.h
@@ -171,7 +171,8 @@ public:
 	bool ShouldReconcile(const FMoverSyncState& AuthorityState) const
 	{
 		return (MovementMode != AuthorityState.MovementMode) || 
-			   SyncStateCollection.ShouldReconcile(AuthorityState.SyncStateCollection);
+			   SyncStateCollection.ShouldReconcile(AuthorityState.SyncStateCollection) ||
+			   MovementModifiers.ShouldReconcile(AuthorityState.MovementModifiers);
 	}
 
 	void Interpolate(const FMoverSyncState* From, const FMoverSyncState* To, float Pct)
