Mover: Add CommonLegacyMovementSettings shared settings in various movement modes to avoid a crash that could happen if they were the only mode added.

diff --git a/Engine/Plugins/Experimental/Mover/Source/Mover/Private/DefaultMovementSet/Modes/AsyncFlyingMode.cpp b/Engine/Plugins/Experimental/Mover/Source/Mover/Private/DefaultMovementSet/Modes/AsyncFlyingMode.cpp
--- a/Engine/Plugins/Experimental/Mover/Source/Mover/Private/DefaultMovementSet/Modes/AsyncFlyingMode.cpp
+++ b/Engine/Plugins/Experimental/Mover/Source/Mover/Private/DefaultMovementSet/Modes/AsyncFlyingMode.cpp
@@ -16,6 +16,8 @@
 UAsyncFlyingMode::UAsyncFlyingMode(const FObjectInitializer& ObjectInitializer)
 	: Super(ObjectInitializer)
 {
+	SharedSettingsClasses.Add(UCommonLegacyMovementSettings::StaticClass());
+	
 	GameplayTags.AddTag(Mover_IsInAir);
 	GameplayTags.AddTag(Mover_IsFlying);
 }
diff --git a/Engine/Plugins/Experimental/Mover/Source/Mover/Private/DefaultMovementSet/Modes/AsyncNavWalkingMode.cpp b/Engine/Plugins/Experimental/Mover/Source/Mover/Private/DefaultMovementSet/Modes/AsyncNavWalkingMode.cpp
--- a/Engine/Plugins/Experimental/Mover/Source/Mover/Private/DefaultMovementSet/Modes/AsyncNavWalkingMode.cpp
+++ b/Engine/Plugins/Experimental/Mover/Source/Mover/Private/DefaultMovementSet/Modes/AsyncNavWalkingMode.cpp
@@ -50,6 +50,8 @@ UAsyncNavWalkingMode::UAsyncNavWalkingMode()
 	, NavDataInterface(nullptr)
 	, bProjectNavMeshOnBothWorldChannels(true)
 {
+	SharedSettingsClasses.Add(UCommonLegacyMovementSettings::StaticClass());
+	
 	GameplayTags.AddTag(Mover_IsOnGround);
 	GameplayTags.AddTag(Mover_IsNavWalking);
 }
diff --git a/Engine/Plugins/Experimental/Mover/Source/Mover/Private/DefaultMovementSet/Modes/FlyingMode.cpp b/Engine/Plugins/Experimental/Mover/Source/Mover/Private/DefaultMovementSet/Modes/FlyingMode.cpp
--- a/Engine/Plugins/Experimental/Mover/Source/Mover/Private/DefaultMovementSet/Modes/FlyingMode.cpp
+++ b/Engine/Plugins/Experimental/Mover/Source/Mover/Private/DefaultMovementSet/Modes/FlyingMode.cpp
@@ -14,6 +14,8 @@
 UFlyingMode::UFlyingMode(const FObjectInitializer& ObjectInitializer)
 	: Super(ObjectInitializer)
 {
+	SharedSettingsClasses.Add(UCommonLegacyMovementSettings::StaticClass());
+	
 	GameplayTags.AddTag(Mover_IsInAir);
 	GameplayTags.AddTag(Mover_IsFlying);
 }
@@ -173,4 +175,4 @@ void UFlyingMode::OnUnregistered()
 	CommonLegacySettings = nullptr;
 
 	Super::OnUnregistered();
-}
\ No newline at end of file
+}
diff --git a/Engine/Plugins/Experimental/Mover/Source/Mover/Private/DefaultMovementSet/Modes/NavWalkingMode.cpp b/Engine/Plugins/Experimental/Mover/Source/Mover/Private/DefaultMovementSet/Modes/NavWalkingMode.cpp
--- a/Engine/Plugins/Experimental/Mover/Source/Mover/Private/DefaultMovementSet/Modes/NavWalkingMode.cpp
+++ b/Engine/Plugins/Experimental/Mover/Source/Mover/Private/DefaultMovementSet/Modes/NavWalkingMode.cpp
@@ -30,6 +30,8 @@ UNavWalkingMode::UNavWalkingMode()
 	, NavDataInterface(nullptr)
 	, bProjectNavMeshOnBothWorldChannels(true)
 {
+	SharedSettingsClasses.Add(UCommonLegacyMovementSettings::StaticClass());
+	
 	GameplayTags.AddTag(Mover_IsOnGround);
 	GameplayTags.AddTag(Mover_IsNavWalking);
 }
