Mover: fixed issue where backends could leave non-persistent structs in sync state between frames, potentially causing downstream issues with systems affected by stale state. e.g. Ziplining mode from the MoverExamples plugin

diff --git a/Engine/Plugins/Experimental/Mover/Source/Mover/Private/MoverComponent.cpp b/Engine/Plugins/Experimental/Mover/Source/Mover/Private/MoverComponent.cpp
--- a/Engine/Plugins/Experimental/Mover/Source/Mover/Private/MoverComponent.cpp
+++ b/Engine/Plugins/Experimental/Mover/Source/Mover/Private/MoverComponent.cpp
@@ -523,7 +523,7 @@ void UMoverComponent::SimulationTick(const FMoverTimeStep& InTimeStep, const FMo
 		}
 	}
 
-	// Sync state data should carry over between frames
+	// Some sync state data should carry over between frames
 	for (const FMoverDataPersistence& PersistentSyncEntry : PersistentSyncStateDataTypes)
 	{
 		bool bShouldAddDefaultData = true;
@@ -543,6 +543,29 @@ void UMoverComponent::SimulationTick(const FMoverTimeStep& InTimeStep, const FMo
 		}
 	}
 
+	// Make sure any other sync state structs that aren't supposed to be persistent are removed
+	const TArray<TSharedPtr<FMoverDataStructBase>>& AllSyncStructs = SimOutput.SyncState.SyncStateCollection.GetDataArray();
+	for (int32 i = AllSyncStructs.Num()-1; i >= 0; --i)
+	{
+		bool bShouldRemoveStructType = true;
+
+		const UScriptStruct* ScriptStruct = AllSyncStructs[i]->GetScriptStruct();
+
+		for (const FMoverDataPersistence& PersistentSyncEntry : PersistentSyncStateDataTypes)
+		{
+			if (PersistentSyncEntry.RequiredType == ScriptStruct)
+			{
+				bShouldRemoveStructType = false;
+				break;
+			}
+		}
+
+		if (bShouldRemoveStructType)
+		{
+			SimOutput.SyncState.SyncStateCollection.RemoveDataByType(ScriptStruct);
+		}	
+	}
+
 	SimOutput.AuxState = SimInput.AuxState;
 
 	FCharacterDefaultInputs* Input = SimInput.InputCmd.InputCollection.FindMutableDataByType<FCharacterDefaultInputs>();
