Mover: potentially reduces memory operations when adding to a data collection by preferring to copy onto an existing matching struct if possible, rather than allocating a new copy.

diff --git a/Engine/Plugins/Experimental/ChaosMover/Source/ChaosMover/Private/ChaosMoverStateMachine.cpp b/Engine/Plugins/Experimental/ChaosMover/Source/ChaosMover/Private/ChaosMoverStateMachine.cpp
--- a/Engine/Plugins/Experimental/ChaosMover/Source/ChaosMover/Private/ChaosMoverStateMachine.cpp
+++ b/Engine/Plugins/Experimental/ChaosMover/Source/ChaosMover/Private/ChaosMoverStateMachine.cpp
@@ -319,7 +319,7 @@ namespace UE::ChaosMover
 				{
 					if (SyncDataIt->Get())
 					{
-						WorkingSubstepStartData.SyncState.SyncStateCollection.AddOrOverwriteData(TSharedPtr<FMoverDataStructBase>(SyncDataIt->Get()->Clone()));
+						WorkingSubstepStartData.SyncState.SyncStateCollection.AddDataByCopy(SyncDataIt->Get());
 					}
 				}
 
diff --git a/Engine/Plugins/Experimental/Mover/Source/Mover/Private/MovementModeStateMachine.cpp b/Engine/Plugins/Experimental/Mover/Source/Mover/Private/MovementModeStateMachine.cpp
--- a/Engine/Plugins/Experimental/Mover/Source/Mover/Private/MovementModeStateMachine.cpp
+++ b/Engine/Plugins/Experimental/Mover/Source/Mover/Private/MovementModeStateMachine.cpp
@@ -215,7 +215,7 @@ void UMovementModeStateMachine::OnSimulationTick(USceneComponent* UpdatedCompone
 			{
 				if (SyncDataIt->Get())
 				{
-					WorkingSubstepStartData.SyncState.SyncStateCollection.AddOrOverwriteData(TSharedPtr<FMoverDataStructBase>(SyncDataIt->Get()->Clone()));
+					WorkingSubstepStartData.SyncState.SyncStateCollection.AddDataByCopy(SyncDataIt->Get());
 				}
 			}
 
diff --git a/Engine/Plugins/Experimental/Mover/Source/Mover/Private/MoverComponent.cpp b/Engine/Plugins/Experimental/Mover/Source/Mover/Private/MoverComponent.cpp
--- a/Engine/Plugins/Experimental/Mover/Source/Mover/Private/MoverComponent.cpp
+++ b/Engine/Plugins/Experimental/Mover/Source/Mover/Private/MoverComponent.cpp
@@ -532,7 +532,7 @@ void UMoverComponent::SimulationTick(const FMoverTimeStep& InTimeStep, const FMo
 		{
 			if (const FMoverDataStructBase* PriorFrameData = SimInput.SyncState.SyncStateCollection.FindDataByType(PersistentSyncEntry.RequiredType))
 			{
-				SimOutput.SyncState.SyncStateCollection.AddOrOverwriteData(TSharedPtr<FMoverDataStructBase>(PriorFrameData->Clone()));
+				SimOutput.SyncState.SyncStateCollection.AddDataByCopy(PriorFrameData);
 				bShouldAddDefaultData = false;
 			}
 		}
diff --git a/Engine/Plugins/Experimental/Mover/Source/Mover/Private/MoverTypes.cpp b/Engine/Plugins/Experimental/Mover/Source/Mover/Private/MoverTypes.cpp
--- a/Engine/Plugins/Experimental/Mover/Source/Mover/Private/MoverTypes.cpp
+++ b/Engine/Plugins/Experimental/Mover/Source/Mover/Private/MoverTypes.cpp
@@ -285,7 +285,7 @@ void FMoverDataCollection::Interpolate(const FMoverDataCollection& From, const F
 		else
 		{
 			// If only present in From, add the block directly to this collection
-			AddOrOverwriteData(FromElement);
+			AddDataByCopy(FromElement.Get());
 		}
 	}
 
@@ -294,7 +294,7 @@ void FMoverDataCollection::Interpolate(const FMoverDataCollection& From, const F
 	{
 		if (!InterpolatedTypes.Contains(ToElement->GetScriptStruct()))
 		{
-			AddOrOverwriteData(ToElement);
+			AddDataByCopy(ToElement.Get());
 		}
 	}
 }
@@ -310,7 +310,7 @@ void FMoverDataCollection::Merge(const FMoverDataCollection& From)
 		else
 		{
 			// If only present in the previous block, copy it into this block
-			AddOrOverwriteData(FromElement);
+			AddDataByCopy(FromElement.Get());
 		}
 	}
 }
@@ -386,7 +386,7 @@ FMoverDataStructBase* FMoverDataCollection::AddDataByType(const UScriptStruct* D
 	if (ensure(!FindDataByType(DataStructType)))
 	{
 		TSharedPtr<FMoverDataStructBase> NewDataInstance = CreateDataByType(DataStructType);
-		AddOrOverwriteData(NewDataInstance);
+		DataArray.Add(NewDataInstance);
 		return NewDataInstance.Get();
 	}
 	
@@ -401,6 +401,24 @@ void FMoverDataCollection::AddOrOverwriteData(const TSharedPtr<FMoverDataStructB
 }
 
 
+void FMoverDataCollection::AddDataByCopy(const FMoverDataStructBase* DataInstanceToCopy)
+{
+	check(DataInstanceToCopy);
+
+	const UScriptStruct* DataScriptStruct = DataInstanceToCopy->GetScriptStruct();
+
+	if (FMoverDataStructBase* ExistingMatchingData = FindDataByType(DataScriptStruct))
+	{
+		DataScriptStruct->CopyScriptStruct(ExistingMatchingData, DataInstanceToCopy, 1);
+
+	}
+	else
+	{
+		DataArray.Add(TSharedPtr<FMoverDataStructBase>(DataInstanceToCopy->Clone()));
+	}
+}
+
+
 FMoverDataStructBase* FMoverDataCollection::FindDataByType(const UScriptStruct* DataStructType) const
 {
 	for (const TSharedPtr<FMoverDataStructBase>& Data : DataArray)
@@ -581,9 +599,7 @@ DEFINE_FUNCTION(UMoverDataCollectionLibrary::execK2_AddDataToCollection)
 		if (ensure(SourceStructProp && SourceStructProp->Struct && SourceStructProp->Struct->IsChildOf(FMoverDataStructBase::StaticStruct()) && SourceDataAsRawPtr))
 		{
 			FMoverDataStructBase* SourceDataAsBasePtr = reinterpret_cast<FMoverDataStructBase*>(SourceDataAsRawPtr);
-			FMoverDataStructBase* ClonedData = SourceDataAsBasePtr->Clone();
-
-			TargetCollection.AddOrOverwriteData( TSharedPtr<FMoverDataStructBase>(ClonedData) );
+			TargetCollection.AddDataByCopy(SourceDataAsBasePtr);
 		}
 
 		P_NATIVE_END;
diff --git a/Engine/Plugins/Experimental/Mover/Source/Mover/Public/MoverTypes.h b/Engine/Plugins/Experimental/Mover/Source/Mover/Public/MoverTypes.h
--- a/Engine/Plugins/Experimental/Mover/Source/Mover/Public/MoverTypes.h
+++ b/Engine/Plugins/Experimental/Mover/Source/Mover/Public/MoverTypes.h
@@ -278,9 +278,15 @@ struct FMoverDataCollection
 		return *(static_cast<T*>(AddDataByType(T::StaticStruct())));
 	}
 
-	/** Adds data to the collection. If an existing data struct of the same type is already there, it will be removed first. */
+	/** Adds this data instance to the collection, taking ownership of it. If an existing data struct of the same type is already there, it will be removed first. */
 	UE_API void AddOrOverwriteData(const TSharedPtr<FMoverDataStructBase> DataInstance);
 
+	/** Adds data to the collection by copying over an existing struct or cloning the provided struct if no matching struct exists. 
+	 *  This is different than AddOrOverwriteData because the instance isn't touched, avoiding memory allocation & array changing if a matching struct exists.
+	 */
+	UE_API void AddDataByCopy(const FMoverDataStructBase* DataInstanceToCopy);
+
+
 	const TArray< TSharedPtr<FMoverDataStructBase> >& GetDataArray() const
 	{
 		return DataArray;
