Mover: formatting tweak to Gameplay Debugger Tool list of active sync state struct types. Readability improvement only, no functional changes.

#rnx

diff --git a/Engine/Plugins/Experimental/Mover/Source/Mover/Private/Debug/GameplayDebuggerCategory_Mover.cpp b/Engine/Plugins/Experimental/Mover/Source/Mover/Private/Debug/GameplayDebuggerCategory_Mover.cpp
--- a/Engine/Plugins/Experimental/Mover/Source/Mover/Private/Debug/GameplayDebuggerCategory_Mover.cpp
+++ b/Engine/Plugins/Experimental/Mover/Source/Mover/Private/Debug/GameplayDebuggerCategory_Mover.cpp
@@ -210,7 +210,7 @@ void FGameplayDebuggerCategory_Mover::DrawData(APlayerController* OwnerPC, FGame
 		*FString::JoinBy(DataPack.ActiveModifiers, TEXT("\n"), [](FString ModifierAsString) { return ModifierAsString; }),
 		*FString::JoinBy(DataPack.ModeMap, TEXT("\n"), [](FString ModeMappingAsString) { return ModeMappingAsString; }),
 		*FString::JoinBy(DataPack.ActiveTransitions, TEXT("\n"), [](FString TransitionAsString) { return TransitionAsString; }),
-		*FString::JoinBy(DataPack.SyncStateDataTypes, TEXT(","), [](FString SyncStateTypeAsString) { return SyncStateTypeAsString; })
+		*FString::JoinBy(DataPack.SyncStateDataTypes, TEXT("  "), [](FString SyncStateTypeAsString) { return SyncStateTypeAsString; })
 		);
 }
 
