Chaos Visual Debugger: Support for the seralization of FInstancedStructs as properties of mover info structs

diff --git a/Engine/Plugins/Experimental/Mover/Source/Mover/Private/ChaosVisualDebugger/MoverCVDRuntimeTrace.cpp b/Engine/Plugins/Experimental/Mover/Source/Mover/Private/ChaosVisualDebugger/MoverCVDRuntimeTrace.cpp
--- a/Engine/Plugins/Experimental/Mover/Source/Mover/Private/ChaosVisualDebugger/MoverCVDRuntimeTrace.cpp
+++ b/Engine/Plugins/Experimental/Mover/Source/Mover/Private/ChaosVisualDebugger/MoverCVDRuntimeTrace.cpp
@@ -26,9 +26,14 @@ CVD_DEFINE_OPTIONAL_DATA_CHANNEL(MoverNetworkedData, EChaosVDDataChannelInitiali
 CVD_DEFINE_OPTIONAL_DATA_CHANNEL(MoverLocalSimData, EChaosVDDataChannelInitializationFlags::CanChangeEnabledState)
 
 // FSkipObjectRefsMemoryWriter and FSkipObjectRefsMemoryReader are a 
-// workaround for serializing mover info structs with object references in them, such as the mover base
-// Ultimately we might want to attempt to translate the object to a particle ID and resolve that on the CVD side
-// by customizing the display and show the particle ID instead
+// workaround for serializing mover info structs with object references in them, such as the mover base.
+// It currently skips object references altogether, except if those are UScriptStruct objects, which it serializes
+// as the script struct name, hoping the type exists on the receiving end. This allows us to pass FInstancedStruct
+// as member properties of mover info structs. It is not backwads compatible though, and may cause crashes if the
+// underlying types have changed.
+// Ultimately we need to implement better backwards compatibility, possiblly using FPropertyBags. When the referenced object
+// is an actor with a primitive component though, we might want to attempt to translate the object reference to a particle ID and resolve that on the CVD side
+// by linking to the corresponding CVD particle, if found.
 class FSkipObjectRefsMemoryWriter : public FMemoryWriter
 {
 public:
@@ -37,13 +42,37 @@ public:
 	{
 	}
 
-	virtual FArchive& operator<<(struct FObjectPtr& Value)
+	virtual FArchive& operator<<(struct FObjectPtr& Value) override
 	{
+		if (UScriptStruct* ScriptStruct = Cast<UScriptStruct>(Value.Get()))
+		{
+			// The FullName of the script struct will be something like "ScriptStruct /Script/Mover.FCharacterDefaultInputs"
+			FString FullStructName = ScriptStruct->GetFullName(nullptr);
+			// We don't need to save the first part since we only ever save UScriptStructs (C++ structs)
+			FString StructName = FullStructName.RightChop(13); // So we chop the "ScriptStruct " part (hence 13 characters)
+			*this << StructName;
+		}
 		return *this;
 	}
 
-	virtual FArchive& operator<<(struct FWeakObjectPtr& Value)
+	virtual FArchive& operator<<(struct FWeakObjectPtr& Value) override
+	{
+		return *this;
+	}
+
+	virtual FArchive& operator<<(class UObject*& Res) override
 	{
+		if (Res)
+		{
+			if (UScriptStruct* ScriptStruct = Cast<UScriptStruct>(Res))
+			{
+				// The FullName of the script struct will be something like "ScriptStruct /Script/Mover.FCharacterDefaultInputs"
+				FString FullStructName = ScriptStruct->GetFullName(nullptr);
+				// We don't need to save the first part since we only ever save UScriptStructs (C++ structs)
+				FString StructName = FullStructName.RightChop(13); // So we chop the "ScriptStruct " part (hence 13 characters)
+				*this << StructName;
+			}
+		}
 		return *this;
 	}
 };
@@ -64,6 +93,14 @@ public:
 	{
 		return *this;
 	}
+
+	virtual FArchive& operator<<(class UObject*& Res) override
+	{
+		FString StructName;
+		*this << StructName;
+		Res = Cast<UScriptStruct>(FindObject<UStruct>(nullptr, *StructName));
+		return *this;
+	}
 };
 
 void FMoverCVDRuntimeTrace::UnwrapSimData(const FMoverCVDSimDataWrapper& InSimDataWrapper, TSharedPtr<FMoverInputCmdContext>& OutInputCmd, TSharedPtr<FMoverSyncState>& OutSyncState, TSharedPtr<FMoverDataCollection>& OutLocalSimData)
