Mover: Category fix for ArcRotationPathPattern

diff --git a/Engine/Plugins/Experimental/Mover/Source/Mover/Public/PhysicsMover/PathedMovement/ArcRotationPathPattern.h b/Engine/Plugins/Experimental/Mover/Source/Mover/Public/PhysicsMover/PathedMovement/ArcRotationPathPattern.h
--- a/Engine/Plugins/Experimental/Mover/Source/Mover/Public/PhysicsMover/PathedMovement/ArcRotationPathPattern.h
+++ b/Engine/Plugins/Experimental/Mover/Source/Mover/Public/PhysicsMover/PathedMovement/ArcRotationPathPattern.h
@@ -22,10 +22,10 @@ protected:
 	virtual FTransform CalcUnmaskedTargetRelativeTransform(float PatternProgress, const FTransform& CurTargetTransform) const override;
 
 	/** The size of the arc angle to rotate along */
-	UPROPERTY(EditAnywhere, Category = EllipticalPathPattern, meta = (Units = Degrees, UIMin = 0, ClampMin = 0, UIMax = 360, ClampMax = 360))
+	UPROPERTY(EditAnywhere, Category = ArcRotationPattern, meta = (Units = Degrees, UIMin = 0, ClampMin = 0, UIMax = 360, ClampMax = 360))
 	float RotationArcAngle = 360.f;
 
 	/** The axis to rotate about */
-	UPROPERTY(EditAnywhere, Category = EllipticalPathPattern, meta = (UIMin = 0, ClampMin = 0, UIMax = 1.f, ClampMax = 1.f))
+	UPROPERTY(EditAnywhere, Category = ArcRotationPattern, meta = (UIMin = 0, ClampMin = 0, UIMax = 1.f, ClampMax = 1.f))
 	FVector RotationAxis = FVector::UpVector;
 };
\ No newline at end of file
