Mover: deleted stale comment. No functional changes.

#rnx

diff --git a/Engine/Plugins/Experimental/Mover/Source/Mover/Private/MoverComponent.cpp b/Engine/Plugins/Experimental/Mover/Source/Mover/Private/MoverComponent.cpp
--- a/Engine/Plugins/Experimental/Mover/Source/Mover/Private/MoverComponent.cpp
+++ b/Engine/Plugins/Experimental/Mover/Source/Mover/Private/MoverComponent.cpp
@@ -608,10 +608,6 @@ void UMoverComponent::SimulationTick(const FMoverTimeStep& InTimeStep, const FMo
 		}
 	}
 
-	// Note that we don't pull the rotation out of the final update transform. Converting back from a quat will lead to a different FRotator than what we are storing
-	// here in the simulation layer. This may not be the best choice for all movement simulations, but is ok for this one.
-	// JAH TODO: re-evaluate the above comment about handling the rotation data
-
 	if (!SimOutput.MoveRecord.GetTotalMoveDelta().IsZero())
 	{
 		UE_LOG(LogMover, VeryVerbose, TEXT("KinematicSimTick: %s (role %i) frame %d: %s"),
