Mover: making spring-based walking mode work in networked play, adding directional intent to its proposed moves for anim hints, and hardening against possible loss of inputs

diff --git a/Engine/Plugins/Experimental/Mover/Source/Mover/Private/DefaultMovementSet/Modes/SpringWalkingMode.cpp b/Engine/Plugins/Experimental/Mover/Source/Mover/Private/DefaultMovementSet/Modes/SpringWalkingMode.cpp
--- a/Engine/Plugins/Experimental/Mover/Source/Mover/Private/DefaultMovementSet/Modes/SpringWalkingMode.cpp
+++ b/Engine/Plugins/Experimental/Mover/Source/Mover/Private/DefaultMovementSet/Modes/SpringWalkingMode.cpp
@@ -12,7 +12,9 @@
 void USpringWalkingMode::GenerateMove_Implementation(const FMoverTickStartData& StartState, const FMoverTimeStep& TimeStep,
                                                    FProposedMove& OutProposedMove) const
 {
-	if (CommonLegacySettings.Get() == nullptr)
+	const FMoverDefaultSyncState* StartingSyncState = StartState.SyncState.SyncStateCollection.FindDataByType<FMoverDefaultSyncState>();
+
+	if (CommonLegacySettings.Get() == nullptr || !StartingSyncState)
 	{
 		return;
 	}
@@ -23,16 +25,30 @@ void USpringWalkingMode::GenerateMove_Implementation(const FMoverTickStartData&
 		return;
 	}
 	
-	const FMoverDefaultSyncState* StartingSyncState = StartState.SyncState.SyncStateCollection.FindDataByType<FMoverDefaultSyncState>();
-
-	// Get input
-	const FCharacterDefaultInputs* CharacterInputs = StartState.InputCmd.InputCollection.FindDataByType<FCharacterDefaultInputs>();
 
 	// Need to hack around the fact that GenerateMove is const, but we want to mutate state data here for our spring
 	FMoverTickStartData* MutableStartData = const_cast<FMoverTickStartData*>(&StartState);
 	FMoverSpringState& SpringState = MutableStartData->SyncState.SyncStateCollection.FindOrAddMutableDataByType<FMoverSpringState>();
 	
-	FVector DesiredVelocity = CharacterInputs->GetMoveInput_WorldSpace();
+	// Get input
+	FVector DesiredVelocity;
+	EMoveInputType MoveInputType;
+	FVector DesiredFacingDir;
+	 
+	if (const FCharacterDefaultInputs* CharacterInputs = StartState.InputCmd.InputCollection.FindDataByType<FCharacterDefaultInputs>())
+	{
+		DesiredVelocity = CharacterInputs->GetMoveInput_WorldSpace();
+		MoveInputType = CharacterInputs->GetMoveInputType();
+		DesiredFacingDir = CharacterInputs->GetOrientationIntentDir_WorldSpace();
+	}
+	else
+	{
+		DesiredVelocity = FVector::ZeroVector;
+		MoveInputType = EMoveInputType::DirectionalIntent;
+		DesiredFacingDir = FVector::ZeroVector;
+	}
+	
+	
 	float MaxMoveSpeed = CommonLegacySettings->MaxSpeed;
 
 	// Subtract vertical component but keep same magnitude
@@ -44,10 +60,11 @@ void USpringWalkingMode::GenerateMove_Implementation(const FMoverTickStartData&
 		DesiredVelocity *= DesiredVelMag / FMath::Sqrt(DesiredVel2DSquaredLength); 
 	}
 
-	switch (CharacterInputs->GetMoveInputType())
+	switch (MoveInputType)
 	{
 	case EMoveInputType::DirectionalIntent:
 		{
+			OutProposedMove.DirectionIntent = DesiredVelocity;	// here, DesiredVelocity is already in "intent space" (unit length for "max intent") so we can use it directly
 			DesiredVelocity *= MaxMoveSpeed;
 		}
 		break;
@@ -55,18 +72,21 @@ void USpringWalkingMode::GenerateMove_Implementation(const FMoverTickStartData&
 		{
 			// Clamp to max move speed
 			DesiredVelocity = DesiredVelocity.GetClampedToMaxSize(MaxMoveSpeed);
+			OutProposedMove.DirectionIntent = DesiredVelocity / MaxMoveSpeed; // here, DesiredVelocity is converted to "intent space"
 		}
 		break;
 	case EMoveInputType::None:
 	case EMoveInputType::Invalid:
 	default:
 		{
-			UE_LOG(LogMover, Warning, TEXT("Unhandled MoveInputType %i in USpringWalkingMode"), CharacterInputs->GetMoveInputType());
+			UE_LOG(LogMover, Warning, TEXT("Unhandled MoveInputType %i in USpringWalkingMode"), MoveInputType);
 			DesiredVelocity = FVector::ZeroVector;
 		}
 	break;
 	}
 
+	OutProposedMove.bHasDirIntent = !OutProposedMove.DirectionIntent.IsNearlyZero();
+
 	// Linear //
 	FVector CurrentPosition = StartingSyncState->GetLocation_WorldSpace();
 	FVector PrevPosition = CurrentPosition;
@@ -76,7 +96,6 @@ void USpringWalkingMode::GenerateMove_Implementation(const FMoverTickStartData&
 	OutProposedMove.LinearVelocity = (CurrentPosition - PrevPosition) / DeltaSeconds;
 
 	// Angular //
-	FVector DesiredFacingDir = CharacterInputs->GetOrientationIntentDir_WorldSpace();
 	DesiredFacingDir -= DesiredFacingDir.ProjectOnTo(GetMoverComponent()->GetUpDirection());
 	if (DesiredFacingDir.IsNearlyZero())
 	{
diff --git a/Engine/Plugins/Experimental/Mover/Source/Mover/Private/DefaultMovementSet/Modes/SpringWalkingState.cpp b/Engine/Plugins/Experimental/Mover/Source/Mover/Private/DefaultMovementSet/Modes/SpringWalkingState.cpp
new file mode 100644
--- /dev/null
+++ b/Engine/Plugins/Experimental/Mover/Source/Mover/Private/DefaultMovementSet/Modes/SpringWalkingState.cpp
@@ -0,0 +1,52 @@
+// Copyright Epic Games, Inc. All Rights Reserved.
+
+#include "SpringWalkingState.h"
+
+UScriptStruct* FMoverSpringState::GetScriptStruct() const
+{ 
+	return StaticStruct(); 
+}
+
+FMoverDataStructBase* FMoverSpringState::Clone() const
+{ 
+	return new FMoverSpringState(*this); 
+}
+
+bool FMoverSpringState::NetSerialize(FArchive& Ar, UPackageMap* Map, bool& bOutSuccess)
+{
+	bool bSuccess = Super::NetSerialize(Ar, Map, bOutSuccess);
+
+	// Could be quantized to save bandwidth
+	Ar << CurrentVelocity;
+	Ar << CurrentAccel;
+	Ar << CurrentAngularVelocity;
+
+	return bSuccess;
+}
+
+void FMoverSpringState::ToString(FAnsiStringBuilderBase& Out) const
+{
+	Super::ToString(Out);
+
+	Out.Appendf("Vel=%s Accel=%s AngVel=%s\n", *CurrentVelocity.ToCompactString(), *CurrentAccel.ToCompactString(), *CurrentAngularVelocity.ToCompactString());
+}
+
+bool FMoverSpringState::ShouldReconcile(const FMoverDataStructBase& AuthorityState) const
+{
+	const FMoverSpringState* AuthoritySpringState = static_cast<const FMoverSpringState*>(&AuthorityState);
+
+	return (!(CurrentVelocity - AuthoritySpringState->CurrentVelocity).IsNearlyZero() ||
+			!(CurrentAccel - AuthoritySpringState->CurrentAccel).IsNearlyZero() ||
+			!(CurrentAngularVelocity - AuthoritySpringState->CurrentAngularVelocity).IsNearlyZero());
+}
+
+
+void FMoverSpringState::Interpolate(const FMoverDataStructBase& From, const FMoverDataStructBase& To, float Pct)
+{
+	const FMoverSpringState* FromState = static_cast<const FMoverSpringState*>(&From);
+	const FMoverSpringState* ToState   = static_cast<const FMoverSpringState*>(&To);
+
+	CurrentVelocity        = FMath::Lerp(FromState->CurrentVelocity, ToState->CurrentVelocity, Pct);
+	CurrentAccel           = FMath::Lerp(FromState->CurrentAccel, ToState->CurrentAccel, Pct);
+	CurrentAngularVelocity = FMath::Lerp(FromState->CurrentAngularVelocity, ToState->CurrentAngularVelocity, Pct);
+}
diff --git a/Engine/Plugins/Experimental/Mover/Source/Mover/Private/DefaultMovementSet/Modes/SpringWalkingState.h b/Engine/Plugins/Experimental/Mover/Source/Mover/Private/DefaultMovementSet/Modes/SpringWalkingState.h
--- a/Engine/Plugins/Experimental/Mover/Source/Mover/Private/DefaultMovementSet/Modes/SpringWalkingState.h
+++ b/Engine/Plugins/Experimental/Mover/Source/Mover/Private/DefaultMovementSet/Modes/SpringWalkingState.h
@@ -15,8 +15,13 @@ struct FMoverSpringState : public FMoverDataStructBase
 {
 	GENERATED_BODY()
 	
-	virtual UScriptStruct* GetScriptStruct() const override { return StaticStruct(); }
-	virtual FMoverDataStructBase* Clone() const override { return new FMoverSpringState(*this); }
+	virtual UScriptStruct* GetScriptStruct() const override;
+	virtual FMoverDataStructBase* Clone() const override;
+	virtual bool NetSerialize(FArchive& Ar, UPackageMap* Map, bool& bOutSuccess) override;
+	virtual void ToString(FAnsiStringBuilderBase& Out) const override;
+	virtual bool ShouldReconcile(const FMoverDataStructBase& AuthorityState) const override;
+	virtual void Interpolate(const FMoverDataStructBase& From, const FMoverDataStructBase& To, float Pct) override;
+
 
 	// Velocity of internal spring model. Not necessarily the velocity of the character
 	UPROPERTY(BlueprintReadOnly, Category = "Mover|Experimental")
