Mover:
Support for scheduled instant movement effects, i.e. effects scheduled to occur after a delay allowing server and peers to receive the event and apply on the same frame. Instant movement effects will only be networked in the context of Chaos Mover (async physics 2.0). Scheduling uses server frame by default, but falls back to using server time (just phyiscs time in standalone) in the variable dt case.
Support for teleportation for async physics.
Simulation events for teleport success and failure
Instant Movement Effects queue works now as a pull mechanism, and not push. This is so different liaison components or the mover state machine can pull those effects on their own terms.

diff --git a/Engine/Plugins/Experimental/Mover/Source/Mover/Private/Backends/MoverNetworkPhysicsLiaisonBase.cpp b/Engine/Plugins/Experimental/Mover/Source/Mover/Private/Backends/MoverNetworkPhysicsLiaisonBase.cpp
--- a/Engine/Plugins/Experimental/Mover/Source/Mover/Private/Backends/MoverNetworkPhysicsLiaisonBase.cpp
+++ b/Engine/Plugins/Experimental/Mover/Source/Mover/Private/Backends/MoverNetworkPhysicsLiaisonBase.cpp
@@ -597,7 +597,7 @@ void UMoverNetworkPhysicsLiaisonComponentBase::PostPhysicsUpdate_External()
 	FMoverInputCmdContext InterpolatedInputCmd;
 	if (bUsingAsyncPhysics)
 	{
-		const float ResultsTime = GetPhysicsSolver()->GetPhysicsResultsTime_External();
+		const double ResultsTime = GetPhysicsSolver()->GetPhysicsResultsTime_External();
 		SimOutputRecord.GetInterpolated(ResultsTime, InterpolatedSyncState, InterpolatedInputCmd);
 	}
 	else
diff --git a/Engine/Plugins/Experimental/Mover/Source/Mover/Private/DefaultMovementSet/InstantMovementEffects/BasicInstantMovementEffects.cpp b/Engine/Plugins/Experimental/Mover/Source/Mover/Private/DefaultMovementSet/InstantMovementEffects/BasicInstantMovementEffects.cpp
--- a/Engine/Plugins/Experimental/Mover/Source/Mover/Private/DefaultMovementSet/InstantMovementEffects/BasicInstantMovementEffects.cpp
+++ b/Engine/Plugins/Experimental/Mover/Source/Mover/Private/DefaultMovementSet/InstantMovementEffects/BasicInstantMovementEffects.cpp
@@ -6,6 +6,7 @@
 #include "MoverComponent.h"
 #include "MoverDataModelTypes.h"
 #include "MoverSimulationTypes.h"
+#include "MoverSimulation.h"
 #include "DefaultMovementSet/Settings/CommonLegacyMovementSettings.h"
 #include "MoveLibrary/MovementUtils.h"
 #include "MoveLibrary/MoverBlackboard.h"
@@ -36,47 +37,64 @@ FTeleportEffect::FTeleportEffect()
 
 bool FTeleportEffect::ApplyMovementEffect(FApplyMovementEffectParams& ApplyEffectParams, FMoverSyncState& OutputState)
 {
-	const FRotator Rotation = bUseActorRotation ? ApplyEffectParams.UpdatedComponent->GetComponentRotation() : TargetRotation;
+	USceneComponent* UpdatedComponent = ApplyEffectParams.UpdatedComponent;
+	const FRotator FinalTargetRotation = bUseActorRotation ? UpdatedComponent->GetComponentRotation() : TargetRotation;
 	
-	if (ApplyEffectParams.UpdatedComponent->GetOwner()->TeleportTo(TargetLocation, Rotation))
+	const FVector PreviousLocation = UpdatedComponent->GetComponentLocation();
+	const FQuat PreviousRotation = UpdatedComponent->GetComponentQuat();
+	AActor* OwnerActor = UpdatedComponent->GetOwner();
+
+	if (OwnerActor->TeleportTo(TargetLocation, FinalTargetRotation))
 	{
+		const FVector UpdatedLocation = UpdatedComponent->GetComponentLocation();
+#if !BUILD_SHIPPING
 		if (ShowTeleportDiffs)
 		{
-			const AActor* OwnerActor = ApplyEffectParams.UpdatedComponent->GetOwner();
-
-			if (!(ApplyEffectParams.UpdatedComponent->GetComponentLocation() - TargetLocation).IsNearlyZero())	// if it was adjusted, show the original error
+			if (!(UpdatedLocation - TargetLocation).IsNearlyZero())	// if it was adjusted, show the original error
 			{
 				DrawDebugCapsule(OwnerActor->GetWorld(), TargetLocation, OwnerActor->GetSimpleCollisionHalfHeight(), OwnerActor->GetSimpleCollisionRadius(), FQuat::Identity, FColor::Red, false, ShowTeleportDiffsLifetimeSecs);
 			}
-
-			DrawDebugCapsule(OwnerActor->GetWorld(), ApplyEffectParams.UpdatedComponent->GetComponentLocation(), OwnerActor->GetSimpleCollisionHalfHeight(), OwnerActor->GetSimpleCollisionRadius(), FQuat::Identity, FColor(100, 100, 255), false, ShowTeleportDiffsLifetimeSecs);
+			DrawDebugCapsule(OwnerActor->GetWorld(), UpdatedLocation, OwnerActor->GetSimpleCollisionHalfHeight(), OwnerActor->GetSimpleCollisionRadius(), FQuat::Identity, FColor(100, 100, 255), false, ShowTeleportDiffsLifetimeSecs);
 		}
+#endif // !BUILD_SHIPPING
 
 		FMoverDefaultSyncState& OutputSyncState = OutputState.SyncStateCollection.FindOrAddMutableDataByType<FMoverDefaultSyncState>();
+		OutputSyncState.SetTransforms_WorldSpace(UpdatedLocation,
+													UpdatedComponent->GetComponentRotation(),
+													OutputSyncState.GetVelocity_WorldSpace(),
+													nullptr ); // no movement base
 		
-		if (const FMoverDefaultSyncState* StartingSyncState = ApplyEffectParams.StartState->SyncState.SyncStateCollection.FindDataByType<FMoverDefaultSyncState>())
+		// TODO: instead of invalidating it, consider checking for a floor. Possibly a dynamic base?
+		if (UMoverBlackboard* SimBlackboard = ApplyEffectParams.MoverComp->GetSimBlackboard_Mutable())
 		{
-			OutputSyncState.SetTransforms_WorldSpace( ApplyEffectParams.UpdatedComponent->GetComponentLocation(),
-													  ApplyEffectParams.UpdatedComponent->GetComponentRotation(),
-													  OutputSyncState.GetVelocity_WorldSpace(),
-													  nullptr ); // no movement base
-		
-			// TODO: instead of invalidating it, consider checking for a floor. Possibly a dynamic base?
-			if (UMoverBlackboard* SimBlackboard = ApplyEffectParams.MoverComp->GetSimBlackboard_Mutable())
-			{
-				SimBlackboard->Invalidate(CommonBlackboard::LastFloorResult);
-				SimBlackboard->Invalidate(CommonBlackboard::LastFoundDynamicMovementBase);
-			}
-
-			return true;
+			SimBlackboard->Invalidate(CommonBlackboard::LastFloorResult);
+			SimBlackboard->Invalidate(CommonBlackboard::LastFoundDynamicMovementBase);
 		}
+
+		ApplyEffectParams.OutputEvents.Add(MakeShared<FTeleportSucceededEventData>(ApplyEffectParams.TimeStep->BaseSimTimeMs, PreviousLocation, PreviousRotation, TargetLocation, FQuat(FinalTargetRotation)));
+
+		return true;
 	}
 
+#if !BUILD_SHIPPING
 	if (ShowTeleportDiffs)
 	{
-		const AActor* OwnerActor = ApplyEffectParams.UpdatedComponent->GetOwner();
 		DrawDebugCapsule(OwnerActor->GetWorld(), TargetLocation, OwnerActor->GetSimpleCollisionHalfHeight(), OwnerActor->GetSimpleCollisionRadius(), FQuat::Identity, FColor::Red, false, ShowTeleportDiffsLifetimeSecs);
 	}
+#endif // !BUILD_SHIPPING
+
+	ApplyEffectParams.OutputEvents.Add(MakeShared<FTeleportFailedEventData>(ApplyEffectParams.TimeStep->BaseSimTimeMs, PreviousLocation, PreviousRotation, TargetLocation, FQuat(FinalTargetRotation), ETeleportFailureReason::Reason_NotAvailable));
+
+	return false;
+}
+
+bool FTeleportEffect::ApplyMovementEffect_Async(FApplyMovementEffectParams_Async& ApplyEffectParams, FMoverSyncState& OutputState)
+{
+	if (ApplyEffectParams.Simulation)
+	{
+		ApplyEffectParams.Simulation->AttemptTeleport(*ApplyEffectParams.TimeStep, FTransform(TargetRotation, TargetLocation), bUseActorRotation, OutputState);
+		return true;
+	}
 
 	return false;
 }
@@ -107,7 +125,7 @@ UScriptStruct* FTeleportEffect::GetScriptStruct() const
 
 FString FTeleportEffect::ToSimpleString() const
 {
-	return FString::Printf(TEXT("Teleport"));
+	return bUseActorRotation ? FString::Printf(TEXT("Teleport to %s (bUseActorRotation = True)"), *TargetLocation.ToString()) : FString::Printf(TEXT("Teleport to %s, %s (bUseActorRotation = False)"), *TargetLocation.ToString(), *FRotator(TargetRotation).ToString());
 }
 
 void FTeleportEffect::AddReferencedObjects(FReferenceCollector& Collector)
diff --git a/Engine/Plugins/Experimental/Mover/Source/Mover/Private/DefaultMovementSet/Modes/SwimmingMode.cpp b/Engine/Plugins/Experimental/Mover/Source/Mover/Private/DefaultMovementSet/Modes/SwimmingMode.cpp
--- a/Engine/Plugins/Experimental/Mover/Source/Mover/Private/DefaultMovementSet/Modes/SwimmingMode.cpp
+++ b/Engine/Plugins/Experimental/Mover/Source/Mover/Private/DefaultMovementSet/Modes/SwimmingMode.cpp
@@ -164,7 +164,6 @@ void USwimmingMode::GenerateMove_Implementation(const FMoverTickStartData& Start
 
 void USwimmingMode::SimulationTick_Implementation(const FSimulationTickParams& Params, FMoverTickEndData& OutputState)
 {
-	// TODO Har
 }
 
 void USwimmingMode::OnRegistered(const FName ModeName)
@@ -186,12 +185,14 @@ void USwimmingMode::OnUnregistered()
 	Super::OnUnregistered();
 }
 
-bool USwimmingMode::AttemptJump(float UpwardsSpeed, FMoverTickEndData& OutputState)
+bool USwimmingMode::AttemptJump(const FSimulationTickParams& Params, float UpwardsSpeed, FMoverTickEndData& OutputState)
 {
 	// TODO: This should check if a jump is even allowed
 	TSharedPtr<FJumpImpulseEffect> JumpMove = MakeShared<FJumpImpulseEffect>();
 	JumpMove->UpwardsSpeed = UpwardsSpeed;
-	GetMoverComponent()->QueueInstantMovementEffect(JumpMove);
+
+	// TODO: Use UMoverSimulation to queue instant movement effects when it is passed as part of FSimulationTickParams
+	GetMoverComponent()->QueueInstantMovementEffect_Internal(Params.TimeStep, JumpMove);
 	
 	return true;
 }
diff --git a/Engine/Plugins/Experimental/Mover/Source/Mover/Private/MovementModeStateMachine.cpp b/Engine/Plugins/Experimental/Mover/Source/Mover/Private/MovementModeStateMachine.cpp
--- a/Engine/Plugins/Experimental/Mover/Source/Mover/Private/MovementModeStateMachine.cpp
+++ b/Engine/Plugins/Experimental/Mover/Source/Mover/Private/MovementModeStateMachine.cpp
@@ -9,6 +9,8 @@
 #include "MoverComponent.h"
 #include "MoveLibrary/MovementMixer.h"
 #include "Templates/SubclassOf.h"
+#include "Engine/World.h"
+#include "GameFramework/GameStateBase.h"
 
 
 #include UE_INLINE_GENERATED_CPP_BY_NAME(MovementModeStateMachine)
@@ -162,6 +164,12 @@ void UMovementModeStateMachine::OnSimulationTick(USceneComponent* UpdatedCompone
 	UMoverComponent* MoverComp = CastChecked<UMoverComponent>(GetOuter());
 	check(MoverComp->MovementMixer);
 
+	// Copy queued instant movement effects from the mover component to the state machine
+	// After this, the state machine works on its own queue, into which it can enqueue instant movement effects while stepping
+	const TArray<FScheduledInstantMovementEffect>& ScheduledInstantMovementEffects = MoverComp->GetQueuedInstantMovementEffects();
+	QueueInstantMovementEffects(ScheduledInstantMovementEffects);
+	MoverComp->ClearQueuedInstantMovementEffects();
+
 	if (!QueuedModeTransition->IsSet())
 	{
 		QueueNextMode(WorkingSubstepStartData.SyncState.MovementMode);
@@ -500,10 +508,25 @@ void UMovementModeStateMachine::QueueLayeredMove(TSharedPtr<FLayeredMoveBase> Mo
 	QueuedLayeredMoves.Add(Move);
 }
 
-void UMovementModeStateMachine::QueueInstantMovementEffect(TSharedPtr<FInstantMovementEffect> Effect)
+void UMovementModeStateMachine::QueueInstantMovementEffect_Internal(const FScheduledInstantMovementEffect& ScheduledInstantMovementEffect)
+{
+	ensure(!IsInGameThread());
+	QueueInstantMovementEffect(ScheduledInstantMovementEffect);
+}
+
+void UMovementModeStateMachine::QueueInstantMovementEffect(const FScheduledInstantMovementEffect& ScheduledInstantMovementEffect)
 {
 	FRWScopeLock QueueLock(InstantEffectsQueueLock, SLT_Write);
-	QueuedInstantEffects.Add(Effect);
+	QueuedInstantEffects.Add(ScheduledInstantMovementEffect);
+}
+
+void UMovementModeStateMachine::QueueInstantMovementEffects(const TArray<FScheduledInstantMovementEffect>& ScheduledInstantMovementEffects)
+{
+	FRWScopeLock QueueLock(InstantEffectsQueueLock, SLT_Write);
+	for (const FScheduledInstantMovementEffect& ScheduledInstantMovementEffect : ScheduledInstantMovementEffects)
+	{
+		QueuedInstantEffects.Add(ScheduledInstantMovementEffect);
+	}
 }
 
 FMovementModifierHandle UMovementModeStateMachine::QueueMovementModifier(TSharedPtr<FMovementModifierBase> Modifier)
@@ -725,22 +748,58 @@ bool UMovementModeStateMachine::ApplyInstantEffects(FApplyMovementEffectParams&
 
 	FRWScopeLock QueueLock(InstantEffectsQueueLock, SLT_Write);
 
-	if (!QueuedInstantEffects.IsEmpty())
+	for (int EffectIndex = 0; EffectIndex < QueuedInstantEffects.Num(); )
 	{
-		for (TSharedPtr<FInstantMovementEffect>& QueuedEffect : QueuedInstantEffects)
+		const FScheduledInstantMovementEffect& QueuedEffect = QueuedInstantEffects[EffectIndex];
+		bool bshouldApplyEffect = false;
+		if (QueuedEffect.bIsFixedDt)
 		{
-			if (QueuedEffect->ApplyMovementEffect(ApplyEffectParams, OutputState))
-			{
-				bInstantMovementEffectApplied = true;
-			}
+			bshouldApplyEffect = QueuedEffect.ShouldExecuteAtFrame(ApplyEffectParams.TimeStep->ServerFrame);
+		}
+		else
+		{
+			UWorld* World = GetWorld();
+			AGameStateBase* GameState = World ? World->GetGameState() : nullptr;
+			double ServerTime = GameState ? GameState->GetServerWorldTimeSeconds() : 0.0;
+			bshouldApplyEffect = QueuedEffect.ShouldExecuteAtTime(ServerTime);
+		}
+		if (bshouldApplyEffect)
+		{
+			bInstantMovementEffectApplied |= QueuedEffect.Effect->ApplyMovementEffect(ApplyEffectParams, OutputState);
+			QueuedInstantEffects.RemoveAt(EffectIndex);
+			ProcessEvents(ApplyEffectParams.OutputEvents);
+			ApplyEffectParams.OutputEvents.Empty();
+		}
+		else
+		{
+			EffectIndex++;
 		}
-		
-		QueuedInstantEffects.Empty();
 	}
 	
 	return bInstantMovementEffectApplied;
 }
 
+void UMovementModeStateMachine::ProcessEvents(const TArray<TSharedPtr<FMoverSimulationEventData>>& InEvents)
+{
+	UMoverComponent* MoverComp = CastChecked<UMoverComponent>(GetOuter());
+	for (TSharedPtr<FMoverSimulationEventData> Event : InEvents)
+	{
+		if (const FMoverSimulationEventData* EventData = Event.Get())
+		{
+			ProcessSimulationEvent(*EventData);
+
+#if !BUILD_SHIPPING
+			ensureMsgf(IsInGameThread(), TEXT("Dispatching an event to the mover component from outside the game thread, this is not thread safe"));
+#endif
+			MoverComp->DispatchSimulationEvent(*EventData);
+		}
+	}
+}
+
+void UMovementModeStateMachine::ProcessSimulationEvent(const FMoverSimulationEventData& EventData)
+{
+}
+
 AActor* UMovementModeStateMachine::GetOwnerActor() const
 {
 	if (const UActorComponent* OwnerMoverComp = Cast<UActorComponent>(GetOuter()))
diff --git a/Engine/Plugins/Experimental/Mover/Source/Mover/Private/MoverComponent.cpp b/Engine/Plugins/Experimental/Mover/Source/Mover/Private/MoverComponent.cpp
--- a/Engine/Plugins/Experimental/Mover/Source/Mover/Private/MoverComponent.cpp
+++ b/Engine/Plugins/Experimental/Mover/Source/Mover/Private/MoverComponent.cpp
@@ -130,7 +130,7 @@ void UMoverComponent::InitializeComponent()
 			ModeFSM->QueueNextMode(StartingMovementMode);
 		}
 
-		// Instantiate out sister backend component that will actually talk to the system driving the simulation
+		// Instantiate our sister backend component that will actually talk to the system driving the simulation
 		if (BackendClass)
 		{
 			UActorComponent* NewLiaisonComp = NewObject<UActorComponent>(this, BackendClass, TEXT("BackendLiaisonComponent"));
@@ -303,6 +303,12 @@ void UMoverComponent::BeginPlay()
 	{
 		MovementMixer = NewObject<UMovementMixer>(this, TEXT("Default Movement Mixer"));
 	}
+
+	// Initialize the fixed delay for event scheduling
+	if (BackendLiaisonComp)
+	{
+		EventSchedulingMinDelaySeconds = BackendLiaisonComp->GetEventSchedulingMinDelaySeconds();
+	}
 }
 
 void UMoverComponent::BindProcessGeneratedMovement(FMover_ProcessGeneratedMovement ProcessGeneratedMovementEvent)
@@ -1543,10 +1549,80 @@ DEFINE_FUNCTION(UMoverComponent::execK2_QueueInstantMovementEffect)
 	P_NATIVE_END;
 }
 
+void UMoverComponent::K2_ScheduleInstantMovementEffect(const int32& EffectAsRawData)
+{
+	// This will never be called, the exec version below will be hit instead
+	checkNoEntry();
+}
+
+DEFINE_FUNCTION(UMoverComponent::execK2_ScheduleInstantMovementEffect)
+{
+	Stack.StepCompiledIn<FStructProperty>(nullptr);
+	void* EffectPtr = Stack.MostRecentPropertyAddress;
+	FStructProperty* StructProp = CastField<FStructProperty>(Stack.MostRecentProperty);
+
+	P_FINISH;
+
+	P_NATIVE_BEGIN;
+
+	const bool bHasValidStructProp = StructProp && StructProp->Struct && StructProp->Struct->IsChildOf(FInstantMovementEffect::StaticStruct());
+
+	if (ensureMsgf((bHasValidStructProp && EffectPtr), TEXT("An invalid type (%s) was sent to a QueueInstantMovementEffect node. A struct derived from FInstantMovementEffect is required. No Movement Effect will be queued."),
+		StructProp ? *GetNameSafe(StructProp->Struct) : *Stack.MostRecentProperty->GetClass()->GetName()))
+	{
+		// Could we steal this instead of cloning? (move semantics)
+		FInstantMovementEffect* EffectAsBasePtr = reinterpret_cast<FInstantMovementEffect*>(EffectPtr);
+		FInstantMovementEffect* ClonedMove = EffectAsBasePtr->Clone();
+
+		P_THIS->ScheduleInstantMovementEffect(TSharedPtr<FInstantMovementEffect>(ClonedMove));
+	}
+
+	P_NATIVE_END;
+}
+
+void UMoverComponent::ScheduleInstantMovementEffect(TSharedPtr<FInstantMovementEffect> InstantMovementEffect)
+{
+	ensure(IsInGameThread());
+	QueueInstantMovementEffect(FScheduledInstantMovementEffect::ScheduleEffect_External(GetWorld(), InstantMovementEffect, /* SchedulingDelaySeconds = */ EventSchedulingMinDelaySeconds));
+}
+
+void UMoverComponent::QueueInstantMovementEffect_Internal(const FMoverTimeStep& TimeStep, TSharedPtr<FInstantMovementEffect> InstantMovementEffect)
+{
+	QueueInstantMovementEffect(FScheduledInstantMovementEffect::ScheduleEffect_Internal(GetWorld(), TimeStep, InstantMovementEffect, /* SchedulingDelaySeconds = */ 0.0f));
+}
 
 void UMoverComponent::QueueInstantMovementEffect(TSharedPtr<FInstantMovementEffect> InstantMovementEffect)
-{	
-	ModeFSM->QueueInstantMovementEffect(InstantMovementEffect);
+{
+	ensure(IsInGameThread());
+	QueueInstantMovementEffect(FScheduledInstantMovementEffect::ScheduleEffect_External(GetWorld(), InstantMovementEffect, /* SchedulingDelaySeconds = */ 0.0f));
+}
+
+void UMoverComponent::QueueInstantMovementEffect(const FScheduledInstantMovementEffect& InstantMovementEffect)
+{
+	if (IsInGameThread())
+	{
+		QueuedInstantMovementEffects.Add(InstantMovementEffect);
+	}
+	else
+	{
+		ModeFSM->QueueInstantMovementEffect_Internal(InstantMovementEffect);
+	}	
+
+#if !BUILD_SHIPPING
+	ENetMode NetMode = GetWorld() ? GetWorld()->GetNetMode() : NM_MAX;
+	UE_LOG(LogMover, Verbose, TEXT("(%s) UMoverComponent::QueueInstantMovementEffect: Game Thread queueing an instant movement effect scheduled for frame %d: %s."),
+		*ToString(NetMode), InstantMovementEffect.ExecutionServerFrame, InstantMovementEffect.Effect.IsValid() ? *InstantMovementEffect.Effect->ToSimpleString() : TEXT("INVALID INSTANT EFFECT"));
+#endif // !BUILD_SHIPPING
+}
+
+const TArray<FScheduledInstantMovementEffect>& UMoverComponent::GetQueuedInstantMovementEffects() const
+{
+	return QueuedInstantMovementEffects;
+}
+
+void UMoverComponent::ClearQueuedInstantMovementEffects()
+{
+	QueuedInstantMovementEffects.Empty();
 }
 
 UBaseMovementMode* UMoverComponent::FindMovementModeByName(FName MovementModeName) const
@@ -2266,17 +2342,7 @@ void UMoverComponent::SetSimulationOutput(const FMoverTimeStep& TimeStep, const
 	{
 		if (const FMoverSimulationEventData* Data = EventData.Get())
 		{
-			// This gives the event a callback when it is processed on the game thread
-			Data->OnEventProcessed();
-
-			// Process the simulation event at the mover component (or derived) level
-			ProcessSimulationEvent(*Data);
-
-			// Broadcast the event to 
-			if (OnPostSimEventReceived.IsBound())
-			{
-				OnPostSimEventReceived.Broadcast(*Data);
-			}
+			DispatchSimulationEvent(*Data);
 		}
 	}
 
@@ -2285,6 +2351,21 @@ void UMoverComponent::SetSimulationOutput(const FMoverTimeStep& TimeStep, const
 	SetAdditionalSimulationOutput(OutputData.AdditionalOutputData);
 }
 
+void UMoverComponent::DispatchSimulationEvent(const FMoverSimulationEventData& EventData)
+{
+	// This gives the event a callback when it is processed on the game thread
+	EventData.OnEventProcessed();
+
+	// Process the simulation event at the mover component (or derived) level
+	ProcessSimulationEvent(EventData);
+
+	// Broadcast the event outside mover component
+	if (OnPostSimEventReceived.IsBound())
+	{
+		OnPostSimEventReceived.Broadcast(EventData);
+	}
+}
+
 void UMoverComponent::ProcessSimulationEvent(const FMoverSimulationEventData& EventData)
 {
 	// On a mode change call deactivate on the previous mode and activate on the new mode,
@@ -2317,6 +2398,14 @@ void UMoverComponent::ProcessSimulationEvent(const FMoverSimulationEventData& Ev
 
 		OnMovementModeChanged.Broadcast(ModeChangedData->PreviousModeName, ModeChangedData->NewModeName);
 	}
+	else if (const FTeleportSucceededEventData* TeleportSucceededEventData = EventData.CastTo<FTeleportSucceededEventData>())
+	{
+		OnTeleportSucceeded.Broadcast(TeleportSucceededEventData->FromLocation, TeleportSucceededEventData->FromRotation, TeleportSucceededEventData->ToLocation, TeleportSucceededEventData->ToRotation);
+	}
+	else if (const FTeleportFailedEventData* TeleportFailedEventData = EventData.CastTo<FTeleportFailedEventData>())
+	{
+		OnTeleportFailed.Broadcast(TeleportFailedEventData->FromLocation, TeleportFailedEventData->FromRotation, TeleportFailedEventData->ToLocation, TeleportFailedEventData->ToRotation, TeleportFailedEventData->TeleportFailureReason);
+	}
 }
 
 void UMoverComponent::SetAdditionalSimulationOutput(const FMoverDataCollection& Data)
diff --git a/Engine/Plugins/Experimental/Mover/Source/Mover/Private/MoverSimulationTypes.cpp b/Engine/Plugins/Experimental/Mover/Source/Mover/Private/MoverSimulationTypes.cpp
--- a/Engine/Plugins/Experimental/Mover/Source/Mover/Private/MoverSimulationTypes.cpp
+++ b/Engine/Plugins/Experimental/Mover/Source/Mover/Private/MoverSimulationTypes.cpp
@@ -1,6 +1,10 @@
 // Copyright Epic Games, Inc. All Rights Reserved.
 
 #include "MoverSimulationTypes.h"
+#include "Physics/Experimental/PhysScene_Chaos.h"
+#include "PBDRigidsSolver.h"
+#include "GameFramework/GameStateBase.h"
+#include "Physics/NetworkPhysicsComponent.h"
 
 #include UE_INLINE_GENERATED_CPP_BY_NAME(MoverSimulationTypes)
 
@@ -10,6 +14,61 @@ UScriptStruct* FMoverSimulationEventData::GetScriptStruct() const
 	return FMoverSimulationEventData::StaticStruct();
 }
 
+FScheduledInstantMovementEffect FScheduledInstantMovementEffect::ScheduleEffect_External(UWorld* World, TSharedPtr<FInstantMovementEffect> InstantMovementEffect, float SchedulingDelaySeconds)
+{
+	ensure(IsInGameThread());
+	FPhysScene* Scene = World ? World->GetPhysicsScene() : nullptr;
+	Chaos::FPhysicsSolver* Solver = Scene ? Scene->GetSolver() : nullptr;
+	bool bIsFixedDt = Solver ? Solver->IsUsingFixedDt() : false;
+	if (bIsFixedDt)
+	{
+		int32 ServerFrame = UE::NetworkPhysicsUtils::GetUpcomingServerFrame_External(World);
+		if (SchedulingDelaySeconds != 0.0f)
+		{
+			const double AsyncDeltaTime = Solver->GetAsyncDeltaTime();
+			const int32 DelayFrames = (AsyncDeltaTime > UE_SMALL_NUMBER) ? FMath::CeilToInt32(SchedulingDelaySeconds / AsyncDeltaTime) : 0;
+			ServerFrame += DelayFrames;
+		}
+
+		return FScheduledInstantMovementEffect(/* ExecutionServerFrame = */ ServerFrame, /* ExecutionServerTime = */ 0.0, /* bIsFixedDt = */ true, InstantMovementEffect);
+	}
+	else
+	{
+		AGameStateBase* GameState = World ? World->GetGameState() : nullptr;
+		double ServerTime = GameState ? GameState->GetServerWorldTimeSeconds() : 0.0;
+		ServerTime += SchedulingDelaySeconds;
+		
+		return FScheduledInstantMovementEffect(/* ExecutionServerFrame = */ INDEX_NONE, /* ExecutionServerTime = */ ServerTime, /* bIsFixedDt = */ false, InstantMovementEffect);
+	}
+}
+
+FScheduledInstantMovementEffect FScheduledInstantMovementEffect::ScheduleEffect_Internal(UWorld* World, const FMoverTimeStep& TimeStep, TSharedPtr<FInstantMovementEffect> InstantMovementEffect, float SchedulingDelaySeconds)
+{
+	ensure(!IsInGameThread());
+	FPhysScene* Scene = World ? World->GetPhysicsScene() : nullptr;
+	Chaos::FPhysicsSolver* Solver = Scene ? Scene->GetSolver() : nullptr;
+	bool bIsFixedDt = Solver ? Solver->IsUsingFixedDt() : false;
+	
+	if (bIsFixedDt)
+	{
+		int32 ServerFrame = TimeStep.ServerFrame;
+		if (SchedulingDelaySeconds != 0.0f)
+		{
+			const double AsyncDeltaTime = Solver->GetAsyncDeltaTime();
+			const int32 DelayFrames = (AsyncDeltaTime > UE_SMALL_NUMBER) ? FMath::CeilToInt32(SchedulingDelaySeconds / AsyncDeltaTime) : 0;
+			ServerFrame += DelayFrames;
+		}
+		return FScheduledInstantMovementEffect(/* ExecutionServerFrame = */ ServerFrame, /* ExecutionServerTime = */ 0.0, /* bIsFixedDt = */ true, InstantMovementEffect);
+	}
+	else
+	{
+		AGameStateBase* GameState = World ? World->GetGameState() : nullptr;
+		double ServerTime = GameState ? GameState->GetServerWorldTimeSeconds() : 0.0;
+		ServerTime += SchedulingDelaySeconds;
+		return FScheduledInstantMovementEffect(/* ExecutionServerFrame = */ INDEX_NONE, /* ExecutionServerTime = */ ServerTime, /* bIsFixedDt = */ false, InstantMovementEffect);
+	}
+}
+
 namespace UE::Mover
 {
 	void FSimulationOutputData::Reset()
diff --git a/Engine/Plugins/Experimental/Mover/Source/Mover/Private/PhysicsMover/Modes/PhysicsDrivenSwimmingMode.cpp b/Engine/Plugins/Experimental/Mover/Source/Mover/Private/PhysicsMover/Modes/PhysicsDrivenSwimmingMode.cpp
--- a/Engine/Plugins/Experimental/Mover/Source/Mover/Private/PhysicsMover/Modes/PhysicsDrivenSwimmingMode.cpp
+++ b/Engine/Plugins/Experimental/Mover/Source/Mover/Private/PhysicsMover/Modes/PhysicsDrivenSwimmingMode.cpp
@@ -101,7 +101,7 @@ void UPhysicsDrivenSwimmingMode::SimulationTick_Implementation(const FSimulation
 
 	const float DeltaSeconds = Params.TimeStep.StepMs * 0.001f;
 
-	if (CharacterInputs && CharacterInputs->bIsJumpJustPressed && AttemptJump(SurfaceSwimmingWaterControlSettings.JumpMultiplier*CommonLegacySettings->JumpUpwardsSpeed, OutputState))
+	if (CharacterInputs && CharacterInputs->bIsJumpJustPressed && AttemptJump(Params, SurfaceSwimmingWaterControlSettings.JumpMultiplier*CommonLegacySettings->JumpUpwardsSpeed, OutputState))
 	{
 		OutputState.MovementEndState.RemainingMs = Params.TimeStep.StepMs;
 		return;
@@ -187,12 +187,14 @@ void UPhysicsDrivenSwimmingMode::SimulationTick_Implementation(const FSimulation
 	OutputState.MovementEndState.RemainingMs = 0.0f;
 }
 
-bool UPhysicsDrivenSwimmingMode::AttemptJump(float UpwardsSpeed, FMoverTickEndData& OutputState)
+bool UPhysicsDrivenSwimmingMode::AttemptJump(const FSimulationTickParams& Params, float UpwardsSpeed, FMoverTickEndData& OutputState)
 {
 	// TODO: This should check if a jump is even allowed
  	TSharedPtr<FJumpImpulseEffect> JumpMove = MakeShared<FJumpImpulseEffect>();
 	JumpMove->UpwardsSpeed = UpwardsSpeed;
-	GetMoverComponent()->QueueInstantMovementEffect(JumpMove);
+
+	// TODO: Use UMoverSimulation to queue instant movement effects when it is passed as part of FSimulationTickParams
+	GetMoverComponent()->QueueInstantMovementEffect_Internal(Params.TimeStep, JumpMove);
 	
 	return true;
 }
diff --git a/Engine/Plugins/Experimental/Mover/Source/Mover/Private/PhysicsMover/Transitions/PhysicsJumpCheck.cpp b/Engine/Plugins/Experimental/Mover/Source/Mover/Private/PhysicsMover/Transitions/PhysicsJumpCheck.cpp
--- a/Engine/Plugins/Experimental/Mover/Source/Mover/Private/PhysicsMover/Transitions/PhysicsJumpCheck.cpp
+++ b/Engine/Plugins/Experimental/Mover/Source/Mover/Private/PhysicsMover/Transitions/PhysicsJumpCheck.cpp
@@ -58,7 +58,7 @@ void UPhysicsJumpCheck::Trigger_Implementation(const FSimulationTickParams& Para
 	{
 		JumpMove->VelocityToApply = FVector(JumpUpwardsSpeed * MoverComp->GetUpDirection());
 		JumpMove->bAdditiveVelocity = true;
-		Params.MovingComps.MoverComponent->QueueInstantMovementEffect(JumpMove);
+		Params.MovingComps.MoverComponent->QueueInstantMovementEffect_Internal(Params.TimeStep, JumpMove);
 
 		if (FractionalGroundReactionImpulse > 0.0f)
 		{
diff --git a/Engine/Plugins/Experimental/Mover/Source/Mover/Private/PhysicsMover/Transitions/PhysicsLaunchCheck.cpp b/Engine/Plugins/Experimental/Mover/Source/Mover/Private/PhysicsMover/Transitions/PhysicsLaunchCheck.cpp
--- a/Engine/Plugins/Experimental/Mover/Source/Mover/Private/PhysicsMover/Transitions/PhysicsLaunchCheck.cpp
+++ b/Engine/Plugins/Experimental/Mover/Source/Mover/Private/PhysicsMover/Transitions/PhysicsLaunchCheck.cpp
@@ -51,7 +51,7 @@ void UPhysicsLaunchCheck::Trigger_Implementation(const FSimulationTickParams& Pa
 			ApplyVelocityEffect->bAdditiveVelocity = LaunchInputs->Mode == EMoverLaunchVelocityMode::Additive;
 			ApplyVelocityEffect->VelocityToApply = LaunchInputs->LaunchVelocity;
 			ApplyVelocityEffect->ForceMovementMode = TransitionToMode;
-			MoverComp->QueueInstantMovementEffect(ApplyVelocityEffect);
+			MoverComp->QueueInstantMovementEffect_Internal(Params.TimeStep, ApplyVelocityEffect);
 		}
 	}
 }
diff --git a/Engine/Plugins/Experimental/Mover/Source/Mover/Public/Backends/MoverBackendLiaison.h b/Engine/Plugins/Experimental/Mover/Source/Mover/Public/Backends/MoverBackendLiaison.h
--- a/Engine/Plugins/Experimental/Mover/Source/Mover/Public/Backends/MoverBackendLiaison.h
+++ b/Engine/Plugins/Experimental/Mover/Source/Mover/Public/Backends/MoverBackendLiaison.h
@@ -32,6 +32,9 @@ public:
 	// Whether this backend will simulate movement asynchronously
 	virtual bool IsAsync() const { return false; }
 
+	// How much delay to apply to scheduled events. This is important for networked events, and should be greater than the RTT to ensure the event will be executed on all end points at the same frame.
+	virtual float GetEventSchedulingMinDelaySeconds() const {return 0.3f;}
+
 	// Pending State: the simulation state currently being authored
 	virtual bool ReadPendingSyncState(OUT FMoverSyncState& OutSyncState) { return false; }
 	virtual bool WritePendingSyncState(const FMoverSyncState& SyncStateToWrite) { return false; }
diff --git a/Engine/Plugins/Experimental/Mover/Source/Mover/Public/DefaultMovementSet/InstantMovementEffects/BasicInstantMovementEffects.h b/Engine/Plugins/Experimental/Mover/Source/Mover/Public/DefaultMovementSet/InstantMovementEffects/BasicInstantMovementEffects.h
--- a/Engine/Plugins/Experimental/Mover/Source/Mover/Public/DefaultMovementSet/InstantMovementEffects/BasicInstantMovementEffects.h
+++ b/Engine/Plugins/Experimental/Mover/Source/Mover/Public/DefaultMovementSet/InstantMovementEffects/BasicInstantMovementEffects.h
@@ -32,6 +32,7 @@ struct FTeleportEffect : public FInstantMovementEffect
 	FRotator TargetRotation;
 	
 	UE_API virtual bool ApplyMovementEffect(FApplyMovementEffectParams& ApplyEffectParams, FMoverSyncState& OutputState) override;
+	UE_API virtual bool ApplyMovementEffect_Async(FApplyMovementEffectParams_Async& ApplyEffectParams, FMoverSyncState& OutputState) override;
 
 	UE_API virtual FInstantMovementEffect* Clone() const override;
 
diff --git a/Engine/Plugins/Experimental/Mover/Source/Mover/Public/DefaultMovementSet/Modes/SwimmingMode.h b/Engine/Plugins/Experimental/Mover/Source/Mover/Public/DefaultMovementSet/Modes/SwimmingMode.h
--- a/Engine/Plugins/Experimental/Mover/Source/Mover/Public/DefaultMovementSet/Modes/SwimmingMode.h
+++ b/Engine/Plugins/Experimental/Mover/Source/Mover/Public/DefaultMovementSet/Modes/SwimmingMode.h
@@ -209,7 +209,7 @@ protected:
 	UE_API virtual void OnRegistered(const FName ModeName) override; 
 	UE_API virtual void OnUnregistered() override;
 
-	UE_API virtual bool AttemptJump(float UpwardsSpeed, FMoverTickEndData& Output);
+	UE_API virtual bool AttemptJump(const FSimulationTickParams& Params, float UpwardsSpeed, FMoverTickEndData& Output);
 
 	UE_API void CaptureFinalState(USceneComponent* UpdatedComponent, const FFloorCheckResult& FloorResult, const FMovementRecord& Record, FMoverDefaultSyncState& OutputSyncState) const;
 
diff --git a/Engine/Plugins/Experimental/Mover/Source/Mover/Public/InstantMovementEffect.h b/Engine/Plugins/Experimental/Mover/Source/Mover/Public/InstantMovementEffect.h
--- a/Engine/Plugins/Experimental/Mover/Source/Mover/Public/InstantMovementEffect.h
+++ b/Engine/Plugins/Experimental/Mover/Source/Mover/Public/InstantMovementEffect.h
@@ -3,6 +3,7 @@
 #pragma once
 
 #include "MoveLibrary/MovementUtilsTypes.h"
+
 #include "InstantMovementEffect.generated.h"
 
 #define UE_API MOVER_API
@@ -12,6 +13,7 @@ class UMoverSimulation;
 struct FMoverTimeStep;
 struct FMoverTickStartData;
 struct FMoverSyncState;
+struct FMoverSimulationEventData;
 
 struct FApplyMovementEffectParams
 {
@@ -24,6 +26,8 @@ struct FApplyMovementEffectParams
 	const FMoverTickStartData* StartState;
 
 	const FMoverTimeStep* TimeStep;
+
+	TArray<TSharedPtr<FMoverSimulationEventData>> OutputEvents;
 };
 
 /** 
diff --git a/Engine/Plugins/Experimental/Mover/Source/Mover/Public/MovementModeStateMachine.h b/Engine/Plugins/Experimental/Mover/Source/Mover/Public/MovementModeStateMachine.h
--- a/Engine/Plugins/Experimental/Mover/Source/Mover/Public/MovementModeStateMachine.h
+++ b/Engine/Plugins/Experimental/Mover/Source/Mover/Public/MovementModeStateMachine.h
@@ -54,16 +54,21 @@ public:
 
 	UE_API void QueueLayeredMove(TSharedPtr<FLayeredMoveBase> Move);
 	
- 	UE_API void QueueInstantMovementEffect(TSharedPtr<FInstantMovementEffect> Effect);
-
-	UE_API FMovementModifierHandle QueueMovementModifier(TSharedPtr<FMovementModifierBase> Modifier);
+ 	UE_API FMovementModifierHandle QueueMovementModifier(TSharedPtr<FMovementModifierBase> Modifier);
 
  	UE_API void CancelModifierFromHandle(FMovementModifierHandle ModifierHandle);
 
  	UE_API const FMovementModifierBase* FindQueuedModifier(FMovementModifierHandle ModifierHandle) const;
  	UE_API const FMovementModifierBase* FindQueuedModifierByType(const UScriptStruct* ModifierType) const;
- 	
+
+	// This function is meant to be used only in async mode on the physics thread, not on the game thread
+	void QueueInstantMovementEffect_Internal(const FScheduledInstantMovementEffect& ScheduledInstantMovementEffect);
 protected:
+	UE_API void QueueInstantMovementEffect(const FScheduledInstantMovementEffect& ScheduledInstantMovementEffect);
+	UE_API void QueueInstantMovementEffects(const TArray<FScheduledInstantMovementEffect>& ScheduledInstantMovementEffects);
+
+	void ProcessEvents(const TArray<TSharedPtr<FMoverSimulationEventData>>& InEvents);
+	UE_API virtual void ProcessSimulationEvent(const FMoverSimulationEventData& EventData);
 
 	UE_API virtual void PostInitProperties() override;
 
@@ -81,7 +86,7 @@ protected:
 	TArray<TSharedPtr<FLayeredMoveBase>> QueuedLayeredMoves;
 
  	/** Effects that are queued to be applied to the simulation at the start of the next sim subtick or at the end of this tick.  Access covered by lock. */
- 	TArray<TSharedPtr<FInstantMovementEffect>> QueuedInstantEffects;
+ 	TArray<FScheduledInstantMovementEffect> QueuedInstantEffects;
 
  	/** Modifiers that are queued to be added to the simulation at the start of the next sim subtick. Access covered by lock. */
  	TArray<TSharedPtr<FMovementModifierBase>> QueuedMovementModifiers;
diff --git a/Engine/Plugins/Experimental/Mover/Source/Mover/Public/MoverComponent.h b/Engine/Plugins/Experimental/Mover/Source/Mover/Public/MoverComponent.h
--- a/Engine/Plugins/Experimental/Mover/Source/Mover/Public/MoverComponent.h
+++ b/Engine/Plugins/Experimental/Mover/Source/Mover/Public/MoverComponent.h
@@ -19,6 +19,7 @@
 #include "Backends/MoverBackendLiaison.h"
 #include "UObject/WeakInterfacePtr.h"
 #include "Templates/SharedPointer.h"
+
 #include "MoverComponent.generated.h"
 
 struct FMoverTimeStep;
@@ -51,6 +52,12 @@ DECLARE_DYNAMIC_MULTICAST_DELEGATE_TwoParams(FMover_OnPostSimRollback, const FMo
 // Fired after changing movement modes. First param is the name of the previous movement mode. Second is the name of the new movement mode. 
 DECLARE_DYNAMIC_MULTICAST_DELEGATE_TwoParams(FMover_OnMovementModeChanged, const FName&, PreviousMovementModeName, const FName&, NewMovementModeName);
 
+// Fired when a teleport has succeeded
+DECLARE_DYNAMIC_MULTICAST_DELEGATE_FourParams(FMover_OnTeleportSucceeded, const FVector&, FromLocation, const FQuat&, FromRotation, const FVector&, ToLocation, const FQuat&, ToRotation);
+
+// Fired when a teleport has failed
+DECLARE_DYNAMIC_MULTICAST_DELEGATE_FiveParams(FMover_OnTeleportFailed, const FVector&, FromLocation, const FQuat&, FromRotation, const FVector&, ToLocation, const FQuat&, ToRotation, ETeleportFailureReason, TeleportFailureReason);
+
 // Fired after a transition has been triggered.
 DECLARE_DYNAMIC_MULTICAST_DELEGATE_OneParam(FMover_OnTransitionTriggered, const UBaseMovementModeTransition*, ModeTransition);
 
@@ -104,6 +111,14 @@ public:
 	UPROPERTY(BlueprintAssignable, Category = Mover)
 	FMover_OnMovementModeChanged OnMovementModeChanged;
 
+	// Broadcast when a teleport has succeeded
+	UPROPERTY(BlueprintAssignable, Category = Mover)
+	FMover_OnTeleportSucceeded OnTeleportSucceeded;
+
+	// Broadcast when a teleport has failed
+	UPROPERTY(BlueprintAssignable, Category = Mover)
+	FMover_OnTeleportFailed OnTeleportFailed;
+
 	// Broadcast when a Transition has been triggered.
 	UPROPERTY(BlueprintAssignable, Category = Mover)
 	FMover_OnTransitionTriggered OnMovementTransitionTriggered;
@@ -236,16 +251,46 @@ public:
 	MOVER_API void CancelModifierFromHandle(FMovementModifierHandle ModifierHandle);
 	
 	/**
-	 * Queue a Instant Movement Effect to start at the end of this frame or start of the next subtick - whichever happens first. This will clone whatever move you pass in, so you'll need to fully set it up before queuing.
+	 * Queue an Instant Movement Effect to start at the end of this frame or start of the next subtick - whichever happens first. This will clone whatever move you pass in, so you'll need to fully set it up before queuing.
 	 * @param InstantMovementEffect			The effect to queue, which must be a FInstantMovementEffect sub-type. 
 	 */
 	UFUNCTION(BlueprintCallable, CustomThunk, Category = Mover, meta = (CustomStructureParam = "EffectAsRawData", AllowAbstract = "false", DisplayName = "Queue Instant Movement Effect"))
 	MOVER_API void K2_QueueInstantMovementEffect(UPARAM(DisplayName="Instant Movement Effect") const int32& EffectAsRawData);
 	DECLARE_FUNCTION(execK2_QueueInstantMovementEffect);
 
-	// Queue a Instant Movement Effect to take place at the end of this frame or start of the next subtick - whichever happens first
-	MOVER_API void QueueInstantMovementEffect(TSharedPtr<FInstantMovementEffect> Move);
-	
+	/**
+	 * Schedule an Instant Movement Effect to be applied as early as possible while ensuring it gets executed on the same frame on all networked end points.
+	 * This adds a delay to the application of the effect, tunable in the NetworkPhysicsSettingsComponent. @see UNetworkPhysicsSettingsComponent, @see FNetworkPhysicsSettings, @see EventSchedulingMinDelaySeconds
+	 * This will clone whatever move you pass in, so you'll need to fully set it up before queuing.
+	 * @param InstantMovementEffect			The effect to queue, which must be a FInstantMovementEffect sub-type. 
+	 */
+	UFUNCTION(BlueprintCallable, CustomThunk, Category = Mover, meta = (CustomStructureParam = "EffectAsRawData", AllowAbstract = "false", DisplayName = "Schedule Instant Movement Effect"))
+	MOVER_API void K2_ScheduleInstantMovementEffect(UPARAM(DisplayName="Instant Movement Effect") const int32& EffectAsRawData);
+	DECLARE_FUNCTION(execK2_ScheduleInstantMovementEffect);
+
+	/**
+	 *  Queue a Instant Movement Effect to take place at the end of this frame or start of the next subtick - whichever happens first
+	 *  @param InstantMovementEffect			The effect to queue, which must be a FInstantMovementEffect sub - type.
+	 */ 
+	MOVER_API void QueueInstantMovementEffect(TSharedPtr<FInstantMovementEffect> InstantMovementEffect);
+	/** 
+	 * Queue a scheduled Instant Movement Effect to take place after delay (tunable in the NetworkPhysicsSettingsComponent)
+	 * ensuring it gets executed on the same frame on all networked end points. @see UNetworkPhysicsSettingsComponent, @see FNetworkPhysicsSettings, @see EventSchedulingMinDelaySeconds
+	 * @param InstantMovementEffect			The effect to queue, which must be a FInstantMovementEffect sub - type.
+	 */
+	MOVER_API void ScheduleInstantMovementEffect(TSharedPtr<FInstantMovementEffect> InstantMovementEffect);
+
+	// Get the queued instant movement effects. This is mostly for internal use, general users should abstain from calling it.
+	MOVER_API const TArray<FScheduledInstantMovementEffect>& GetQueuedInstantMovementEffects() const;
+	// Clears the queued instant movement effects. This is mostly for internal use, general users should abstain from calling it.
+	MOVER_API void ClearQueuedInstantMovementEffects();
+	// Queue an instant movement effect in async mode. Do not use on the game thread.
+	void QueueInstantMovementEffect_Internal(const FMoverTimeStep& TimeStep, TSharedPtr<FInstantMovementEffect> InstantMovementEffect);
+protected:
+	// Queue a scheduled instant movement effect. Thread safe, can be used outside the game thread.
+	void QueueInstantMovementEffect(const FScheduledInstantMovementEffect& ScheduledInstantMovementEffect);
+
+public:	
 	// Queue a movement mode change to occur during the next simulation frame. If bShouldReenter is true, then a mode change will occur even if already in that mode.
 	UFUNCTION(BlueprintCallable, Category = Mover, DisplayName="Queue Next Movement Mode")
 	MOVER_API void QueueNextMode(FName DesiredModeName, bool bShouldReenter=false);
@@ -673,12 +718,13 @@ private:
 	UPROPERTY(EditDefaultsOnly, Category="Mover|Gravity", meta=(ForceUnits = "cm/s^2"))
 	FVector GravityAccelOverride;
 
-private:
-
 	/** Settings that can lock movement to a particular plane */
 	UPROPERTY(EditDefaultsOnly, Category = "Mover|Constraints")
 	FPlanarConstraint PlanarConstraint;
 
+	/** Effects queued to be applied to the simulation at a given frame. If the frame happens to be in the past, the effect will be applied at the earliest occasion */
+	TArray<FScheduledInstantMovementEffect> QueuedInstantMovementEffects;
+
 public:
 
 	// If enabled, the movement of the primary visual component will be smoothed via an offset from the root moving component. This is useful in fixed-tick simulations with variable rendering rates.
@@ -687,6 +733,9 @@ public:
 
 	MOVER_API void SetSimulationOutput(const FMoverTimeStep& TimeStep, const UE::Mover::FSimulationOutputData& OutputData);
 
+	// Dispatch a simulation event. It will be processed immediately.
+	MOVER_API void DispatchSimulationEvent(const FMoverSimulationEventData& EventData);
+
 protected:
 	MOVER_API virtual void ProcessSimulationEvent(const FMoverSimulationEventData& EventData);
 	MOVER_API virtual void SetAdditionalSimulationOutput(const FMoverDataCollection& Data);
@@ -736,6 +785,10 @@ protected:
 	UPROPERTY(EditDefaultsOnly, Category = "Mover")
 	bool bSupportsKinematicBasedMovement = true;
 
+	// Delay added to scheduled instant movement effects
+	// This value is cached from the settings found in the network settings component
+	float EventSchedulingMinDelaySeconds = 0.3f;
+
 	FMoverAuxStateContext CachedLastAuxState;
 
 	friend class UBaseMovementMode;
diff --git a/Engine/Plugins/Experimental/Mover/Source/Mover/Public/MoverSimulation.h b/Engine/Plugins/Experimental/Mover/Source/Mover/Public/MoverSimulation.h
--- a/Engine/Plugins/Experimental/Mover/Source/Mover/Public/MoverSimulation.h
+++ b/Engine/Plugins/Experimental/Mover/Source/Mover/Public/MoverSimulation.h
@@ -9,6 +9,7 @@
 #define UE_API MOVER_API
 
 class UMoverBlackboard;
+struct FMoverSyncState;
 
 /**
 * WIP Base class for a Mover simulation.
@@ -29,6 +30,17 @@ public:
 	UFUNCTION(BlueprintPure, Category = Mover)
 	UE_API UMoverBlackboard* GetBlackboard_Mutable();
 
+	/**
+	* Attempt to teleport to TargetTransform. The teleport is not guaranteed to happen. This function is meant to be called by an instant movement effect as part of its effect application.
+	* If it succeeds a FTeleportSucceededEventData will be emitted, if it fails a FTeleportFailedEventData will be sent.
+	* @param TimeStep The time step of the current step or substep being simulated. This will come from the ApplyMovementEffect function.
+	* @param TargetTransform The transform to teleport to. In the case bUseActorRotation is true, the rotation of this transform will be ignored.
+	* @param bUseActorRotation If true, the rotation will not be modified upon teleportation. If false, the rotation in TargetTransform will be used to orient the teleported.
+	* @param OutputState This is the sync state that me modified as a result of the application of this effect. Like TimeStep, this should come from the ApplyMovementEffect function.
+	*/
+	UFUNCTION(BlueprintCallable, Category = Mover)
+	UE_API virtual void AttemptTeleport(const FMoverTimeStep& TimeStep, const FTransform& TargetTransform, bool bUseActorRotation, FMoverSyncState& OutputState) {}
+
 protected:
 	UPROPERTY(Transient)
 	TObjectPtr<UMoverBlackboard> Blackboard = nullptr;
diff --git a/Engine/Plugins/Experimental/Mover/Source/Mover/Public/MoverSimulationTypes.h b/Engine/Plugins/Experimental/Mover/Source/Mover/Public/MoverSimulationTypes.h
--- a/Engine/Plugins/Experimental/Mover/Source/Mover/Public/MoverSimulationTypes.h
+++ b/Engine/Plugins/Experimental/Mover/Source/Mover/Public/MoverSimulationTypes.h
@@ -10,6 +10,7 @@
 #include "LayeredMove.h"
 #include "MovementModifier.h"
 #include "MoverDataModelTypes.h"
+#include "InstantMovementEffect.h"
 #include "UObject/Interface.h"
 #include <functional>
 
@@ -68,8 +69,86 @@ struct FMovementModeTickEndState
 
 };
 
+USTRUCT()
+struct FScheduledInstantMovementEffect
+{
+	GENERATED_BODY()
+
+	/** Turns a FInstantMovementEffect into a scheduled one (FScheduledInstantMovementEffect)
+	*	The effect can be scheduled to apply immediately, or scheduled to apply with a delay
+	*   This function can only be called on the game thread
+	*   @param World The world, used to retrieve the current server frame in async mode, or the sim time otherwise
+	*   @param InstantMovementEffect the effect to schedule
+	*   @param SchedulingDelaySeconds Scheduling delay to ensure it applies on all end points on the same frame (this is only perfectly accurate when simulation dt is fixed)
+	*/
+	static FScheduledInstantMovementEffect ScheduleEffect_External(UWorld* World, TSharedPtr<FInstantMovementEffect> InstantMovementEffect, float SchedulingDelaySeconds = 0.0f);
+
+	/** Turns a FInstantMovementEffect into a scheduled one (FScheduledInstantMovementEffect)
+	*	The effect can be scheduled to apply immediately, or scheduled to apply with a delay
+	*   This function should not be called on the game thread
+	*   @param World The world, used to retrieve the current server frame in async mode, or the sim time otherwise
+	*   @param TimeStep the time step used in the async update which this function should be called from
+	*   @param InstantMovementEffect the effect to schedule
+	*   @param SchedulingDelaySeconds Scheduling delay to ensure it applies on all end points on the same frame (this is only perfectly accurate when simulation dt is fixed)
+	*/
+	static FScheduledInstantMovementEffect ScheduleEffect_Internal(UWorld* World, const FMoverTimeStep& TimeStep, TSharedPtr<FInstantMovementEffect> InstantMovementEffect, float SchedulingDelaySeconds = 0.0f);
+
+	bool ShouldExecuteAtFrame(int32 CurrentServerFrame) const
+	{
+		ensureMsgf(bIsFixedDt, TEXT("In variable delta time mode, use the version of ShouldExecute that takes a floating point time"));
+		return (CurrentServerFrame >= ExecutionServerFrame);
+	}
 
+	bool ShouldExecuteAtTime(double CurrentServerTime) const
+	{
+		ensureMsgf(!bIsFixedDt, TEXT("In fixed delta time mode, use the version of ShouldExecute that takes a frame number"));
+		return (CurrentServerTime >= ExecutionServerTimeSeconds);
+	}
 
+	void NetSerialize(const FNetSerializeParams& P)
+	{
+		P.Ar.SerializeBits(&bIsFixedDt, 1);
+		if (bIsFixedDt)
+		{
+			P.Ar << ExecutionServerFrame;
+		}
+		else
+		{
+			P.Ar << ExecutionServerTimeSeconds;
+		}
+		
+		Effect->NetSerialize(P.Ar);
+	}
+
+	void ToString(FAnsiStringBuilderBase& Out) const
+	{
+		if (bIsFixedDt)
+		{
+			Out.Appendf("ExecutionServerFrame: %d", ExecutionServerFrame);
+		}
+		else
+		{
+			Out.Appendf("ExecutionDateSeconds: %f", ExecutionServerTimeSeconds);
+		}
+
+		Out.Appendf(" | Effect = % s", *(Effect.IsValid() ? Effect->ToSimpleString() : "Invalid"));
+	}
+
+	// Server frame at which this instant movement effect should be applied
+	// Only valid if bIsFixedDt is true, i.e. in fixed time step mode
+	UPROPERTY(VisibleAnywhere, Category = "Mover")
+	int32 ExecutionServerFrame = INDEX_NONE;
+
+	// Server Time (in seconds) after which this instant movement effect should be applied
+	// Only valid if bIsFixedDt is false, i.e. in variable time step mode
+	UPROPERTY(VisibleAnywhere, Category = "Mover")
+	double ExecutionServerTimeSeconds = 0.0;
+
+	UPROPERTY(VisibleAnywhere, Category = "Mover")
+	bool bIsFixedDt = true;
+
+	TSharedPtr<FInstantMovementEffect> Effect;
+};
 
 /**
  * The client generates this representation of "input" to the simulated actor for one simulation frame. This can be direct mapping
@@ -80,11 +159,13 @@ struct FMoverInputCmdContext
 {
 	GENERATED_BODY()
 
-public:
 	UPROPERTY(BlueprintReadWrite, Category = Mover)
 	FMoverDataCollection InputCollection;
 
-	UScriptStruct* GetStruct() const { return StaticStruct(); }
+	UScriptStruct* GetStruct() const
+	{
+		return StaticStruct();
+	}
 
 	void NetSerialize(const FNetSerializeParams& P)
 	{
@@ -489,6 +570,66 @@ struct FMovementModeChangedEventData : public FMoverSimulationEventData
 	FName NewModeName = NAME_None;
 };
 
+USTRUCT()
+struct FTeleportSucceededEventData : public FMoverSimulationEventData
+{
+	GENERATED_BODY()
+
+	FTeleportSucceededEventData(float InEventTimeMs, const FVector& InFromLocation, const FQuat& InFromRotation, const FVector& InToLocation, const FQuat& InToRotation)
+		: FMoverSimulationEventData(InEventTimeMs)
+		, FromLocation(InFromLocation)
+		, FromRotation(InFromRotation)
+		, ToLocation(InToLocation)
+		, ToRotation(InToRotation)
+	{
+	}
+	FTeleportSucceededEventData() {}
+
+	virtual UScriptStruct* GetScriptStruct() const override
+	{
+		return FTeleportSucceededEventData::StaticStruct();
+	}
+
+	FVector FromLocation;
+	FQuat FromRotation;
+	FVector ToLocation;
+	FQuat ToRotation;
+};
+
+UENUM(BlueprintType)
+enum class ETeleportFailureReason : uint8
+{
+	Reason_NotAvailable UMETA(DisplayName = "Reason Not Available", Tooltip = "A reason for the teleport failure was not indicated"),
+};
+
+USTRUCT()
+struct FTeleportFailedEventData : public FMoverSimulationEventData
+{
+	GENERATED_BODY()
+
+	FTeleportFailedEventData(float InEventTimeMs, const FVector& InFromLocation, const FQuat& InFromRotation, const FVector& InToLocation, const FQuat& InToRotation, ETeleportFailureReason InTeleportFailureReason)
+		: FMoverSimulationEventData(InEventTimeMs)
+		, FromLocation(InFromLocation)
+		, FromRotation(InFromRotation)
+		, ToLocation(InToLocation)
+		, ToRotation(InToRotation)
+		, TeleportFailureReason(InTeleportFailureReason)
+	{
+	}
+	FTeleportFailedEventData() {}
+
+	virtual UScriptStruct* GetScriptStruct() const override
+	{
+		return FTeleportFailedEventData::StaticStruct();
+	}
+
+	FVector FromLocation;
+	FQuat FromRotation;
+	FVector ToLocation;
+	FQuat ToRotation;
+	ETeleportFailureReason TeleportFailureReason;
+};
+
 namespace UE::Mover
 {
 	struct FSimulationOutputData
diff --git a/Engine/Plugins/Experimental/Mover/Source/Mover/Public/PhysicsMover/Modes/PhysicsDrivenSwimmingMode.h b/Engine/Plugins/Experimental/Mover/Source/Mover/Public/PhysicsMover/Modes/PhysicsDrivenSwimmingMode.h
--- a/Engine/Plugins/Experimental/Mover/Source/Mover/Public/PhysicsMover/Modes/PhysicsDrivenSwimmingMode.h
+++ b/Engine/Plugins/Experimental/Mover/Source/Mover/Public/PhysicsMover/Modes/PhysicsDrivenSwimmingMode.h
@@ -24,7 +24,7 @@ class UPhysicsDrivenSwimmingMode : public USwimmingMode, public IPhysicsCharacte
 public:
 	UE_API virtual void OnRegistered(const FName ModeName) override;
 	UE_API virtual void SimulationTick_Implementation(const FSimulationTickParams& Params, FMoverTickEndData& OutputState) override;
-	UE_API virtual bool AttemptJump(float UpwardsSpeed, FMoverTickEndData& Output) override;
+	UE_API virtual bool AttemptJump(const FSimulationTickParams& Params, float UpwardsSpeed, FMoverTickEndData& Output) override;
 
 	UE_API virtual void UpdateConstraintSettings(Chaos::FCharacterGroundConstraint& Constraint) const override;
 
