Mover: adding options for handling out-of-band movement changes, whether to warn about it and whether to accept it as the new starting sim state. Usage is primarily for non-networked play, where an external system is allowed to move our actor around, such as during cinematics or temporary physics simulation.

#virtualized

diff --git a/Engine/Plugins/Experimental/Mover/Source/Mover/Private/MoverComponent.cpp b/Engine/Plugins/Experimental/Mover/Source/Mover/Private/MoverComponent.cpp
--- a/Engine/Plugins/Experimental/Mover/Source/Mover/Private/MoverComponent.cpp
+++ b/Engine/Plugins/Experimental/Mover/Source/Mover/Private/MoverComponent.cpp
@@ -517,17 +517,7 @@ void UMoverComponent::SimulationTick(const FMoverTimeStep& InTimeStep, const FMo
 		return;
 	}
 
-	if (const FMoverDefaultSyncState* StartingSyncState = SimInput.SyncState.SyncStateCollection.FindDataByType<FMoverDefaultSyncState>())
-	{
-		if (!(UpdatedComponent->GetComponentLocation().Equals(StartingSyncState->GetLocation_WorldSpace()) || StartingSyncState->GetMovementBase()))
-		{
-			UE_LOG(LogMover, Warning, TEXT("%s %s: Simulation start location (%s) disagrees with actual mover component location (%s). This indicates movement of the component out-of-band with the simulation, and may cause poor quality motion."),
-			*GetNameSafe(GetOwner()),
-			*StaticEnum<ENetRole>()->GetValueAsString(GetOwnerRole()),
-			*StartingSyncState->GetLocation_WorldSpace().ToCompactString(),
-			*UpdatedComponent->GetComponentLocation().ToCompactString());
-		}
-	}
+	CheckForExternalMovement(SimInput);
 
 	// Some sync state data should carry over between frames
 	for (const FMoverDataPersistence& PersistentSyncEntry : PersistentSyncStateDataTypes)
@@ -2413,4 +2403,45 @@ void UMoverComponent::SetAdditionalSimulationOutput(const FMoverDataCollection&
 
 }
 
+
+void UMoverComponent::CheckForExternalMovement(const FMoverTickStartData& SimStartingData)
+{
+	if (!bWarnOnExternalMovement && !bAcceptExternalMovement)
+	{
+		return;
+	}
+
+	if (const FMoverDefaultSyncState* StartingSyncState = SimStartingData.SyncState.SyncStateCollection.FindDataByType<FMoverDefaultSyncState>())
+	{		
+		if (StartingSyncState->GetMovementBase())
+		{
+			return;	// TODO: need alternative handling of movement checks when based on another object
+		}
+
+		const FTransform& ComponentTransform = UpdatedComponent->GetComponentTransform();
+
+		if (!ComponentTransform.GetLocation().Equals(StartingSyncState->GetLocation_WorldSpace()))
+		{
+			if (bWarnOnExternalMovement)
+			{
+				UE_LOG(LogMover, Warning, TEXT("%s %s: Simulation start location (%s) disagrees with actual mover component location (%s). This indicates movement of the component out-of-band with the simulation, and may cause poor quality motion."),
+					*GetNameSafe(GetOwner()),
+					*StaticEnum<ENetRole>()->GetValueAsString(GetOwnerRole()),
+					*StartingSyncState->GetLocation_WorldSpace().ToCompactString(),
+					*UpdatedComponent->GetComponentLocation().ToCompactString());
+			}
+
+			if (bAcceptExternalMovement)
+			{
+				FMoverDefaultSyncState* MutableSyncState = const_cast<FMoverDefaultSyncState*>(StartingSyncState);
+
+				MutableSyncState->SetTransforms_WorldSpace(ComponentTransform.GetLocation(), 
+				                                           ComponentTransform.GetRotation().Rotator(),
+				                                           MutableSyncState->GetVelocity_WorldSpace());
+			}
+		}
+	}
+}
+
+
 #undef LOCTEXT_NAMESPACE
diff --git a/Engine/Plugins/Experimental/Mover/Source/Mover/Public/MoverComponent.h b/Engine/Plugins/Experimental/Mover/Source/Mover/Public/MoverComponent.h
--- a/Engine/Plugins/Experimental/Mover/Source/Mover/Public/MoverComponent.h
+++ b/Engine/Plugins/Experimental/Mover/Source/Mover/Public/MoverComponent.h
@@ -731,6 +731,14 @@ public:
 	UPROPERTY(EditAnywhere, BlueprintReadWrite, Category = "Mover")
 	EMoverSmoothingMode SmoothingMode = EMoverSmoothingMode::VisualComponentOffset;
 
+	// Whether to warn when we detect that an external system has moved our object, outside of movement simulation control
+	UPROPERTY(EditAnywhere, BlueprintReadWrite, Category = "Mover", AdvancedDisplay)
+	uint8 bWarnOnExternalMovement : 1 = 1;
+
+	// If enabled, we'll accept any movements from an external system in the next simulation state update
+	UPROPERTY(EditAnywhere, BlueprintReadWrite, Category = "Mover", AdvancedDisplay)
+	uint8 bAcceptExternalMovement : 1 = 0;
+
 	MOVER_API void SetSimulationOutput(const FMoverTimeStep& TimeStep, const UE::Mover::FSimulationOutputData& OutputData);
 
 	// Dispatch a simulation event. It will be processed immediately.
@@ -739,6 +747,7 @@ public:
 protected:
 	MOVER_API virtual void ProcessSimulationEvent(const FMoverSimulationEventData& EventData);
 	MOVER_API virtual void SetAdditionalSimulationOutput(const FMoverDataCollection& Data);
+	MOVER_API virtual void CheckForExternalMovement(const FMoverTickStartData& SimStartingData);
 
 private:
 	/** Whether or not gravity is overridden on this actor. Otherwise, fall back on world settings. See @SetGravityOverride */
