[Mover] Make UpDirection overridable

diff --git a/Engine/Plugins/Experimental/Mover/Source/Mover/Private/MoverComponent.cpp b/Engine/Plugins/Experimental/Mover/Source/Mover/Private/MoverComponent.cpp
--- a/Engine/Plugins/Experimental/Mover/Source/Mover/Private/MoverComponent.cpp
+++ b/Engine/Plugins/Experimental/Mover/Source/Mover/Private/MoverComponent.cpp
@@ -1280,7 +1280,19 @@ bool UMoverComponent::ValidateSetup(FDataValidationContext& Context) const
 		}
 	}
 
-	
+	// Verify that the up direction override is a normalized vector
+	if (bHasUpDirectionOverride)
+	{
+		if (!UpDirectionOverride.IsNormalized())
+		{
+			Context.AddError(FText::Format(LOCTEXT("InvalidUpDirectionOverrideError", "UpDirectionOverride {0} needs to be a normalized vector, but it is not. {1}"),
+				FText::FromString(UpDirectionOverride.ToString()),
+				FText::FromString(GetNameSafe(GetOwner()))));
+
+			bDidFindAnyProblems = true;
+		}
+	}
+
 	return !bDidFindAnyProblems;
 }
 
@@ -1987,9 +1999,29 @@ FVector UMoverComponent::GetGravityAcceleration() const
 	return MoverComponentConstants::DefaultGravityAccel;
 }
 
+void UMoverComponent::SetUpDirectionOverride(bool bOverrideUpDirection, FVector UpDirection)
+{
+	bHasUpDirectionOverride = bOverrideUpDirection;
+	if (bOverrideUpDirection)
+	{
+		if (UpDirection.IsNearlyZero())
+		{
+			UE_LOG(LogMover, Warning, TEXT("Ignoring the provided UpDirection (%s) override because it is a zero vector. (%s)"), *UpDirection.ToString(), *GetNameSafe(GetOwner()));
+			bHasGravityOverride = false;
+			return;
+		}
+		UpDirectionOverride = UpDirection.GetSafeNormal();
+	}
+}
 
 FVector UMoverComponent::GetUpDirection() const
 {
+	// Use the up direction override if enabled
+	if (bHasUpDirectionOverride)
+	{
+		return UpDirectionOverride;
+	}
+	
 	return UMovementUtils::DeduceUpDirectionFromGravity(GetGravityAcceleration());
 }
 
diff --git a/Engine/Plugins/Experimental/Mover/Source/Mover/Public/MoverComponent.h b/Engine/Plugins/Experimental/Mover/Source/Mover/Public/MoverComponent.h
--- a/Engine/Plugins/Experimental/Mover/Source/Mover/Public/MoverComponent.h
+++ b/Engine/Plugins/Experimental/Mover/Source/Mover/Public/MoverComponent.h
@@ -321,7 +321,11 @@ public:
 
 	/** Returns a quaternion transforming from gravity to world space. */
 	FQuat GetGravityToWorldTransform() const { return GravityToWorldTransform; }
-	
+
+	// Set UpDirection override. This is a fixed direction that overrides the gravity-derived up direction.
+	UFUNCTION(BlueprintCallable, Category = Mover)
+	MOVER_API void SetUpDirectionOverride(bool bOverrideUpDirection, FVector UpDirection=FVector::UpVector);
+
 	// Get the normalized direction considered "up" in worldspace. Typically aligned with gravity, and typically determines the plane an actor tries to move along.
 	UFUNCTION(BlueprintPure = false, Category = Mover)
 	MOVER_API FVector GetUpDirection() const;
@@ -750,6 +754,14 @@ protected:
 	MOVER_API virtual void CheckForExternalMovement(const FMoverTickStartData& SimStartingData);
 
 private:
+	// Whether to override the up direction with a fixed value instead of using gravity to deduce it
+	UPROPERTY(EditDefaultsOnly, Category="Mover|UpDirection")
+	bool bHasUpDirectionOverride = false;
+
+	// A fixed up direction to use if bHasUpDirectionOverride is true
+	UPROPERTY(EditDefaultsOnly, Category="Mover|UpDirection", meta = (EditCondition = "bHasUpDirectionOverride"))
+	FVector UpDirectionOverride = FVector::UpVector;
+
 	/** Whether or not gravity is overridden on this actor. Otherwise, fall back on world settings. See @SetGravityOverride */
 	UPROPERTY(EditDefaultsOnly, Category="Mover|Gravity")
 	bool bHasGravityOverride = false;
