CL follow up from BugHawk feedback

diff --git a/Engine/Plugins/Experimental/Mover/Source/Mover/Private/DefaultMovementSet/Modes/SimpleSpringState.h b/Engine/Plugins/Experimental/Mover/Source/Mover/Private/DefaultMovementSet/Modes/SimpleSpringState.h
--- a/Engine/Plugins/Experimental/Mover/Source/Mover/Private/DefaultMovementSet/Modes/SimpleSpringState.h
+++ b/Engine/Plugins/Experimental/Mover/Source/Mover/Private/DefaultMovementSet/Modes/SimpleSpringState.h
@@ -8,7 +8,7 @@
 #include "SimpleSpringState.generated.h"
 
 /**
-* Internal state data for the SpringWalkingState
+* Internal state data for the SimpleWalkingMode
 */
 USTRUCT()
 struct FSimpleSpringState : public FMoverDataStructBase
diff --git a/Engine/Plugins/Experimental/Mover/Source/Mover/Private/DefaultMovementSet/Modes/SimpleWalkingMode.cpp b/Engine/Plugins/Experimental/Mover/Source/Mover/Private/DefaultMovementSet/Modes/SimpleWalkingMode.cpp
--- a/Engine/Plugins/Experimental/Mover/Source/Mover/Private/DefaultMovementSet/Modes/SimpleWalkingMode.cpp
+++ b/Engine/Plugins/Experimental/Mover/Source/Mover/Private/DefaultMovementSet/Modes/SimpleWalkingMode.cpp
@@ -63,7 +63,7 @@ void USimpleWalkingMode::GenerateMove_Implementation(const FMoverTickStartData&
 		{
 			// Clamp to max move speed
 			DesiredVelocity = DesiredVelocity.GetClampedToMaxSize(MaxMoveSpeed);
-			OutProposedMove.DirectionIntent = DesiredVelocity / MaxMoveSpeed; // here, DesiredVelocity is converted to "intent space"
+			OutProposedMove.DirectionIntent = MaxMoveSpeed > UE_KINDA_SMALL_NUMBER ? DesiredVelocity / MaxMoveSpeed : FVector::ZeroVector; // here, DesiredVelocity is converted to "intent space"
 		}
 		break;
 	case EMoveInputType::None:
@@ -72,6 +72,7 @@ void USimpleWalkingMode::GenerateMove_Implementation(const FMoverTickStartData&
 		{
 			UE_LOG(LogMover, Warning, TEXT("Unhandled MoveInputType %i in USpringWalkingMode"), MoveInputType);
 			DesiredVelocity = FVector::ZeroVector;
+			OutProposedMove.DirectionIntent = FVector::ZeroVector;
 		}
 	break;
 	}
diff --git a/Engine/Plugins/Experimental/Mover/Source/Mover/Private/DefaultMovementSet/Modes/SmoothWalkingState.cpp b/Engine/Plugins/Experimental/Mover/Source/Mover/Private/DefaultMovementSet/Modes/SmoothWalkingState.cpp
--- a/Engine/Plugins/Experimental/Mover/Source/Mover/Private/DefaultMovementSet/Modes/SmoothWalkingState.cpp
+++ b/Engine/Plugins/Experimental/Mover/Source/Mover/Private/DefaultMovementSet/Modes/SmoothWalkingState.cpp
@@ -29,7 +29,7 @@ void FSmoothWalkingState::ToString(FAnsiStringBuilderBase& Out) const
 {
 	Super::ToString(Out);
 
-	Out.Appendf("Accel=%s IntVel=%s TurningAccel=% AngVel=%s\n", *Accel.ToCompactString(), *IntermediateVelocity.ToCompactString(), *TurningAccel.ToCompactString(), *CurrentAngularVelocity.ToCompactString());
+	Out.Appendf("Accel=%s IntVel=%s TurningAccel=%s AngVel=%s\n", *Accel.ToCompactString(), *IntermediateVelocity.ToCompactString(), *TurningAccel.ToCompactString(), *CurrentAngularVelocity.ToCompactString());
 }
 
 bool FSmoothWalkingState::ShouldReconcile(const FMoverDataStructBase& AuthorityState) const
diff --git a/Engine/Plugins/Experimental/Mover/Source/Mover/Private/DefaultMovementSet/Modes/SmoothWalkingState.h b/Engine/Plugins/Experimental/Mover/Source/Mover/Private/DefaultMovementSet/Modes/SmoothWalkingState.h
--- a/Engine/Plugins/Experimental/Mover/Source/Mover/Private/DefaultMovementSet/Modes/SmoothWalkingState.h
+++ b/Engine/Plugins/Experimental/Mover/Source/Mover/Private/DefaultMovementSet/Modes/SmoothWalkingState.h
@@ -8,7 +8,7 @@
 #include "SmoothWalkingState.generated.h"
 
 /**
-* Internal state data for the SpringWalkingState
+* Internal state data for the SmoothWalkingMode
 */
 USTRUCT()
 struct FSmoothWalkingState : public FMoverDataStructBase
