PhysicsMover - Always project constraint target onto movement surface when in walking mode

diff --git a/Engine/Plugins/Experimental/Mover/Source/Mover/Private/PhysicsMover/Modes/PhysicsDrivenWalkingMode.cpp b/Engine/Plugins/Experimental/Mover/Source/Mover/Private/PhysicsMover/Modes/PhysicsDrivenWalkingMode.cpp
--- a/Engine/Plugins/Experimental/Mover/Source/Mover/Private/PhysicsMover/Modes/PhysicsDrivenWalkingMode.cpp
+++ b/Engine/Plugins/Experimental/Mover/Source/Mover/Private/PhysicsMover/Modes/PhysicsDrivenWalkingMode.cpp
@@ -358,11 +358,7 @@ void UPhysicsDrivenWalkingMode::SimulationTick_Implementation(const FSimulationT
 		const bool bIsGroundMoving = ProjectedGroundVelocity.SizeSquared() > UE_KINDA_SMALL_NUMBER;
 
 		TargetVelocity = StartingSyncState->GetVelocity_WorldSpace();
-		FVector TargetPosition = StartingSyncState->GetLocation_WorldSpace();
-		if (FloorResult.FloorDist > TargetHeight + UE_KINDA_SMALL_NUMBER)
-		{
-			TargetPosition += UpDir * (TargetHeight - FloorResult.FloorDist);
-		}
+		FVector TargetPosition = StartingSyncState->GetLocation_WorldSpace() + UpDir * (TargetHeight - FloorResult.FloorDist);
 
 		constexpr float ParallelCosThreshold = 0.999f;
 		const bool bVerticalVelocity = FVector::Parallel(TargetVelocity.GetSafeNormal(), UpDir, ParallelCosThreshold);
