Removed velocity clamps from SmoothWalkingMode

diff --git a/Engine/Plugins/Experimental/Mover/Source/Mover/Private/DefaultMovementSet/Modes/SmoothWalkingMode.cpp b/Engine/Plugins/Experimental/Mover/Source/Mover/Private/DefaultMovementSet/Modes/SmoothWalkingMode.cpp
--- a/Engine/Plugins/Experimental/Mover/Source/Mover/Private/DefaultMovementSet/Modes/SmoothWalkingMode.cpp
+++ b/Engine/Plugins/Experimental/Mover/Source/Mover/Private/DefaultMovementSet/Modes/SmoothWalkingMode.cpp
@@ -38,8 +38,6 @@ void USmoothWalkingMode::GenerateWalkMove(const FMoverTickStartData& StartState,
 	float MaxAccelerationToUse =  bIsAccelerating ? TargetAcceleration : TargetDeceleration;
 
 	float DesiredVelocityLengthSq = DesiredVelocity.SquaredLength();
-	check(CommonLegacySettings);	// Should be already validated by this point by base class
-	float MaxSpeed = CommonLegacySettings->MaxSpeed;
 
 	if (TurningStrength > FLT_EPSILON)
 	{
@@ -57,7 +55,6 @@ void USmoothWalkingMode::GenerateWalkMove(const FMoverTickStartData& StartState,
 		if (!TargetMoveDir.IsNearlyZero())
 		{
 			SpringMath::SimpleSpringDamper(InOutVelocity, SpringState.TurningAccel, TargetMoveDir * InOutVelocity.Length(), SpringMath::SpringStrengthToSmoothingTime(TurningStrength), DeltaSeconds);
-			InOutVelocity = InOutVelocity.GetClampedToMaxSize(MaxSpeed);
 		}
 	}
 
@@ -85,7 +82,6 @@ void USmoothWalkingMode::GenerateWalkMove(const FMoverTickStartData& StartState,
 
 	ActualAccel = ActualAccel.GetClampedToMaxSize(AccelClamp);
 	InOutVelocity = PreviousVelocity + (ActualAccel * DeltaSeconds);
-	InOutVelocity = InOutVelocity.GetClampedToMaxSize(MaxSpeed);
 
 	// Angular / Facing //
 	// We use a simple spring
