Mover: Fix for a instant movement effect affecting non physics mover, where instant movement effects were being scheduled using the wrong frame and therefore being delayed when they were meant to be applied immediately. This affected the jumping for instance.

diff --git a/Engine/Plugins/Experimental/Mover/Source/Mover/Private/MoverComponent.cpp b/Engine/Plugins/Experimental/Mover/Source/Mover/Private/MoverComponent.cpp
--- a/Engine/Plugins/Experimental/Mover/Source/Mover/Private/MoverComponent.cpp
+++ b/Engine/Plugins/Experimental/Mover/Source/Mover/Private/MoverComponent.cpp
@@ -1584,23 +1584,39 @@ DEFINE_FUNCTION(UMoverComponent::execK2_ScheduleInstantMovementEffect)
 
 void UMoverComponent::ScheduleInstantMovementEffect(TSharedPtr<FInstantMovementEffect> InstantMovementEffect)
 {
+	ensureMsgf(IsInGameThread(), TEXT("UMoverComponent::ScheduleInstantMovementEffect should only be called from the game thread. Inspect code for incorrect calls."));
 	ensure(IsInGameThread());
-	QueueInstantMovementEffect(FScheduledInstantMovementEffect::ScheduleEffect_External(GetWorld(), InstantMovementEffect, /* SchedulingDelaySeconds = */ EventSchedulingMinDelaySeconds));
+	FMoverTimeStep TimeStep;
+	if (ensureMsgf(BackendLiaisonComp, TEXT("UMoverComponent::ScheduleInstantMovementEffect was unexpectedly called with a null backend liaison component. The instant movement effect will be ignored.")))
+	{
+		TimeStep.BaseSimTimeMs = BackendLiaisonComp->GetCurrentSimTimeMs();
+		TimeStep.ServerFrame = BackendLiaisonComp->GetCurrentSimFrame();
+		// TimeStep.StepMs is not used by FScheduledInstantMovementEffect::ScheduleEffect
+		QueueInstantMovementEffect(FScheduledInstantMovementEffect::ScheduleEffect(GetWorld(), TimeStep, InstantMovementEffect, /* SchedulingDelaySeconds = */ EventSchedulingMinDelaySeconds));
+	}
 }
 
 void UMoverComponent::QueueInstantMovementEffect_Internal(const FMoverTimeStep& TimeStep, TSharedPtr<FInstantMovementEffect> InstantMovementEffect)
 {
-	QueueInstantMovementEffect(FScheduledInstantMovementEffect::ScheduleEffect_Internal(GetWorld(), TimeStep, InstantMovementEffect, /* SchedulingDelaySeconds = */ 0.0f));
+	QueueInstantMovementEffect(FScheduledInstantMovementEffect::ScheduleEffect(GetWorld(), TimeStep, InstantMovementEffect, /* SchedulingDelaySeconds = */ 0.0f));
 }
 
 void UMoverComponent::QueueInstantMovementEffect(TSharedPtr<FInstantMovementEffect> InstantMovementEffect)
 {
-	ensure(IsInGameThread());
-	QueueInstantMovementEffect(FScheduledInstantMovementEffect::ScheduleEffect_External(GetWorld(), InstantMovementEffect, /* SchedulingDelaySeconds = */ 0.0f));
+	ensureMsgf(IsInGameThread(), TEXT("UMoverComponent::QueueInstantMovementEffect(TSharedPtr<FInstantMovementEffect>) should only be called from the game thread. Inspect code for incorrect calls."));
+	FMoverTimeStep TimeStep;
+	if (ensureMsgf(BackendLiaisonComp, TEXT("UMoverComponent::ScheduleInstantMovementEffect was unexpectedly called with a null backend liaison component. The instant movement effect will be ignored.")))
+	{
+		TimeStep.BaseSimTimeMs = BackendLiaisonComp->GetCurrentSimTimeMs();
+		TimeStep.ServerFrame = BackendLiaisonComp->GetCurrentSimFrame();
+		// TimeStep.StepMs is not used by FScheduledInstantMovementEffect::ScheduleEffect
+		QueueInstantMovementEffect(FScheduledInstantMovementEffect::ScheduleEffect(GetWorld(), TimeStep, InstantMovementEffect, /* SchedulingDelaySeconds = */ 0.0f));
+	}
 }
 
 void UMoverComponent::QueueInstantMovementEffect(const FScheduledInstantMovementEffect& InstantMovementEffect)
 {
+	// TODO Move QueueInstantMovementEffect to UMoverSimulation and implement differently in sync or async mode
 	if (IsInGameThread())
 	{
 		QueuedInstantMovementEffects.Add(InstantMovementEffect);
diff --git a/Engine/Plugins/Experimental/Mover/Source/Mover/Private/MoverSimulationTypes.cpp b/Engine/Plugins/Experimental/Mover/Source/Mover/Private/MoverSimulationTypes.cpp
--- a/Engine/Plugins/Experimental/Mover/Source/Mover/Private/MoverSimulationTypes.cpp
+++ b/Engine/Plugins/Experimental/Mover/Source/Mover/Private/MoverSimulationTypes.cpp
@@ -14,48 +14,18 @@ UScriptStruct* FMoverSimulationEventData::GetScriptStruct() const
 	return FMoverSimulationEventData::StaticStruct();
 }
 
-FScheduledInstantMovementEffect FScheduledInstantMovementEffect::ScheduleEffect_External(UWorld* World, TSharedPtr<FInstantMovementEffect> InstantMovementEffect, float SchedulingDelaySeconds)
+FScheduledInstantMovementEffect FScheduledInstantMovementEffect::ScheduleEffect(UWorld* World, const FMoverTimeStep& TimeStep, TSharedPtr<FInstantMovementEffect> InstantMovementEffect, float SchedulingDelaySeconds)
 {
-	ensure(IsInGameThread());
 	FPhysScene* Scene = World ? World->GetPhysicsScene() : nullptr;
 	Chaos::FPhysicsSolver* Solver = Scene ? Scene->GetSolver() : nullptr;
 	bool bIsFixedDt = Solver ? Solver->IsUsingFixedDt() : false;
-	if (bIsFixedDt)
-	{
-		int32 ServerFrame = UE::NetworkPhysicsUtils::GetUpcomingServerFrame_External(World);
-		if (SchedulingDelaySeconds != 0.0f)
-		{
-			const double AsyncDeltaTime = Solver->GetAsyncDeltaTime();
-			const int32 DelayFrames = (AsyncDeltaTime > UE_SMALL_NUMBER) ? FMath::CeilToInt32(SchedulingDelaySeconds / AsyncDeltaTime) : 0;
-			ServerFrame += DelayFrames;
-		}
-
-		return FScheduledInstantMovementEffect(/* ExecutionServerFrame = */ ServerFrame, /* ExecutionServerTime = */ 0.0, /* bIsFixedDt = */ true, InstantMovementEffect);
-	}
-	else
-	{
-		AGameStateBase* GameState = World ? World->GetGameState() : nullptr;
-		double ServerTime = GameState ? GameState->GetServerWorldTimeSeconds() : 0.0;
-		ServerTime += SchedulingDelaySeconds;
-		
-		return FScheduledInstantMovementEffect(/* ExecutionServerFrame = */ INDEX_NONE, /* ExecutionServerTime = */ ServerTime, /* bIsFixedDt = */ false, InstantMovementEffect);
-	}
-}
-
-FScheduledInstantMovementEffect FScheduledInstantMovementEffect::ScheduleEffect_Internal(UWorld* World, const FMoverTimeStep& TimeStep, TSharedPtr<FInstantMovementEffect> InstantMovementEffect, float SchedulingDelaySeconds)
-{
-	ensure(!IsInGameThread());
-	FPhysScene* Scene = World ? World->GetPhysicsScene() : nullptr;
-	Chaos::FPhysicsSolver* Solver = Scene ? Scene->GetSolver() : nullptr;
-	bool bIsFixedDt = Solver ? Solver->IsUsingFixedDt() : false;
-	
 	if (bIsFixedDt)
 	{
 		int32 ServerFrame = TimeStep.ServerFrame;
 		if (SchedulingDelaySeconds != 0.0f)
 		{
-			const double AsyncDeltaTime = Solver->GetAsyncDeltaTime();
-			const int32 DelayFrames = (AsyncDeltaTime > UE_SMALL_NUMBER) ? FMath::CeilToInt32(SchedulingDelaySeconds / AsyncDeltaTime) : 0;
+			const double DeltaTime = Solver->GetAsyncDeltaTime();
+			const int32 DelayFrames = (DeltaTime > UE_SMALL_NUMBER) ? FMath::CeilToInt32(SchedulingDelaySeconds / DeltaTime) : 0;
 			ServerFrame += DelayFrames;
 		}
 		return FScheduledInstantMovementEffect(/* ExecutionServerFrame = */ ServerFrame, /* ExecutionServerTime = */ 0.0, /* bIsFixedDt = */ true, InstantMovementEffect);
@@ -65,6 +35,7 @@ FScheduledInstantMovementEffect FScheduledInstantMovementEffect::ScheduleEffect_
 		AGameStateBase* GameState = World ? World->GetGameState() : nullptr;
 		double ServerTime = GameState ? GameState->GetServerWorldTimeSeconds() : 0.0;
 		ServerTime += SchedulingDelaySeconds;
+		
 		return FScheduledInstantMovementEffect(/* ExecutionServerFrame = */ INDEX_NONE, /* ExecutionServerTime = */ ServerTime, /* bIsFixedDt = */ false, InstantMovementEffect);
 	}
 }
diff --git a/Engine/Plugins/Experimental/Mover/Source/Mover/Public/MoverSimulationTypes.h b/Engine/Plugins/Experimental/Mover/Source/Mover/Public/MoverSimulationTypes.h
--- a/Engine/Plugins/Experimental/Mover/Source/Mover/Public/MoverSimulationTypes.h
+++ b/Engine/Plugins/Experimental/Mover/Source/Mover/Public/MoverSimulationTypes.h
@@ -74,24 +74,15 @@ struct FScheduledInstantMovementEffect
 {
 	GENERATED_BODY()
 
-	/** Turns a FInstantMovementEffect into a scheduled one (FScheduledInstantMovementEffect)
-	*	The effect can be scheduled to apply immediately, or scheduled to apply with a delay
-	*   This function can only be called on the game thread
-	*   @param World The world, used to retrieve the current server frame in async mode, or the sim time otherwise
-	*   @param InstantMovementEffect the effect to schedule
-	*   @param SchedulingDelaySeconds Scheduling delay to ensure it applies on all end points on the same frame (this is only perfectly accurate when simulation dt is fixed)
-	*/
-	static FScheduledInstantMovementEffect ScheduleEffect_External(UWorld* World, TSharedPtr<FInstantMovementEffect> InstantMovementEffect, float SchedulingDelaySeconds = 0.0f);
-
 	/** Turns a FInstantMovementEffect into a scheduled one (FScheduledInstantMovementEffect)
 	*	The effect can be scheduled to apply immediately, or scheduled to apply with a delay
 	*   This function should not be called on the game thread
 	*   @param World The world, used to retrieve the current server frame in async mode, or the sim time otherwise
-	*   @param TimeStep the time step used in the async update which this function should be called from
+	*   @param TimeStep the time step of the current or upcoming tick
 	*   @param InstantMovementEffect the effect to schedule
 	*   @param SchedulingDelaySeconds Scheduling delay to ensure it applies on all end points on the same frame (this is only perfectly accurate when simulation dt is fixed)
 	*/
-	static FScheduledInstantMovementEffect ScheduleEffect_Internal(UWorld* World, const FMoverTimeStep& TimeStep, TSharedPtr<FInstantMovementEffect> InstantMovementEffect, float SchedulingDelaySeconds = 0.0f);
+	static FScheduledInstantMovementEffect ScheduleEffect(UWorld* World, const FMoverTimeStep& TimeStep, TSharedPtr<FInstantMovementEffect> InstantMovementEffect, float SchedulingDelaySeconds = 0.0f);
 
 	bool ShouldExecuteAtFrame(int32 CurrentServerFrame) const
 	{
