[Mover] Changing the default DecayAmountMultipliers to 1.

diff --git a/Engine/Plugins/Experimental/Mover/Source/Mover/Private/MoverDataModelTypes.cpp b/Engine/Plugins/Experimental/Mover/Source/Mover/Private/MoverDataModelTypes.cpp
--- a/Engine/Plugins/Experimental/Mover/Source/Mover/Private/MoverDataModelTypes.cpp
+++ b/Engine/Plugins/Experimental/Mover/Source/Mover/Private/MoverDataModelTypes.cpp
@@ -142,7 +142,7 @@ void FCharacterDefaultInputs::Merge(const FMoverDataStructBase& From)
 	bIsJumpPressed |= TypedFrom.bIsJumpPressed;
 }
 
-static float CharacterDefaultInputsDecayAmountMultiplier = 0.25f;
+static float CharacterDefaultInputsDecayAmountMultiplier = 1.f;
 FAutoConsoleVariableRef CVarCharacterDefaultInputsDecayAmountMultiplier(
 	TEXT("Mover.Input.CharacterDefaultInputsDecayAmountMultiplier"),
 	CharacterDefaultInputsDecayAmountMultiplier,
diff --git a/Engine/Plugins/Experimental/Mover/Source/Mover/Private/PhysicsMover/PhysicsMoverSimulationTypes.cpp b/Engine/Plugins/Experimental/Mover/Source/Mover/Private/PhysicsMover/PhysicsMoverSimulationTypes.cpp
--- a/Engine/Plugins/Experimental/Mover/Source/Mover/Private/PhysicsMover/PhysicsMoverSimulationTypes.cpp
+++ b/Engine/Plugins/Experimental/Mover/Source/Mover/Private/PhysicsMover/PhysicsMoverSimulationTypes.cpp
@@ -96,7 +96,7 @@ void FMovementSettingsInputs::Merge(const FMoverDataStructBase& From)
 {
 }
 
-static float MovementSettingsInputsDecayAmountMultiplier = 0.25f;
+static float MovementSettingsInputsDecayAmountMultiplier = 1.f;
 FAutoConsoleVariableRef CVarMovementSettingsInputsDecayAmountMultiplier(
 	TEXT("Mover.Input.MovementSettingsInputsDecayAmountMultiplier"),
 	MovementSettingsInputsDecayAmountMultiplier,
