Mover - Add option to apply friction to lateral in air movement

diff --git a/Engine/Plugins/Experimental/Mover/Source/Mover/Private/DefaultMovementSet/Modes/FallingMode.cpp b/Engine/Plugins/Experimental/Mover/Source/Mover/Private/DefaultMovementSet/Modes/FallingMode.cpp
--- a/Engine/Plugins/Experimental/Mover/Source/Mover/Private/DefaultMovementSet/Modes/FallingMode.cpp
+++ b/Engine/Plugins/Experimental/Mover/Source/Mover/Private/DefaultMovementSet/Modes/FallingMode.cpp
@@ -18,6 +18,7 @@ UFallingMode::UFallingMode(const FObjectInitializer& ObjectInitializer)
 	, bCancelVerticalSpeedOnLanding(true)
 	, AirControlPercentage(0.4f)
 	, FallingDeceleration(200.0f)
+	, FallingLateralFriction(0.0f)
 	, OverTerminalSpeedFallingDeceleration(800.0f)
 	, TerminalMovementPlaneSpeed(1500.0f)
 	, bShouldClampTerminalVerticalSpeed(true)
@@ -89,6 +90,7 @@ void UFallingMode::GenerateMove_Implementation(const FMoverTickStartData& StartS
 	Params.Deceleration = FallingDeceleration;
 	Params.WorldToGravityQuat = MoverComp->GetWorldToGravityTransform();
 	Params.bUseAccelerationForVelocityMove = CommonLegacySettings->bUseAccelerationForVelocityMove;
+	Params.Friction = FallingLateralFriction;
 
 	// Check if any current velocity values are over our terminal velocity - if so limit the move input in that direction and apply OverTerminalVelocityFallingDeceleration
 	if (Params.MoveInput.Dot(StartVelocity) > 0 && StartHorizontalVelocity.Size() >= TerminalMovementPlaneSpeed)
diff --git a/Engine/Plugins/Experimental/Mover/Source/Mover/Private/MoveLibrary/AirMovementUtils.cpp b/Engine/Plugins/Experimental/Mover/Source/Mover/Private/MoveLibrary/AirMovementUtils.cpp
--- a/Engine/Plugins/Experimental/Mover/Source/Mover/Private/MoveLibrary/AirMovementUtils.cpp
+++ b/Engine/Plugins/Experimental/Mover/Source/Mover/Private/MoveLibrary/AirMovementUtils.cpp
@@ -29,6 +29,7 @@ FProposedMove UAirMovementUtils::ComputeControlledFreeMove(const FFreeMoveParams
 	ComputeVelocityParams.MoveInputType = InParams.MoveInputType;
 	ComputeVelocityParams.MoveInput = InParams.MoveInput; 
 	ComputeVelocityParams.bUseAccelerationForVelocityMove = InParams.bUseAccelerationForVelocityMove;
+	ComputeVelocityParams.Friction = InParams.Friction;
 	
 	OutMove.LinearVelocity = UMovementUtils::ComputeVelocity(ComputeVelocityParams);
 	OutMove.AngularVelocity = UMovementUtils::ComputeAngularVelocity(InParams.PriorOrientation, InParams.OrientationIntent, InParams.WorldToGravityQuat, InParams.DeltaSeconds, InParams.TurningRate);
diff --git a/Engine/Plugins/Experimental/Mover/Source/Mover/Public/DefaultMovementSet/Modes/FallingMode.h b/Engine/Plugins/Experimental/Mover/Source/Mover/Public/DefaultMovementSet/Modes/FallingMode.h
--- a/Engine/Plugins/Experimental/Mover/Source/Mover/Public/DefaultMovementSet/Modes/FallingMode.h
+++ b/Engine/Plugins/Experimental/Mover/Source/Mover/Public/DefaultMovementSet/Modes/FallingMode.h
@@ -53,6 +53,15 @@ public:
 	UPROPERTY(Category=Mover, EditAnywhere, BlueprintReadWrite, meta=(ClampMin="0", ForceUnits = "cm/s^2"))
 	float FallingDeceleration;
 
+	/**
+	 * Friction to apply to air movement.
+	 * Lateral velocity is scaled each tick by a factor (1-Friction * TimeStep) so friction
+	 * values greater than 1/TimeStep will result in all lateral velocity being removed.
+	 * Note: This is NOT applied to vertical velocity, only movement plane velocity
+	 */
+	UPROPERTY(Category = Mover, EditAnywhere, BlueprintReadWrite, meta = (ClampMin = "0"))
+	float FallingLateralFriction;
+
 	/**
      * Deceleration to apply to air movement when falling faster than terminal velocity
 	 * Note: This is NOT applied to vertical velocity, only movement plane velocity
diff --git a/Engine/Plugins/Experimental/Mover/Source/Mover/Public/MoveLibrary/AirMovementUtils.h b/Engine/Plugins/Experimental/Mover/Source/Mover/Public/MoveLibrary/AirMovementUtils.h
--- a/Engine/Plugins/Experimental/Mover/Source/Mover/Public/MoveLibrary/AirMovementUtils.h
+++ b/Engine/Plugins/Experimental/Mover/Source/Mover/Public/MoveLibrary/AirMovementUtils.h
@@ -50,6 +50,9 @@ struct FFreeMoveParams
 	UPROPERTY(BlueprintReadWrite, EditAnywhere, Category = Mover)
 	float DeltaSeconds = 0.f;
 
+	UPROPERTY(BlueprintReadWrite, EditAnywhere, Category = Mover)
+	float Friction = 0.f;
+
 	UPROPERTY(BlueprintReadWrite, EditAnywhere, Category = Mover)
 	FQuat WorldToGravityQuat = FQuat::Identity;
 
