Mover: first pass of rollback-friendly blackboard. It uses a circular buffer to maintain a recent history of value changes. This initial version supports data entries of a fixed/predefined buffer size only, useful for things that do not change every frame.
Also includes tweaks to the simulation state machine, which is required to better keep the sim time up-to-date during substepping.

diff --git a/Engine/Plugins/Experimental/Mover/Source/Mover/Private/Backends/MoverNetworkPhysicsLiaisonBase.cpp b/Engine/Plugins/Experimental/Mover/Source/Mover/Private/Backends/MoverNetworkPhysicsLiaisonBase.cpp
--- a/Engine/Plugins/Experimental/Mover/Source/Mover/Private/Backends/MoverNetworkPhysicsLiaisonBase.cpp
+++ b/Engine/Plugins/Experimental/Mover/Source/Mover/Private/Backends/MoverNetworkPhysicsLiaisonBase.cpp
@@ -662,7 +662,8 @@ void UMoverNetworkPhysicsLiaisonComponentBase::ProcessInputs_Internal(int32 Phys
 			// So to account for all possible configs, always establish the initial sync state of the frame based on LatestSyncState on clients (it'd be redundant on servers).
 			if (!NetworkPhysicsComponent->GetNetworkPhysicsComponent_Internal()->IsServer())
 			{
-				Chaos::FPBDRigidsEvolution* RigidsEvolution = GetPhysicsSolver()->GetEvolution();
+				Chaos::FPhysicsSolver* Solver = GetPhysicsSolver();
+				Chaos::FPBDRigidsEvolution* RigidsEvolution = Solver->GetEvolution();
 				const bool bIsSolverResim = RigidsEvolution->IsResimming();
 				const bool bIsFirstResimFrame = RigidsEvolution->IsResetting();
 				const bool bIsProxyRepResim = NetworkPhysicsComponent->GetNetworkPhysicsComponent_Internal()->GetPhysicsReplicationMode() == EPhysicsReplicationMode::Resimulation;
@@ -674,6 +675,10 @@ void UMoverNetworkPhysicsLiaisonComponentBase::ProcessInputs_Internal(int32 Phys
 					FMoverAuxStateContext UnusedInvalidAuxState;
 					FMoverAuxStateContext UnusedAuxState;
 
+					FMoverTimeStep CurrentBaseTimeStep;
+					CurrentBaseTimeStep.BaseSimTimeMs = Solver->GetSolverTime() * 1000.f;
+					CurrentBaseTimeStep.ServerFrame = Solver->GetCurrentFrame();
+
 					// Prepare Networked SyncState for rollback
 					if (Chaos::FPBDRigidParticleHandle* ParticleHandle = GetControlledParticle_Internal())
 					{
@@ -681,9 +686,9 @@ void UMoverNetworkPhysicsLiaisonComponentBase::ProcessInputs_Internal(int32 Phys
 						NetSyncState.SetTransforms_WorldSpace(ParticleHandle->GetX(), FRotator(ParticleHandle->GetR()), ParticleHandle->GetV());
 					}
 
-					MoverComp.OnSimulationPreRollback(&Input.SyncState, &LatestSyncState, &UnusedInvalidAuxState, &UnusedAuxState);
+					MoverComp.OnSimulationPreRollback(&Input.SyncState, &LatestSyncState, &UnusedInvalidAuxState, &UnusedAuxState, CurrentBaseTimeStep);
 					GetCurrentStateData(Input.SyncState);
-					MoverComp.OnSimulationRollback(&Input.SyncState, &UnusedAuxState);
+					MoverComp.OnSimulationRollback(&Input.SyncState, &UnusedAuxState, CurrentBaseTimeStep);
 				}
 				else
 				{
diff --git a/Engine/Plugins/Experimental/Mover/Source/Mover/Private/Backends/MoverNetworkPredictionLiaison.cpp b/Engine/Plugins/Experimental/Mover/Source/Mover/Private/Backends/MoverNetworkPredictionLiaison.cpp
--- a/Engine/Plugins/Experimental/Mover/Source/Mover/Private/Backends/MoverNetworkPredictionLiaison.cpp
+++ b/Engine/Plugins/Experimental/Mover/Source/Mover/Private/Backends/MoverNetworkPredictionLiaison.cpp
@@ -57,7 +57,40 @@ void UMoverNetworkPredictionLiaisonComponent::ProduceInput(const int32 DeltaTime
 void UMoverNetworkPredictionLiaisonComponent::RestoreFrame(const FMoverSyncState* SyncState, const FMoverAuxStateContext* AuxState)
 {
 	check(MoverComp);
-	MoverComp->RestoreFrame(SyncState, AuxState);
+
+	int32 NewBaseSimTimeMs = 0;
+	int32 NextFrameNum = 0;
+
+	switch (UNetworkPredictionWorldManager::ActiveInstance->PreferredDefaultTickingPolicy())
+	{
+		default:	// fall through
+		case ENetworkPredictionTickingPolicy::Fixed:
+		{
+			const FFixedTickState& FixedTickState = UNetworkPredictionWorldManager::ActiveInstance->GetFixedTickState();
+			FNetSimTimeStep TimeStep = FixedTickState.GetNextTimeStep();
+			NewBaseSimTimeMs = TimeStep.TotalSimulationTime;
+			NextFrameNum = TimeStep.Frame;
+		}
+		break; 
+
+		case ENetworkPredictionTickingPolicy::Independent:
+		{
+			const FVariableTickState& VariableTickState = UNetworkPredictionWorldManager::ActiveInstance->GetVariableTickState();
+			const FNetSimTimeStep NextVariableTimeStep = VariableTickState.GetNextTimeStep(VariableTickState.Frames[VariableTickState.ConfirmedFrame]);
+			NewBaseSimTimeMs = NextVariableTimeStep.TotalSimulationTime;
+			NextFrameNum = NextVariableTimeStep.Frame;
+
+		}
+		break;
+	}
+
+	FMoverTimeStep MoverTimeStep;
+
+	MoverTimeStep.ServerFrame = NextFrameNum;
+	MoverTimeStep.BaseSimTimeMs = NewBaseSimTimeMs;
+	MoverTimeStep.StepMs = 0;
+
+	MoverComp->RestoreFrame(SyncState, AuxState, MoverTimeStep);
 }
 
 void UMoverNetworkPredictionLiaisonComponent::FinalizeFrame(const FMoverSyncState* SyncState, const FMoverAuxStateContext* AuxState)
diff --git a/Engine/Plugins/Experimental/Mover/Source/Mover/Private/MoveLibrary/RollbackBlackboard.cpp b/Engine/Plugins/Experimental/Mover/Source/Mover/Private/MoveLibrary/RollbackBlackboard.cpp
new file mode 100644
--- /dev/null
+++ b/Engine/Plugins/Experimental/Mover/Source/Mover/Private/MoveLibrary/RollbackBlackboard.cpp
@@ -0,0 +1,107 @@
+// Copyright Epic Games, Inc. All Rights Reserved.
+
+
+#include "MoveLibrary/RollbackBlackboard.h"
+#include "MoverTypes.h"
+#include "GenericPlatform/GenericPlatformMath.h"
+
+#include UE_INLINE_GENERATED_CPP_BY_NAME(RollbackBlackboard)
+
+
+uint32 URollbackBlackboard::BlackboardEntryBase::ComputeBufferSize(EBlackboardSizingPolicy SizingPolicy, uint32 FixedBufferSize)
+{
+	uint32 BufferSize = 2;
+
+	// TODO: handle these special cases:
+	//     if we are in a non-networked situation w/ non-async simulation, we only need 1 entry slot
+	//     if we are in a non-networked situation w/ async simulation, we only need 2 entry slots (one for simulating forwarD)
+
+	switch (SizingPolicy)
+	{
+		case EBlackboardSizingPolicy::FixedDeclaredSize:
+			// TODO: Warn if FixedBufferSize is invalid
+			BufferSize = FMath::Max(BufferSize, FixedBufferSize);
+			break;
+
+		default:	//TODO: warn about an unhandled policy type
+			break;
+	}
+
+	return BufferSize;
+}
+
+void URollbackBlackboard::BlackboardEntryBase::RollBack(uint32 NewPendingFrame)
+{
+	// Goal: adjust entry to point at the value from the prior frame. May make the entry invalidated, if there were no values that old.
+	check(ExternalIdx == InternalIdx);
+
+	const uint32 LowestPossibleIdx = (ExternalIdx > Timestamps.Capacity()) ? (ExternalIdx - Timestamps.Capacity()) : 0u;
+
+	// Walk downwards to find the highest index with a frame < NewPendingFrame. 
+	// If we hit uint32::max, that indicates we wrapped around from 0 and we're out of slots to check.
+	for (uint32 IdxToCheck = ExternalIdx; IdxToCheck >= LowestPossibleIdx && IdxToCheck != TNumericLimits<uint32>::Max(); --IdxToCheck)
+	{
+		if (Timestamps[IdxToCheck].IsValid() && Timestamps[IdxToCheck].Frame < NewPendingFrame)
+		{
+			ExternalIdx = InternalIdx = IdxToCheck;
+			return;
+		}
+	}
+
+	// If we made it here, then there are no entries that weren't rolled back, so let's make it clear
+	Timestamps[ExternalIdx].Invalidate();
+
+
+}
+
+
+void URollbackBlackboard::BeginSimulationFrame(const FMoverTimeStep& PendingTimeStep)
+{
+	check(!bIsSimulationInProgress && !bIsRollbackInProgress);
+	InProgressSimFrameThreadId = FPlatformTLS::GetCurrentThreadId();
+	bIsSimulationInProgress = true;
+
+	InProgressSimTimeStamp.TimeMs = PendingTimeStep.BaseSimTimeMs;
+	InProgressSimTimeStamp.Frame = PendingTimeStep.ServerFrame;
+}
+
+void URollbackBlackboard::EndSimulationFrame()
+{
+	check(bIsSimulationInProgress && (InProgressSimFrameThreadId == FPlatformTLS::GetCurrentThreadId()));
+	bIsSimulationInProgress = false;
+
+	// TODO: need some kind of lock so we wait on any in-progress operations before advancing the CurrentSimTimeStamp
+	CurrentSimTimeStamp = InProgressSimTimeStamp;
+
+
+	// TODO: could we skip this if no "set" operations occurred? Consider different policies, like change-every-frame entries
+
+	for (const TPair<FName, TUniquePtr<BlackboardEntryBase>>& KVP : EntryMap)
+	{
+		KVP.Value->OnSimulationFrameEnd();
+	}
+}
+
+
+void URollbackBlackboard::BeginRollback()
+{
+	check(!bIsSimulationInProgress && !bIsRollbackInProgress);
+	InRollbackThreadId = FPlatformTLS::GetCurrentThreadId();
+	bIsRollbackInProgress = true;
+}
+
+void URollbackBlackboard::RollBack(uint32 NewPendingFrame)
+{
+	check(!bIsSimulationInProgress && bIsRollbackInProgress);
+	
+	for (const TPair<FName, TUniquePtr<BlackboardEntryBase>>& KVP : EntryMap)
+	{
+		KVP.Value->RollBack(NewPendingFrame);
+	}
+}
+
+void URollbackBlackboard::EndRollback()
+{
+	check (bIsRollbackInProgress && (InRollbackThreadId == FPlatformTLS::GetCurrentThreadId()));
+	bIsRollbackInProgress = false;
+}
diff --git a/Engine/Plugins/Experimental/Mover/Source/Mover/Private/MovementModeStateMachine.cpp b/Engine/Plugins/Experimental/Mover/Source/Mover/Private/MovementModeStateMachine.cpp
--- a/Engine/Plugins/Experimental/Mover/Source/Mover/Private/MovementModeStateMachine.cpp
+++ b/Engine/Plugins/Experimental/Mover/Source/Mover/Private/MovementModeStateMachine.cpp
@@ -8,6 +8,7 @@
 #include "MoveLibrary/MovementUtils.h"
 #include "MoverComponent.h"
 #include "MoveLibrary/MovementMixer.h"
+#include "MoveLibrary/RollbackBlackboard.h"
 #include "Templates/SubclassOf.h"
 #include "Engine/World.h"
 #include "GameFramework/GameStateBase.h"
@@ -158,6 +159,9 @@ void UMovementModeStateMachine::OnSimulationTick(USceneComponent* UpdatedCompone
 {
 	//GEngine->AddOnScreenDebugMessage(-1, -0.1f, FColor::White, FString::Printf(TEXT("Mode Tick: %s  Queued: %s"), *CurrentModeName.ToString(), *NextModeName.ToString()));
 	FMoverTimeStep SubTimeStep = TimeStep;
+	CurrentBaseTimeStep = TimeStep;
+
+
 	WorkingSubstepStartData = StartState;
 	bool bIsWorkingStartStateReady = true;	// this flag is used to avoid unneeded data copying after substeps
 
@@ -405,6 +409,12 @@ void UMovementModeStateMachine::OnSimulationTick(USceneComponent* UpdatedCompone
 			//GEngine->AddOnScreenDebugMessage(-1, -0.1f, FColor::White, FString::Printf(TEXT("NextModeName: %s  Queued: %s"), *Output.MovementEndState.NextModeName.ToString(), *NextModeName.ToString()));
 		}
 
+
+		const float RemainingMs = FMath::Clamp(OutputState.MovementEndState.RemainingMs, 0.0f, SubTimeStep.StepMs);
+		const float SubstepUsedMs = (SubTimeStep.StepMs - RemainingMs);
+		CurrentBaseTimeStep.BaseSimTimeMs = SubTimeStep.BaseSimTimeMs + SubstepUsedMs;
+		TotalUsedMs += SubstepUsedMs;
+
 		// Switch modes if necessary (note that this will allow exit/enter on the same state)
 		AdvanceToNextMode();
 		OutputState.SyncState.MovementMode = CurrentModeName;
@@ -414,10 +424,7 @@ void UMovementModeStateMachine::OnSimulationTick(USceneComponent* UpdatedCompone
 			Modifier->OnPostMovement(MoverComp, SubTimeStep, OutputState.SyncState, OutputState.AuxState);
 		}
 		
-		const float RemainingMs = FMath::Clamp(OutputState.MovementEndState.RemainingMs, 0.0f, SubTimeStep.StepMs);
-		const float SubstepUsedMs = (SubTimeStep.StepMs - RemainingMs);
 		SubTimeStep.BaseSimTimeMs += SubstepUsedMs;
-		TotalUsedMs += SubstepUsedMs;
 		SubTimeStep.StepMs = RemainingMs;
 
 		bIsWorkingStartStateReady = false;
@@ -451,12 +458,13 @@ void UMovementModeStateMachine::OnSimulationTick(USceneComponent* UpdatedCompone
 	}
 }
 
-void UMovementModeStateMachine::OnSimulationPreRollback(const FMoverSyncState* InvalidSyncState, const FMoverSyncState* SyncState, const FMoverAuxStateContext* InvalidAuxState, const FMoverAuxStateContext* AuxState)
+void UMovementModeStateMachine::OnSimulationPreRollback(const FMoverSyncState* InvalidSyncState, const FMoverSyncState* SyncState, const FMoverAuxStateContext* InvalidAuxState, const FMoverAuxStateContext* AuxState, const FMoverTimeStep& NewBaseTimeStep)
 {
+	CurrentBaseTimeStep = NewBaseTimeStep;
 	RollbackModifiers(InvalidSyncState, SyncState, InvalidAuxState, AuxState);
 }
 
-void UMovementModeStateMachine::OnSimulationRollback(const FMoverSyncState* SyncState, const FMoverAuxStateContext* AuxState)
+void UMovementModeStateMachine::OnSimulationRollback(const FMoverSyncState* SyncState, const FMoverAuxStateContext* AuxState, const FMoverTimeStep& NewBaseTimeStep)
 {
 	ClearQueuedMode();
 
@@ -615,6 +623,8 @@ void UMovementModeStateMachine::AdvanceToNextMode()
 	const FName NextModeName = QueuedModeTransition->GetNextModeName();
 	const bool bShouldNextModeReenter = QueuedModeTransition->ShouldReenter();
 
+	ClearQueuedMode();
+
 	if ((NextModeName != NAME_None && Modes.Contains(NextModeName)) &&
 		(CurrentModeName != NextModeName || bShouldNextModeReenter))
 	{
@@ -634,10 +644,18 @@ void UMovementModeStateMachine::AdvanceToNextMode()
 
 		// signal movement mode change event
 		const UMoverComponent* MoverComp = CastChecked<UMoverComponent>(GetOuter());
+
+		URollbackBlackboard_InternalWrapper* RollbackBlackboard = MoverComp->GetRollbackBlackboard_Internal();
+
+		FMovementModeChangeRecord ModeChangeRecord;
+		ModeChangeRecord.ModeName = CurrentModeName;
+		ModeChangeRecord.Frame = CurrentBaseTimeStep.ServerFrame;
+		ModeChangeRecord.SimTimeMs = CurrentBaseTimeStep.BaseSimTimeMs;
+
+		RollbackBlackboard->TrySet(CommonBlackboard::LastModeChangeRecord, ModeChangeRecord);
+
 		MoverComp->OnMovementModeChanged.Broadcast(PreviousModeName, NextModeName);
 	}
-
-	ClearQueuedMode();
 }
 
 void UMovementModeStateMachine::FlushQueuedMovesToGroup(FLayeredMoveGroup& Group)
diff --git a/Engine/Plugins/Experimental/Mover/Source/Mover/Private/MoverComponent.cpp b/Engine/Plugins/Experimental/Mover/Source/Mover/Private/MoverComponent.cpp
--- a/Engine/Plugins/Experimental/Mover/Source/Mover/Private/MoverComponent.cpp
+++ b/Engine/Plugins/Experimental/Mover/Source/Mover/Private/MoverComponent.cpp
@@ -11,6 +11,7 @@
 #include "MoveLibrary/MovementMixer.h"
 #include "MoveLibrary/MovementUtils.h"
 #include "MoveLibrary/FloorQueryUtils.h"
+#include "MoveLibrary/RollbackBlackboard.h"
 #include "MoverLog.h"
 #include "InstantMovementEffect.h"
 #include "Backends/MoverNetworkPredictionLiaison.h"
@@ -93,6 +94,20 @@ void UMoverComponent::InitializeComponent()
 
 		SimBlackboard = NewObject<UMoverBlackboard>(this, TEXT("MoverBlackboard"), RF_Transient);
 
+		RollbackBlackboard = NewObject<URollbackBlackboard>(this, TEXT("RollbackBlackboard"), RF_Transient);
+		RollbackBlackboard_InternalWrapper = NewObject<URollbackBlackboard_InternalWrapper>(this, TEXT("RollbackBlackboard_Internal"), RF_Transient);
+		RollbackBlackboard_InternalWrapper->Init(*RollbackBlackboard);
+
+		// create any internal entries
+		static URollbackBlackboard::EntrySettings ModeActivationSettings;
+		ModeActivationSettings.SizingPolicy = EBlackboardSizingPolicy::FixedDeclaredSize;
+		ModeActivationSettings.FixedSize = 4;
+		ModeActivationSettings.PersistencePolicy = EBlackboardPersistencePolicy::Forever;
+		ModeActivationSettings.RollbackPolicy = EBlackboardRollbackPolicy::InvalidatedOnRollback;
+
+		RollbackBlackboard->CreateEntry<FMovementModeChangeRecord>(CommonBlackboard::LastModeChangeRecord, ModeActivationSettings);
+
+
 		FindDefaultUpdatedComponent();
 
 		// Set up FSM and initial movement states
@@ -336,13 +351,13 @@ void UMoverComponent::ProduceInput(const int32 DeltaTimeMS, FMoverInputCmdContex
 	CachedLastProducedInputCmd = *Cmd;
 }
 
-void UMoverComponent::RestoreFrame(const FMoverSyncState* SyncState, const FMoverAuxStateContext* AuxState)
+void UMoverComponent::RestoreFrame(const FMoverSyncState* SyncState, const FMoverAuxStateContext* AuxState, const FMoverTimeStep& NewBaseTimeStep)
 {
 	const FMoverSyncState& InvalidSyncState = GetSyncState();
 	const FMoverAuxStateContext& InvalidAuxState = CachedLastAuxState;
-	OnSimulationPreRollback(&InvalidSyncState, SyncState, &InvalidAuxState, AuxState);
+	OnSimulationPreRollback(&InvalidSyncState, SyncState, &InvalidAuxState, AuxState, NewBaseTimeStep);
 	SetFrameStateFromContext(SyncState, AuxState, /* rebase? */ true);
-	OnSimulationRollback(SyncState, AuxState);
+	OnSimulationRollback(SyncState, AuxState, NewBaseTimeStep);
 }
 
 void UMoverComponent::FinalizeFrame(const FMoverSyncState* SyncState, const FMoverAuxStateContext* AuxState)
@@ -577,6 +592,8 @@ void UMoverComponent::SimulationTick(const FMoverTimeStep& InTimeStep, const FMo
 		OnPreMovement.Broadcast(MoverTimeStep, SimInput.InputCmd, SimInput.SyncState, SimInput.AuxState);
 	}
 
+	RollbackBlackboard_InternalWrapper->BeginSimulationFrame(MoverTimeStep);
+
 	// Tick the actual simulation. This is where the proposed moves are queried and executed, affecting change to the moving actor's gameplay state and captured in the output sim state
 	if (IsInGameThread())
 	{
@@ -604,6 +621,8 @@ void UMoverComponent::SimulationTick(const FMoverTimeStep& InTimeStep, const FMo
 		}
 	}
 
+	RollbackBlackboard_InternalWrapper->EndSimulationFrame();
+
 	if (!SimOutput.MoveRecord.GetTotalMoveDelta().IsZero())
 	{
 		UE_LOG(LogMover, VeryVerbose, TEXT("KinematicSimTick: %s (role %i) frame %d: %s"),
@@ -1020,17 +1039,23 @@ void UMoverComponent::UpdateTickRegistration()
 	SetComponentTickEnabled(bHasUpdatedComponent && bAutoActivate);
 }
 
-void UMoverComponent::OnSimulationPreRollback(const FMoverSyncState* InvalidSyncState, const FMoverSyncState* SyncState, const FMoverAuxStateContext* InvalidAuxState, const FMoverAuxStateContext* AuxState)
+void UMoverComponent::OnSimulationPreRollback(const FMoverSyncState* InvalidSyncState, const FMoverSyncState* SyncState, const FMoverAuxStateContext* InvalidAuxState, const FMoverAuxStateContext* AuxState, const FMoverTimeStep& NewBaseTimeStep)
 {
-	ModeFSM->OnSimulationPreRollback(InvalidSyncState, SyncState, InvalidAuxState, AuxState);
+	ModeFSM->OnSimulationPreRollback(InvalidSyncState, SyncState, InvalidAuxState, AuxState, NewBaseTimeStep);
 }
 
 
-void UMoverComponent::OnSimulationRollback(const FMoverSyncState* SyncState, const FMoverAuxStateContext* AuxState)
+void UMoverComponent::OnSimulationRollback(const FMoverSyncState* SyncState, const FMoverAuxStateContext* AuxState, const FMoverTimeStep& NewBaseTimeStep)
 {
 	SimBlackboard->Invalidate(EInvalidationReason::Rollback);
-	ModeFSM->OnSimulationRollback(SyncState, AuxState);
 
+	RollbackBlackboard_InternalWrapper->BeginRollback();
+
+	RollbackBlackboard_InternalWrapper->RollBack(NewBaseTimeStep.ServerFrame);
+
+	ModeFSM->OnSimulationRollback(SyncState, AuxState, NewBaseTimeStep);
+
+	RollbackBlackboard_InternalWrapper->EndRollback();
 	bHasRolledBack = true;
 }
 
diff --git a/Engine/Plugins/Experimental/Mover/Source/Mover/Public/MoveLibrary/RollbackBlackboard.h b/Engine/Plugins/Experimental/Mover/Source/Mover/Public/MoveLibrary/RollbackBlackboard.h
new file mode 100644
--- /dev/null
+++ b/Engine/Plugins/Experimental/Mover/Source/Mover/Public/MoveLibrary/RollbackBlackboard.h
@@ -0,0 +1,454 @@
+// Copyright Epic Games, Inc. All Rights Reserved.
+
+#pragma once
+
+#include "Containers/CircularBuffer.h"
+#include "MoverLog.h"
+#include "UObject/Interface.h"
+#include "UObject/Object.h"
+#include "Templates/SharedPointer.h"
+#include "Templates/Casts.h"
+#include "RollbackBlackboard.generated.h"
+
+struct FMoverTimeStep;
+
+#define UE_API MOVER_API
+
+
+/** RollbackBlackboard: this is a generic map that stores any type of data for local-only access. It will not be replicated.
+ *  It can be used as a way for decoupled systems to store calculations or transient state data that isn't 
+ *  necessary to reconstitute the movement simulation from scratch.  
+ *  Examples: the normal of the current walkable floor, or time of the last movement mode change
+ *  
+ * Key Features:
+ *  - Rollback support: when movement simulations roll back during a correction, this blackboard also rolls back to previous values.
+ *  - Concurrent access: for async simulations, it provides concurrent access support so external systems / scripting can access data while new data is being authored
+ *  - Policies: there are a variety of policy options to control buffer sizing, invalidation behavior, and entry persistence.
+ * 
+ * Notes:
+ *  - This is implemented under-the-hood using circular buffers, alleviating the need for mem alloc/frees during use. It's important to create entries with policy settings appropriate for their use pattern.
+ *  - We use an "InternalWrapper" class to allow in-simulation objects like movement modes to use the same API, without needing to choose between internal- or external-facing functions. 
+ */
+
+ // TODO: 
+ // - Implement stronger concurrency controls at time of frame changes, and when new in-sim entries are written
+ // - Implement sizing policies that rely on knowing a max number of elements to buffer (the backend may know these details)
+ // - Consider using TArrays instead of circular buffers, due to the power-of-2 trick potentially wasting a lot of memory
+ // - Expose for Blueprint use
+
+/**
+ * Determines how a blackboard's entry buffer is sized
+ */
+UENUM()
+enum class EBlackboardSizingPolicy : uint8
+{
+	/** Buffer size is set at a given number and and never changes */
+	FixedDeclaredSize,
+
+	/** Buffer size matches the backend simulation's history size, determined by ticking rate and history settings.
+	 * A simulation running a fixed 30 fps with a 1 second history will need a buffer size of 30. For variable rate simulations, buffer size cam only be estimated. 
+	 */
+	FixedBackendBufferSize,
+
+	/** Minimizes data size for blackboard entries where thread contention and rolling back shouldn't change the value. Example: a statistic tracking the maximum time spent in a particular mode's simulation tick. */
+	SingleEntry,
+};
+
+/**
+ * Determines how long a blackboard entry is valid once its been set.
+ */
+UENUM()
+enum class EBlackboardPersistencePolicy : uint8
+{
+	Forever,
+
+	NextFrameOnly,
+
+	LimitedTime,
+};
+
+
+/**
+ * Determines how a blackboard's buffered entries are treated when a rollback occurs
+ */
+UENUM()
+enum class EBlackboardRollbackPolicy : uint8
+{
+	/** Entries occurring after the rollback time are effectively erased and no longer retrievable. New entries may be written during resimulation. */
+	InvalidatedOnRollback,
+
+	/** Entries are not invalidated during rollback, and may be read/written during resim. If left alone, they will persist. */
+	WritableDuringResimulation,
+
+	/** Entries are not invalidated during rollback, and may be read during resim. But writes cannot be performed during resim. */
+	LockedDuringResimulation,
+};
+
+
+
+
+/**
+ * This is the core rollback blackboard, with API for external access.
+ * See @URollbackBlackboard_InternalWrapper for in-simulation use.
+ */ 
+UCLASS(MinimalAPI)
+class URollbackBlackboard : public UObject
+{
+	GENERATED_BODY()
+
+public:
+	struct EntrySettings
+	{
+		EBlackboardSizingPolicy      SizingPolicy      = EBlackboardSizingPolicy::FixedDeclaredSize;
+		EBlackboardPersistencePolicy PersistencePolicy = EBlackboardPersistencePolicy::Forever;
+		EBlackboardRollbackPolicy    RollbackPolicy    = EBlackboardRollbackPolicy::InvalidatedOnRollback;
+
+		uint32                       FixedSize = 0;
+		double                       MaxLifetimeSecs = -1.0;
+	};
+
+private:
+
+	struct EntryTimeStamp
+	{
+		double TimeMs = -1.0;
+		uint32 Frame = 0;
+
+		bool IsValid()
+		{
+			return TimeMs >= 0.0;
+		}
+
+		void Invalidate()
+		{
+			TimeMs = -1.0;
+			Frame = 0;
+		}
+	};
+
+	// untyped base, so we can use entries of multiple types in a single-typed container
+	struct BlackboardEntryBase
+	{
+	public:
+		void OnSimulationFrameEnd()
+		{
+			//Advances the external index to match
+			ExternalIdx = InternalIdx;
+		}
+
+		// Note that NewPendingFrame means the frame that will now be re-simulated. So any existing entries with a Frame >= NewPendingFrame can be invalidated.
+		void RollBack(uint32 NewPendingFrame);
+
+	protected:
+		uint32 ComputeBufferSize(EBlackboardSizingPolicy InSizingPolicy, uint32 FixedBufferSize = -1);
+
+		const EntrySettings* Settings = nullptr;
+
+		TCircularBuffer<EntryTimeStamp> Timestamps = TCircularBuffer<EntryTimeStamp>(0);
+
+		uint32 ExternalIdx;	// Indexes to the last entry on a committed frame
+		uint32 InternalIdx;	// Indexes to where in-progress simulations should read/write (matches ExternalIdx when not mid-simulation, and advances if written during simulation)
+	};
+
+	template<typename EntryT>
+	struct BlackboardEntry : BlackboardEntryBase
+	{
+	private:
+
+		BlackboardEntry() = delete;	// disallow use of default constructor
+
+	public:
+
+		// TODO: consider what options need to be specified
+		BlackboardEntry(const EntrySettings& InSettings, const EntryT* OptionalInitialObj=nullptr)
+		{ 
+			Settings = &InSettings;
+
+			int32 BufferSize = ComputeBufferSize(Settings->SizingPolicy, Settings->FixedSize);
+
+			Timestamps = TCircularBuffer<EntryTimeStamp>(BufferSize, EntryTimeStamp());
+
+			// JAH TODO: Consider making the initial object required, so we can also initilize where in the buffer we point to
+			if (OptionalInitialObj)
+			{
+				EntryBuffer = TCircularBuffer<EntryT>(BufferSize, *OptionalInitialObj);
+			}
+			else
+			{
+				EntryBuffer = TCircularBuffer<EntryT>(BufferSize);
+			}
+
+			ExternalIdx = InternalIdx = 0;
+		}
+
+
+		bool TryGetEntryValue(EntryT& ValueOut)
+		{
+			if (!Timestamps[ExternalIdx].IsValid())
+			{
+				return false;	// Probably not initialized yet
+			}
+
+			ValueOut = EntryBuffer[ExternalIdx];
+			return true;
+		}
+
+		// TODO: Consider whether it's worthwhile to have this function... dangerous to let external folks hold onto references to internal storage, but may allow more efficient editing
+		EntryT* TryGetEntryRef()
+		{
+			if (!Timestamps[ExternalIdx].IsValid())
+			{
+				return nullptr;	// Probably not initialized yet
+			}
+
+			return &EntryBuffer[ExternalIdx];
+		}
+
+		bool TryGetEntryValue_Internal(EntryT& ValueOut)
+		{
+			if (!Timestamps[ExternalIdx].IsValid())
+			{
+				return false;	// Probably not initialized yet
+			}
+
+			ValueOut = EntryBuffer[InternalIdx];
+			return true;
+		}
+
+		EntryT* TryGetEntryRef_Internal()
+		{
+			if (!Timestamps[ExternalIdx].IsValid()) 
+			{
+				return nullptr;	// Probably not initialized yet
+			}
+
+			return &EntryBuffer[InternalIdx];
+		}
+
+		void SetEntryValue_Internal(const EntryT& Value, const EntryTimeStamp& TimeStamp)
+		{
+			InternalIdx = ExternalIdx + 1;	// always writing internal, one slot ahead of external slot
+
+			EntryBuffer[InternalIdx] = Value;
+			Timestamps[InternalIdx] = TimeStamp;
+		}
+
+	
+		// TODO: Consider whether these need to be TCircularBuffers, or whether a TArray would be just fine.  TCircularBuffers round up to the next power of 2 in capacity, so may be wasting significant memory for larger structures.
+		// Note that these EntryBuffer is always the same capacity as Timestamps
+		TCircularBuffer<EntryT> EntryBuffer = TCircularBuffer<EntryT>(0);
+
+
+	};	// end BlackboardEntry
+	
+
+
+public:
+	template<typename EntryT>
+	bool CreateEntry(FName EntryName, const EntrySettings& InSettings)
+	{
+		if (EntryMap.Contains(EntryName))
+		{
+			UE_LOG(LogMover, Warning, TEXT("Skipping attempt to create a new blackboard entry named %s since it already exists. Policy settings from the first entry will be retained."), *EntryName.ToString());
+			return false;
+		}
+
+		EntryMap.Emplace(EntryName, MakeUnique<BlackboardEntry<EntryT>>(InSettings, nullptr));
+		return true;
+	}
+
+	/** Store object by a named key, overwriting any existing object */
+	template<typename EntryT>
+	bool TrySet(FName ObjName, EntryT& Obj)
+	{
+		// TODO: Verify this isn't happening mid-sim on the sim thread (_Internal should be used instead)
+
+		if (BlackboardEntry<EntryT>* Entry = FindEntry(ObjName))
+		{
+			Entry->SetEntryValue(Obj);
+			return true;
+		}
+
+		return false;
+	}
+
+
+	/** Attempt to retrieve an object from the blackboard. If found, OutFoundValue will be set. Returns true/false to indicate whether it was found. */
+	template<typename EntryT>
+	bool TryGet(FName ObjName, EntryT& OutFoundValue) const
+	{
+		if (BlackboardEntry<EntryT>* Entry = FindEntry(ObjName))
+		{
+			return Entry->TryGetEntryValue(OutFoundValue);
+		}
+
+		return false;
+	}
+
+	/** Users may access the reference of the entry, which may be more efficient than copy - out, copy - in.
+	 * Although you can get access to the reference, it is not safe to hold references over time. 
+	 */
+	template<typename EntryT>
+	EntryT* TryGetRef(FName ObjName) const
+	{
+		if (BlackboardEntry<EntryT>* Entry = FindEntry(ObjName))
+		{
+			return Entry->TryGetEntryRef();
+		}
+
+		return nullptr;
+	}
+
+
+private:
+	/** Store object by a named key, overwriting any existing object. Should only be called from within the simulation. */
+	template<typename EntryT>
+	bool TrySet_Internal(FName ObjName, EntryT& Obj)
+	{
+		check((bIsSimulationInProgress && (InProgressSimFrameThreadId == FPlatformTLS::GetCurrentThreadId())) || 
+			  (bIsRollbackInProgress && (InRollbackThreadId == FPlatformTLS::GetCurrentThreadId()))) ;
+
+		if (BlackboardEntry<EntryT>* Entry = static_cast<BlackboardEntry<EntryT>*>(FindEntry(ObjName)))
+		{
+			Entry->SetEntryValue_Internal(Obj, InProgressSimTimeStamp);
+			return true;
+		}
+
+		return false;
+	}
+
+	template<typename EntryT>
+	bool TryGet_Internal(FName ObjName, EntryT& OutFoundValue) const
+	{
+		check((bIsSimulationInProgress && (InProgressSimFrameThreadId == FPlatformTLS::GetCurrentThreadId())) ||
+			  (bIsRollbackInProgress && (InRollbackThreadId == FPlatformTLS::GetCurrentThreadId())));
+
+		if (BlackboardEntry<EntryT>* Entry = static_cast<BlackboardEntry<EntryT>*>(FindEntry(ObjName)))
+		{
+			return Entry->TryGetEntryValue_Internal(OutFoundValue);
+		}
+
+		return false;
+	}
+
+	template<typename EntryT>
+	EntryT* TryGetRef_Internal(FName ObjName) const
+	{
+		check((bIsSimulationInProgress && (InProgressSimFrameThreadId == FPlatformTLS::GetCurrentThreadId())) ||
+			  (bIsRollbackInProgress && (InRollbackThreadId == FPlatformTLS::GetCurrentThreadId())));
+
+		if (BlackboardEntry<EntryT>* Entry = static_cast<BlackboardEntry<EntryT>*>(FindEntry(ObjName)))
+		{
+			return Entry->TryGetEntryRef_Internal();
+		}
+
+		return nullptr;
+	}
+
+
+
+	BlackboardEntryBase* FindEntry(FName EntryName) const
+	{
+		if (const TUniquePtr<BlackboardEntryBase>* EntryPtr = EntryMap.Find(EntryName))
+		{
+			return EntryPtr->Get();
+		}
+
+		return nullptr;
+	}
+
+
+private:
+	void BeginSimulationFrame(const FMoverTimeStep& PendingTimeStep);
+	void EndSimulationFrame();
+
+	void BeginRollback();
+	void RollBack(uint32 NewPendingFrame);
+	void EndRollback();
+
+
+private:
+	
+	TMap<FName, TUniquePtr<BlackboardEntryBase>> EntryMap;
+
+
+	bool bIsSimulationInProgress = false;
+	uint32 InProgressSimFrameThreadId;
+
+	bool bIsRollbackInProgress = false;
+	uint32 InRollbackThreadId;
+
+
+	EntryTimeStamp CurrentSimTimeStamp;
+	EntryTimeStamp InProgressSimTimeStamp;	// only valid during a simulation step
+
+	friend class URollbackBlackboard_InternalWrapper;	// expose non-public API to this wrapper class
+};
+
+
+
+/** 
+ * Wrapper class for in-simulation access to the blackboard. This exposes otherwise private API, and redirects
+ * set/get operations to the internal versions.
+ */
+UCLASS(MinimalAPI)
+class URollbackBlackboard_InternalWrapper : public UObject
+{
+	GENERATED_BODY()
+
+public:
+	void Init(URollbackBlackboard& InBlackboard)
+	{
+		Blackboard = &InBlackboard;
+	}
+
+	template<typename EntryT>
+	bool TrySet(FName ObjName, EntryT& Obj)
+	{
+		return Blackboard->TrySet_Internal<EntryT>(ObjName, Obj);
+	}
+
+	template<typename EntryT>
+	bool TryGet(FName ObjName, EntryT& OutFoundValue) const
+	{
+		return Blackboard->TryGet_Internal(ObjName, OutFoundValue);
+	}
+
+	template<typename EntryT>
+	EntryT* TryGetRef(FName ObjName) const
+	{
+		return Blackboard->TryGetRef_Internal<EntryT>(ObjName);
+	}
+
+
+	void BeginSimulationFrame(const FMoverTimeStep& PendingTimeStep)
+	{
+		Blackboard->BeginSimulationFrame(PendingTimeStep);
+	}
+
+	void EndSimulationFrame()
+	{
+		Blackboard->EndSimulationFrame();
+	}
+
+	void BeginRollback()
+	{
+		Blackboard->BeginRollback();
+	}
+
+	void RollBack(uint32 NewPendingFrame)
+	{
+		Blackboard->RollBack(NewPendingFrame);
+	}
+
+	void EndRollback()
+	{
+		Blackboard->EndRollback();
+	}
+
+private: 
+	UPROPERTY()
+	TObjectPtr<URollbackBlackboard> Blackboard;
+};
+
+#undef UE_API
diff --git a/Engine/Plugins/Experimental/Mover/Source/Mover/Public/MovementModeStateMachine.h b/Engine/Plugins/Experimental/Mover/Source/Mover/Public/MovementModeStateMachine.h
--- a/Engine/Plugins/Experimental/Mover/Source/Mover/Public/MovementModeStateMachine.h
+++ b/Engine/Plugins/Experimental/Mover/Source/Mover/Public/MovementModeStateMachine.h
@@ -43,8 +43,8 @@ public:
 	UE_API void ClearQueuedMode();
 
 	UE_API void OnSimulationTick(USceneComponent* UpdatedComponent, UPrimitiveComponent* UpdatedPrimitive, UMoverBlackboard* SimBlackboard, const FMoverTickStartData& StartState, const FMoverTimeStep& TimeStep, FMoverTickEndData& OutputState);
- 	UE_API void OnSimulationPreRollback(const FMoverSyncState* InvalidSyncState, const FMoverSyncState* SyncState, const FMoverAuxStateContext* InvalidAuxState, const FMoverAuxStateContext* AuxState);
-	UE_API void OnSimulationRollback(const FMoverSyncState* SyncState, const FMoverAuxStateContext* AuxState);
+ 	UE_API void OnSimulationPreRollback(const FMoverSyncState* InvalidSyncState, const FMoverSyncState* SyncState, const FMoverAuxStateContext* InvalidAuxState, const FMoverAuxStateContext* AuxState, const FMoverTimeStep& NewBaseTimeStep);
+	UE_API void OnSimulationRollback(const FMoverSyncState* SyncState, const FMoverAuxStateContext* AuxState, const FMoverTimeStep& NewBaseTimeStep);
 
 	FName GetCurrentModeName() const { return CurrentModeName; }
 
@@ -82,6 +82,9 @@ protected:
 	FName DefaultModeName = NAME_None;
 	FName CurrentModeName = NAME_None;
 
+	// Represents the current sim time that's passed, and the next frame number that's next to be simulated.
+	FMoverTimeStep CurrentBaseTimeStep;
+
 	/** Moves that are queued to be added to the simulation at the start of the next sim subtick. Access covered by lock. */
 	TArray<TSharedPtr<FLayeredMoveBase>> QueuedLayeredMoves;
 
diff --git a/Engine/Plugins/Experimental/Mover/Source/Mover/Public/MoverComponent.h b/Engine/Plugins/Experimental/Mover/Source/Mover/Public/MoverComponent.h
--- a/Engine/Plugins/Experimental/Mover/Source/Mover/Public/MoverComponent.h
+++ b/Engine/Plugins/Experimental/Mover/Source/Mover/Public/MoverComponent.h
@@ -27,6 +27,8 @@ struct FInstantMovementEffect;
 struct FMoverSimulationEventData;
 class UMovementModeStateMachine;
 class UMovementMixer;
+class URollbackBlackboard;
+class URollbackBlackboard_InternalWrapper;
 
 namespace MoverComponentConstants
 {
@@ -156,8 +158,8 @@ public:
 	// Get latest local input prior to simulation step. Called by backend system on owner's instance (autonomous or authority).
 	MOVER_API virtual void ProduceInput(const int32 DeltaTimeMS, FMoverInputCmdContext* Cmd);
 
-	// Restore a previous frame prior to resimulating. Called by backend system.
-	MOVER_API void RestoreFrame(const FMoverSyncState* SyncState, const FMoverAuxStateContext* AuxState);
+	// Restore a previous frame prior to resimulating. Called by backend system. NewBaseTimeStep represents the current time and frame we'll simulate next.
+	MOVER_API void RestoreFrame(const FMoverSyncState* SyncState, const FMoverAuxStateContext* AuxState, const FMoverTimeStep& NewBaseTimeStep);
 
 	// Take output for simulation. Called by backend system.
 	MOVER_API void FinalizeFrame(const FMoverSyncState* SyncState, const FMoverAuxStateContext* AuxState);
@@ -472,6 +474,10 @@ public:	// Queries
 
 	MOVER_API UMoverBlackboard* GetSimBlackboard_Mutable() const;
 
+
+	MOVER_API URollbackBlackboard* GetRollbackBlackboard() const { return RollbackBlackboard.Get(); }
+	MOVER_API URollbackBlackboard_InternalWrapper* GetRollbackBlackboard_Internal() const { return RollbackBlackboard_InternalWrapper.Get(); }
+
 	/** Find settings object by type. Returns null if there is none of that type */
 	const IMovementSettingsInterface* FindSharedSettings(const UClass* ByType) const { return FindSharedSettings_Mutable(ByType); }
 	template<typename SettingsT = IMovementSettingsInterface UE_REQUIRES(std::is_base_of_v<IMovementSettingsInterface, SettingsT>)>
@@ -634,11 +640,11 @@ protected:
 	UPROPERTY(BlueprintAssignable, Category = Mover)
 	FMover_OnPreMovement OnPreMovement;
 
-	/** Called when a rollback occurs, before the simulation state has been restored */
-	MOVER_API void OnSimulationPreRollback(const FMoverSyncState* InvalidSyncState, const FMoverSyncState* SyncState, const FMoverAuxStateContext* InvalidAuxState, const FMoverAuxStateContext* AuxState);
+	/** Called when a rollback occurs, before the simulation state has been restored. NewBaseTimeStep represents the current time and frame we're about to simulate. */
+	MOVER_API void OnSimulationPreRollback(const FMoverSyncState* InvalidSyncState, const FMoverSyncState* SyncState, const FMoverAuxStateContext* InvalidAuxState, const FMoverAuxStateContext* AuxState, const FMoverTimeStep& NewBaseTimeStep);
 	
-	/** Called when a rollback occurs, after the simulation state has been restored */
-	MOVER_API void OnSimulationRollback(const FMoverSyncState* SyncState, const FMoverAuxStateContext* AuxState);
+	/** Called when a rollback occurs, after the simulation state has been restored. NewBaseTimeStep represents the current time and frame we're about to simulate. */
+	MOVER_API void OnSimulationRollback(const FMoverSyncState* SyncState, const FMoverAuxStateContext* AuxState, const FMoverTimeStep& NewBaseTimeStep);
 
 	MOVER_API void ProcessFirstSimTickAfterRollback(const FMoverTimeStep& TimeStep);
 
@@ -713,6 +719,15 @@ protected:
 	UPROPERTY(Transient)
 	TObjectPtr<UMoverBlackboard> SimBlackboard = nullptr;
 
+	/** Used to store cached data & computations between decoupled systems, that can be referenced by name. Rollback-aware. */
+	UPROPERTY(Transient)
+	TObjectPtr<URollbackBlackboard> RollbackBlackboard = nullptr;
+
+	UPROPERTY(Transient)
+	TObjectPtr<URollbackBlackboard_InternalWrapper> RollbackBlackboard_InternalWrapper = nullptr;	// This is a thin layer for use only by in-simulation users
+
+	
+	
 private:
 	/** Collection of settings objects that are shared between movement modes. This list is automatically managed based on the @MovementModes contents. */
 	UPROPERTY(EditDefaultsOnly, EditFixedSize, Instanced, Category = Mover, meta = (NoResetToDefault, ObjectMustImplement = "/Script/Mover.MovementSettingsInterface"))
diff --git a/Engine/Plugins/Experimental/Mover/Source/Mover/Public/MoverSimulationTypes.h b/Engine/Plugins/Experimental/Mover/Source/Mover/Public/MoverSimulationTypes.h
--- a/Engine/Plugins/Experimental/Mover/Source/Mover/Public/MoverSimulationTypes.h
+++ b/Engine/Plugins/Experimental/Mover/Source/Mover/Public/MoverSimulationTypes.h
@@ -33,6 +33,8 @@ namespace CommonBlackboard
 	const FName LastFoundDynamicMovementBase = TEXT("LastFoundDynamicMovementBase");
 	const FName LastAppliedDynamicMovementBase = TEXT("LastAppliedDynamicMovementBase");
 	const FName TimeSinceSupported = TEXT("TimeSinceSupported");
+
+	const FName LastModeChangeRecord = TEXT("LastModeChangeRecord");
 }
 
 
diff --git a/Engine/Plugins/Experimental/Mover/Source/Mover/Public/MoverTypes.h b/Engine/Plugins/Experimental/Mover/Source/Mover/Public/MoverTypes.h
--- a/Engine/Plugins/Experimental/Mover/Source/Mover/Public/MoverTypes.h
+++ b/Engine/Plugins/Experimental/Mover/Source/Mover/Public/MoverTypes.h
@@ -88,7 +88,7 @@ struct FMoverTimeStep
 	GENERATED_BODY()
 	
 
-	// The server simulation frame this timestep is associated with
+	// The server simulation frame this timestep is associated with. This is the frame that's being worked on now.
 	UPROPERTY(EditAnywhere, BlueprintReadWrite, Category=Mover)
 	int32 ServerFrame=-1;
 
@@ -352,6 +352,15 @@ struct FMoverDataPersistence
 	bool bCopyFromPriorFrame = true;
 };
 
+// Information about a change in movement mode
+struct FMovementModeChangeRecord
+{
+	FName ModeName;
+	uint32 Frame;
+	double SimTimeMs;
+};
+
+
 
 // Blueprint helper functions for working with a Mover data collection
 UCLASS(MinimalAPI)
