Mover: rollback blackboard tweaks for internal get functions possibly validating against a past timestep, plus documentation improvement

diff --git a/Engine/Plugins/Experimental/Mover/Source/Mover/Public/MoveLibrary/RollbackBlackboard.h b/Engine/Plugins/Experimental/Mover/Source/Mover/Public/MoveLibrary/RollbackBlackboard.h
--- a/Engine/Plugins/Experimental/Mover/Source/Mover/Public/MoveLibrary/RollbackBlackboard.h
+++ b/Engine/Plugins/Experimental/Mover/Source/Mover/Public/MoveLibrary/RollbackBlackboard.h
@@ -50,7 +50,10 @@ enum class EBlackboardSizingPolicy : uint8
 	 */
 	FixedBackendBufferSize,
 
-	/** Minimizes data size for blackboard entries where thread contention and rolling back shouldn't change the value. Example: a statistic tracking the maximum time spent in a particular mode's simulation tick. */
+	/** Minimizes data size for blackboard entries where thread contention and rolling back shouldn't change the value.
+	 * Example: a statistic tracking the maximum time spent in a particular mode's simulation tick. 
+	 * Note this will still have multiple elements for asynchronous simulations, to prevent thread contention over the entry.
+	 */
 	SingleEntry,
 };
 
@@ -206,7 +209,7 @@ private:
 
 		bool TryGetEntryValue_Internal(EntryT& ValueOut)
 		{
-			if (!Timestamps[ExternalIdx].IsValid())
+			if (!Timestamps[InternalIdx].IsValid())
 			{
 				return false;	// Probably not initialized yet
 			}
@@ -217,7 +220,7 @@ private:
 
 		EntryT* TryGetEntryRef_Internal()
 		{
-			if (!Timestamps[ExternalIdx].IsValid()) 
+			if (!Timestamps[InternalIdx].IsValid())
 			{
 				return nullptr;	// Probably not initialized yet
 			}
