Updated SmoothWalkingMode based on design feedback.

diff --git a/Engine/Plugins/Experimental/Mover/Source/Mover/Private/DefaultMovementSet/Modes/SmoothWalkingMode.cpp b/Engine/Plugins/Experimental/Mover/Source/Mover/Private/DefaultMovementSet/Modes/SmoothWalkingMode.cpp
--- a/Engine/Plugins/Experimental/Mover/Source/Mover/Private/DefaultMovementSet/Modes/SmoothWalkingMode.cpp
+++ b/Engine/Plugins/Experimental/Mover/Source/Mover/Private/DefaultMovementSet/Modes/SmoothWalkingMode.cpp
@@ -30,62 +30,132 @@ void USmoothWalkingMode::GenerateWalkMove(const FMoverTickStartData& StartState,
 	}
 	
 	const FMoverDefaultSyncState* StartingSyncState = StartState.SyncState.SyncStateCollection.FindDataByType<FMoverDefaultSyncState>();
+	if (!ensure(StartingSyncState))
+	{
+		return;
+	}
 
 	FMoverTickStartData* MutableStartData = const_cast<FMoverTickStartData*>(&StartState);
-	FSmoothWalkingState& SpringState = MutableStartData->SyncState.SyncStateCollection.FindOrAddMutableDataByType<FSmoothWalkingState>();
-
-	bool bIsAccelerating = (1.1f * DesiredVelocity.SquaredLength()) > InOutVelocity.SquaredLength();
-	float MaxAccelerationToUse =  bIsAccelerating ? TargetAcceleration : TargetDeceleration;
-
-	float DesiredVelocityLengthSq = DesiredVelocity.SquaredLength();
+	
+	// Find or add a FSmoothWalkingState to the SyncState
+	bool bSmoothWalkingStateAdded = false;
+	FSmoothWalkingState& SpringState = MutableStartData->SyncState.SyncStateCollection.FindOrAddMutableDataByType<FSmoothWalkingState>(bSmoothWalkingStateAdded);
+
+	// If the state was not there already we need to initialize some of the intermediate state to whatever we have as the current state to avoid 
+	// a discontinuity. Unfortunately there is no way currently to initialize the angular velocities or accelerations right now as these are not 
+	// carried between movement modes in FMoverDefaultSyncState.
+	if (bSmoothWalkingStateAdded)
+	{
+		SpringState.IntermediateFacing = InOutFacing;
+		SpringState.IntermediateVelocity = InOutVelocity;
+	}
 
-	if (TurningStrength > FLT_EPSILON)
+	// Rotate the intermediate velocity towards the target direction using the TurningStrength
+	if (TurningStrength > 0.0f)
 	{
-		// 'Friction' or Turning spring
-		FVector TargetMoveDir = FVector::ZeroVector;
-		if(DesiredVelocityLengthSq > UE_KINDA_SMALL_NUMBER)
+		if (!DesiredVelocity.IsNearlyZero())
 		{
-			TargetMoveDir = DesiredVelocity.GetSafeNormal();
-		}
-		else if(!InOutVelocity.IsNearlyZero())
-		{
-			TargetMoveDir = InOutVelocity.GetSafeNormal();
+			FMath::ExponentialSmoothingApprox(
+				SpringState.IntermediateVelocity,
+				DesiredVelocity.GetSafeNormal() * SpringState.IntermediateVelocity.Length(),
+				DeltaSeconds,
+				SpringMath::SpringStrengthToSmoothingTime(TurningStrength));
 		}
+	}
+
+	// Check if we are acceleration or decelerating and work out how much lateral vs directional acceleration to apply. Note that even when the
+	// DirectionalAccelerationFactor is high Deceleration is always applied laterally. This is similar to how the default walking mode behaves.
+	const bool bIsAccelerating = (1.01f * DesiredVelocity.SquaredLength()) > InOutVelocity.SquaredLength();
+	const float LateralAccelerationMagnitude =  bIsAccelerating ? (1.0f - DirectionalAccelerationFactor) * Acceleration : Deceleration;
+	const float DirectionalAccelerationMagnitude = bIsAccelerating ? DirectionalAccelerationFactor * Acceleration : 0.0f;
+
+	// Record the previous velocity
+	const FVector PreviousVelocity = SpringState.IntermediateVelocity;
+	
+	// Compute the desired difference in velocity
+	const FVector VelocityDifference = DesiredVelocity - SpringState.IntermediateVelocity;
+
+	// Compute the lateral acceleration moving directly toward the desired velocity
+	const FVector LateralAccelerationVector = VelocityDifference.GetSafeNormal() * FMath::Min(LateralAccelerationMagnitude, VelocityDifference.Length() / DeltaSeconds);
+
+	// Compute the directional acceleration moving in the desired direction. This emulates how acceleration is applied in the default movement mode.
+	const FVector DirectionalAccelerationVector = DesiredVelocity.GetSafeNormal() * DirectionalAccelerationMagnitude;
+
+	// Find the combined desired acceleration as the addition of the two previous accelerations
+	const FVector DesiredAcceleration = LateralAccelerationVector + DirectionalAccelerationVector;
 
-		if (!TargetMoveDir.IsNearlyZero())
+	// Integrate the desired acceleration to estimate the next velocity, and future velocity to track. Here we don't want to over-shoot the desired
+	// velocity target so we just snap to that if we are within enough distance.
+	// 
+	// We then clamp this velocity to make it no larger than the previous intermediate velocity or the total desired velocity this stops the 
+	// DirectionalAcceleration adding velocity to the system and infinitely speeding up the character
+	FVector NextVelocity = VelocityDifference.Length() > DesiredAcceleration.Length() * DeltaSeconds ?
+		SpringState.IntermediateVelocity + DesiredAcceleration * DeltaSeconds : DesiredVelocity;
+
+	NextVelocity = NextVelocity.GetClampedToMaxSize(FMath::Max(PreviousVelocity.Length(), DesiredVelocity.Length()));
+
+	// Compute the smoothing time based on acceleration or deceleration
+	const float VelocitySmoothingTime = bIsAccelerating ? AccelerationSmoothingTime : DecelerationSmoothingTime;
+
+	// Compute the smoothing compensation based on acceleration or deceleration
+	const float VelocitySmoothingCompensation = bIsAccelerating ? AccelerationSmoothingCompensation : DecelerationSmoothingCompensation;
+
+	// Get the lag associated with the current velocity smoothing time and perform the same process to estimate the velocity we need to track to
+	// avoid lagging behind
+	const float LagSeconds =  DeltaSeconds + VelocitySmoothingCompensation * SpringMath::HalfLifeToLag(SpringMath::SpringSmoothingTimeToHalfLife(VelocitySmoothingTime));
+
+	FVector TrackVelocity =
+		VelocityDifference.Length() > DesiredAcceleration.Length() * LagSeconds ?
+		SpringState.IntermediateVelocity + DesiredAcceleration * LagSeconds : DesiredVelocity;
+
+	TrackVelocity = TrackVelocity.GetClampedToMaxSize(FMath::Max(PreviousVelocity.Length(), DesiredVelocity.Length()));
+
+	// Apply the smoothing to the track velocity - effectively tracking the intermediate velocity at the appropriate time in the future
+	SpringMath::SimpleSpringDamper(InOutVelocity, SpringState.SpringAcceleration, TrackVelocity, VelocitySmoothingTime, DeltaSeconds);
+
+	// Snap the velocity to the desired velocity based on the dead-zone
+	if ((DesiredVelocity - InOutVelocity).SquaredLength() < FMath::Square(VelocityDeadzoneThreshold))
+	{
+		// We reached our target
+		InOutVelocity = DesiredVelocity;
+	
+		// If we've reached our target then also snap the acceleration to zero once it is close enough
+		if (SpringState.SpringAcceleration.SquaredLength() < FMath::Square(AccelerationDeadzoneThreshold))
 		{
-			SpringMath::SimpleSpringDamper(InOutVelocity, SpringState.TurningAccel, TargetMoveDir * InOutVelocity.Length(), SpringMath::SpringStrengthToSmoothingTime(TurningStrength), DeltaSeconds);
+			SpringState.SpringAcceleration = FVector::ZeroVector;
 		}
 	}
 
-	// Spring towards desired velocity using a velocity spring (which gives control over acceleration)
+	// Update the intermediate velocity
+	SpringState.IntermediateVelocity = NextVelocity;
 
-	// First we rotate the intermediate velocity towards target so that it doesn't lag in direction
-	if(DesiredVelocityLengthSq > UE_KINDA_SMALL_NUMBER)
+	// For controlling the facing direction we use either a spring or a double spring
+	if (bSmoothFacingWithDoubleSpring)
 	{
-		float VelIntLength = SpringState.IntermediateVelocity.Length();
-		SpringState.IntermediateVelocity = (VelIntLength / FMath::Sqrt(DesiredVelocityLengthSq)) * DesiredVelocity;
+		SpringMath::SimpleSpringDamperQuat(SpringState.IntermediateFacing, SpringState.IntermediateAngularVelocity, DesiredFacing, FacingSmoothingTime / 2.0f, DeltaSeconds);
+		SpringMath::SimpleSpringDamperQuat(InOutFacing, SpringState.CurrentAngularVelocity, SpringState.IntermediateFacing, FacingSmoothingTime / 2.0f, DeltaSeconds);
 	}
-
-	// Now update the velocity spring
-	FVector PreviousVelocity = InOutVelocity;
-	SpringMath::VelocitySpringDamper(InOutVelocity, SpringState.Accel, SpringState.IntermediateVelocity, DesiredVelocity, MaxAccelerationToUse, VelocitySmoothingTime, DeltaSeconds);
-
-	// Clamp accelerations and velocities
-	FVector ActualAccel = (InOutVelocity - PreviousVelocity) / DeltaSeconds;
-
-	float AccelClamp = MaxAccelerationToUse;
-	if(VelocitySmoothingTime > 0.0f)
+	else
 	{
-		AccelClamp *= 1.2f; // add a small buffer to allow spring to catch up to target
+		SpringState.IntermediateFacing = DesiredFacing;
+		SpringState.IntermediateAngularVelocity = SpringState.CurrentAngularVelocity;
+		SpringMath::SimpleSpringDamperQuat(InOutFacing, SpringState.CurrentAngularVelocity, DesiredFacing, FacingSmoothingTime, DeltaSeconds);
 	}
 
-	ActualAccel = ActualAccel.GetClampedToMaxSize(AccelClamp);
-	InOutVelocity = PreviousVelocity + (ActualAccel * DeltaSeconds);
+	// Snap the facing to the desired facing based on the dead-zone
+	if (DesiredFacing.AngularDistance(InOutFacing) < FMath::DegreesToRadians(FacingDeadzoneThreshold))
+	{
+		// We reached our target
+		InOutFacing = DesiredFacing;
+		SpringState.IntermediateFacing = DesiredFacing;
 
-	// Angular / Facing //
-	// We use a simple spring
-	SpringMath::SimpleSpringDamperQuat(InOutFacing, SpringState.CurrentAngularVelocity, DesiredFacing, FacingSmoothingTime, DeltaSeconds);
+		// If we've reached our target then also snap the angular velocity to zero once it is close enough
+		if (SpringState.CurrentAngularVelocity.SquaredLength() < FMath::Square(FMath::DegreesToRadians(AngularVelocityDeadzoneThreshold)))
+		{
+			SpringState.CurrentAngularVelocity = FVector::ZeroVector;
+			SpringState.IntermediateAngularVelocity = FVector::ZeroVector;
+		}
+	}
 }
 
 
diff --git a/Engine/Plugins/Experimental/Mover/Source/Mover/Private/DefaultMovementSet/Modes/SmoothWalkingState.cpp b/Engine/Plugins/Experimental/Mover/Source/Mover/Private/DefaultMovementSet/Modes/SmoothWalkingState.cpp
--- a/Engine/Plugins/Experimental/Mover/Source/Mover/Private/DefaultMovementSet/Modes/SmoothWalkingState.cpp
+++ b/Engine/Plugins/Experimental/Mover/Source/Mover/Private/DefaultMovementSet/Modes/SmoothWalkingState.cpp
@@ -17,9 +17,10 @@ bool FSmoothWalkingState::NetSerialize(FArchive& Ar, UPackageMap* Map, bool& bOu
 	bool bSuccess = Super::NetSerialize(Ar, Map, bOutSuccess);
 
 	// Could be quantized to save bandwidth
-	Ar << Accel;
+	Ar << SpringAcceleration;
 	Ar << IntermediateVelocity;
-	Ar << TurningAccel;
+	Ar << IntermediateFacing;
+	Ar << IntermediateAngularVelocity;
 	Ar << CurrentAngularVelocity;
 
 	return bSuccess;
@@ -29,17 +30,23 @@ void FSmoothWalkingState::ToString(FAnsiStringBuilderBase& Out) const
 {
 	Super::ToString(Out);
 
-	Out.Appendf("Accel=%s IntVel=%s TurningAccel=%s AngVel=%s\n", *Accel.ToCompactString(), *IntermediateVelocity.ToCompactString(), *TurningAccel.ToCompactString(), *CurrentAngularVelocity.ToCompactString());
+	Out.Appendf("SpringAcceleration=%s IntVel=%s IntFac=%s IntAng=%s AngVel=%s\n",
+		*SpringAcceleration.ToCompactString(),
+		*IntermediateVelocity.ToCompactString(), 
+		*IntermediateFacing.ToString(), 
+		*IntermediateAngularVelocity.ToString(),  
+		*CurrentAngularVelocity.ToCompactString());
 }
 
 bool FSmoothWalkingState::ShouldReconcile(const FMoverDataStructBase& AuthorityState) const
 {
 	const FSmoothWalkingState* AuthoritySpringState = static_cast<const FSmoothWalkingState*>(&AuthorityState);
 
-	return (!(Accel - AuthoritySpringState->Accel).IsNearlyZero() ||
+	return (!(SpringAcceleration - AuthoritySpringState->SpringAcceleration).IsNearlyZero() ||
 			!(IntermediateVelocity - AuthoritySpringState->IntermediateVelocity).IsNearlyZero() ||
-			!(TurningAccel - AuthoritySpringState->TurningAccel).IsNearlyZero() ||
-			!(CurrentAngularVelocity - AuthoritySpringState->CurrentAngularVelocity).IsNearlyZero());
+			 (IntermediateFacing.AngularDistance(AuthoritySpringState->IntermediateFacing) > UE_SMALL_NUMBER ||
+			!(IntermediateAngularVelocity - AuthoritySpringState->IntermediateAngularVelocity).IsNearlyZero() ||
+			!(CurrentAngularVelocity - AuthoritySpringState->CurrentAngularVelocity).IsNearlyZero()));
 }
 
 
@@ -48,8 +55,9 @@ void FSmoothWalkingState::Interpolate(const FMoverDataStructBase& From, const FM
 	const FSmoothWalkingState* FromState = static_cast<const FSmoothWalkingState*>(&From);
 	const FSmoothWalkingState* ToState   = static_cast<const FSmoothWalkingState*>(&To);
 
-	Accel = FMath::Lerp(FromState->Accel, ToState->Accel, Pct);
+	SpringAcceleration = FMath::Lerp(FromState->SpringAcceleration, ToState->SpringAcceleration, Pct);
 	IntermediateVelocity = FMath::Lerp(FromState->IntermediateVelocity, ToState->IntermediateVelocity, Pct);
-	TurningAccel = FMath::Lerp(FromState->TurningAccel, ToState->TurningAccel, Pct);
+	IntermediateFacing = FQuat::Slerp(FromState->IntermediateFacing, ToState->IntermediateFacing, Pct);
+	IntermediateAngularVelocity = FMath::Lerp(FromState->IntermediateAngularVelocity, ToState->IntermediateAngularVelocity, Pct);
 	CurrentAngularVelocity = FMath::Lerp(FromState->CurrentAngularVelocity, ToState->CurrentAngularVelocity, Pct);
 }
diff --git a/Engine/Plugins/Experimental/Mover/Source/Mover/Private/DefaultMovementSet/Modes/SmoothWalkingState.h b/Engine/Plugins/Experimental/Mover/Source/Mover/Private/DefaultMovementSet/Modes/SmoothWalkingState.h
--- a/Engine/Plugins/Experimental/Mover/Source/Mover/Private/DefaultMovementSet/Modes/SmoothWalkingState.h
+++ b/Engine/Plugins/Experimental/Mover/Source/Mover/Private/DefaultMovementSet/Modes/SmoothWalkingState.h
@@ -22,18 +22,23 @@ struct FSmoothWalkingState : public FMoverDataStructBase
 	virtual bool ShouldReconcile(const FMoverDataStructBase& AuthorityState) const override;
 	virtual void Interpolate(const FMoverDataStructBase& From, const FMoverDataStructBase& To, float Pct) override;
 
-
-	// Acceleration of internal spring model
+	// Acceleration of internal velocity spring
 	UPROPERTY(BlueprintReadOnly, Category = "Mover|Experimental")
-	FVector Accel = FVector::ZeroVector;
+	FVector SpringAcceleration = FVector::ZeroVector;
 
+	// Intermediate velocity which the velocity spring tracks as a target
 	UPROPERTY(BlueprintReadOnly, Category = "Mover|Experimental")
 	FVector IntermediateVelocity = FVector::ZeroVector;
 
+	// Intermediate facing direction when using a double spring
+	UPROPERTY(BlueprintReadOnly, Category = "Mover|Experimental")
+	FQuat IntermediateFacing = FQuat::Identity;
+
+	// Angular velocity of the intermediate spring when using a double spring
 	UPROPERTY(BlueprintReadOnly, Category = "Mover|Experimental")
-	FVector TurningAccel = FVector::ZeroVector;
+	FVector IntermediateAngularVelocity = FVector::ZeroVector;
 
-	// Angular velocity of internal spring model in radians per second
+	// Angular velocity of rotational spring model in radians per second
 	UPROPERTY(BlueprintReadOnly, Category = "Mover|Experimental")
 	FVector CurrentAngularVelocity = FVector::ZeroVector;
 };
diff --git a/Engine/Plugins/Experimental/Mover/Source/Mover/Public/DefaultMovementSet/Modes/SmoothWalkingMode.h b/Engine/Plugins/Experimental/Mover/Source/Mover/Public/DefaultMovementSet/Modes/SmoothWalkingMode.h
--- a/Engine/Plugins/Experimental/Mover/Source/Mover/Public/DefaultMovementSet/Modes/SmoothWalkingMode.h
+++ b/Engine/Plugins/Experimental/Mover/Source/Mover/Public/DefaultMovementSet/Modes/SmoothWalkingMode.h
@@ -10,8 +10,7 @@
 #define UE_API MOVER_API
 
 /**
- * A walking mode that uses a critically damped spring for translation and rotation.
- * The strength of the critically damped spring is set via smoothing times (separate for translation and rotation)
+ * A walking mode that provides a simplified version of the default walking mode model with additional options for smoothing.
  */
 UCLASS(BlueprintType, Experimental)
 class USmoothWalkingMode : public USimpleWalkingMode
@@ -23,25 +22,82 @@ public:
 	UE_API virtual void GenerateWalkMove(const FMoverTickStartData& StartState, float DeltaSeconds, const FVector& DesiredVelocity,
 									 const FQuat& DesiredFacing, FVector& InOutVelocity, FQuat& InOutFacing) const override;
 
-protected:
-	UPROPERTY(EditAnywhere, BlueprintReadWrite, Category = "Mover|Spring Settings", meta = (ClampMin = "0", UIMin = "0", ForceUnits = "s"))
-	float VelocitySmoothingTime = 0.2f;
 
-	UPROPERTY(EditAnywhere, BlueprintReadWrite, Category = "Mover|Spring Settings", meta = (ClampMin = "0", UIMin = "0", ForceUnits = "s"))
-	float FacingSmoothingTime = 0.2f;
+protected: // Velocity Controls
 
-	UPROPERTY(EditAnywhere, BlueprintReadWrite, Category = "Mover|Spring Settings", meta = (ClampMin = "0", UIMin = "0", ForceUnits = "cm/s^2"))
-	float TargetAcceleration = 700.0f;
+	// Base acceleration to apply when the desired velocity magnitude is greater than the current velocity.
+	UPROPERTY(EditAnywhere, BlueprintReadWrite, Category = "Mover|Smooth Walking Settings", meta = (ClampMin = "0", UIMin = "0", ForceUnits = "cm/s^2"))
+	float Acceleration = 800.0f;
 
-	UPROPERTY(EditAnywhere, BlueprintReadWrite, Category = "Mover|Spring Settings", meta = (ClampMin = "0", UIMin = "0", ForceUnits = "cm/s^2"))
-	float TargetDeceleration = 700.0f;
+	// Base deceleration to apply when the desired velocity magnitude is less than the current velocity
+	UPROPERTY(EditAnywhere, BlueprintReadWrite, Category = "Mover|Smooth Walking Settings", meta = (ClampMin = "0", UIMin = "0", ForceUnits = "cm/s^2"))
+	float Deceleration = 1000.0f;
 
-	UPROPERTY(EditAnywhere, BlueprintReadWrite, Category = "Mover|Spring Settings", meta = (ClampMin = "0", UIMin = "0"))
-	float TurningStrength = 10.0f;
+	/**
+	 * Value between 0 and 1 that controls how the acceleration is applied. When set to 1 this will effectively add the acceleration in the direction 
+	 * of the desired velocity on top of the current velocity then clip the result. This emulates the behavior of the default walking mode. When set 
+	 * to 0 it will apply an acceleration which changes the current velocity directly towards the desired velocity. This makes directional changes in 
+	 * velocity faster and more regular, always using a fixed acceleration, but can cause the overall velocity magnitude to dip during turning.
+	 */
+	UPROPERTY(EditAnywhere, BlueprintReadWrite, Category = "Mover|Advanced Smooth Walking Settings", meta = (ClampMin = "0", UIMin = "0", ClampMax = "1", UIMax = "1"))
+	float DirectionalAccelerationFactor = 0.0f;
+
+	/**
+	 * Applies an additional force that rotates the velocity to make it face the desired direction of travel. This allows quick turns which don't 
+	 * lose speed but adds additional acceleration to the system. Valid values range roughly between 0 and 100, where 10 would be considered already 
+	 * quite a strong turning force and 100 will produce near instant turns (before smoothing).
+	 */
+	UPROPERTY(EditAnywhere, BlueprintReadWrite, Category = "Mover|Smooth Walking Settings", meta = (ClampMin = "0", UIMin = "0"))
+	float TurningStrength = 3.0f;
+
+	// Controls how much smoothing is applied to the velocity changes of the movement when accelerating. Set to zero to apply no smoothing.
+	UPROPERTY(EditAnywhere, BlueprintReadWrite, Category = "Mover|Smooth Walking Settings", meta = (ClampMin = "0", UIMin = "0", ForceUnits = "s"))
+	float AccelerationSmoothingTime = 0.2f;
+
+	// Controls how much smoothing is applied to the velocity changes of the movement when decelerating. Set to zero to apply no smoothing.
+	UPROPERTY(EditAnywhere, BlueprintReadWrite, Category = "Mover|Smooth Walking Settings", meta = (ClampMin = "0", UIMin = "0", ForceUnits = "s"))
+	float DecelerationSmoothingTime = 0.05f;
+
+	/**
+	 * When the velocity is smoothed it will naturally lag behind the unsmoothed target. This factor controls how much this is compensated for by
+	 * tracking a future target that is the appropriate amount of time in the future when accelerating. Setting this to 0 will produce smoother, more 
+	 * S-shaped velocity profiles but may reduce the feeling of responsiveness. Setting this to 1 increases the feeling of responsiveness but makes 
+	 * the initial change in velocity less smooth and reduces the effective "lead-in".
+	 */
+	UPROPERTY(EditAnywhere, BlueprintReadWrite, Category = "Mover|Advanced Smooth Walking Settings", meta = (ClampMin = "0", UIMin = "0", ClampMax = "1", UIMax = "1"))
+	float AccelerationSmoothingCompensation = 1.0f;
+
+	// This parameter acts the same as AccelerationSmoothingCompensation but is applied during deceleration instead
+	UPROPERTY(EditAnywhere, BlueprintReadWrite, Category = "Mover|Advanced Smooth Walking Settings", meta = (ClampMin = "0", UIMin = "0", ClampMax = "1", UIMax = "1"))
+	float DecelerationSmoothingCompensation = 0.0f;
+
+	// Controls the point at which the velocity will "snap" to the desired velocity once it is close enough 
+	UPROPERTY(EditAnywhere, BlueprintReadWrite, Category = "Mover|Advanced Smooth Walking Settings", meta = (ClampMin = "0", UIMin = "0", ForceUnits = "cm/s"))
+	float VelocityDeadzoneThreshold = 0.01f;
+
+	// Controls the point at which the acceleration will "snap" to zero when the desired velocity is reached
+	UPROPERTY(EditAnywhere, BlueprintReadWrite, Category = "Mover|Advanced Smooth Walking Settings", meta = (ClampMin = "0", UIMin = "0", ForceUnits = "cm/s^2"))
+	float AccelerationDeadzoneThreshold = 0.001f;
+
+
+protected: // Facing Controls
+
+	// Controls how much smoothing is applied to the facing direction. Set to zero to apply no smoothing.
+	UPROPERTY(EditAnywhere, BlueprintReadWrite, Category = "Mover|Smooth Walking Settings", meta = (ClampMin = "0", UIMin = "0", ForceUnits = "s"))
+	float FacingSmoothingTime = 0.3f;
+
+	// Smooths facing using a double spring rather than a single spring. This produces a more S-shaped profile with a shorter "lead-out".
+	UPROPERTY(EditAnywhere, BlueprintReadWrite, Category = "Mover|Advanced Smooth Walking Settings")
+	bool bSmoothFacingWithDoubleSpring = true;
+
+	// Controls the point at which the facing will "snap" to the desired facing once it is close enough 
+	UPROPERTY(EditAnywhere, BlueprintReadWrite, Category = "Mover|Advanced Smooth Walking Settings", meta = (ClampMin = "0", UIMin = "0", ForceUnits = "deg"))
+	float FacingDeadzoneThreshold = 0.1f;
+
+	// Controls the point at which the angular velocity will "snap" to zero when the desired facing is reached
+	UPROPERTY(EditAnywhere, BlueprintReadWrite, Category = "Mover|Advanced Smooth Walking Settings", meta = (ClampMin = "0", UIMin = "0", ForceUnits = "deg/s"))
+	float AngularVelocityDeadzoneThreshold = 0.01f;
 
-	UPROPERTY(EditAnywhere, BlueprintReadWrite, Category = "Mover|Spring Settings", meta = (ClampMin = "0", UIMin = "0", ForceUnits = "cm/s"))
-	// Below this speed we set velocity to 0
-	float VelocityDeadzoneThreshold = 0.1f;
 };
 
 #undef UE_API
diff --git a/Engine/Plugins/Experimental/Mover/Source/Mover/Public/MoverTypes.h b/Engine/Plugins/Experimental/Mover/Source/Mover/Public/MoverTypes.h
--- a/Engine/Plugins/Experimental/Mover/Source/Mover/Public/MoverTypes.h
+++ b/Engine/Plugins/Experimental/Mover/Source/Mover/Public/MoverTypes.h
@@ -278,6 +278,20 @@ struct FMoverDataCollection
 		return *(static_cast<T*>(AddDataByType(T::StaticStruct())));
 	}
 
+	/** Find data of a specific type in the collection (mutable version). If not found, a new default instance will be added. Outputs if the state was added or not. */
+	template <typename T>
+	T& FindOrAddMutableDataByType(bool& bOutAdded)
+	{
+		if (T* ExistingData = FindMutableDataByType<T>())
+		{
+			bOutAdded = false;
+			return *ExistingData;
+		}
+
+		bOutAdded = true;
+		return *(static_cast<T*>(AddDataByType(T::StaticStruct())));
+	}
+
 	/** Adds this data instance to the collection, taking ownership of it. If an existing data struct of the same type is already there, it will be removed first. */
 	UE_API void AddOrOverwriteData(const TSharedPtr<FMoverDataStructBase> DataInstance);
 
