Mover: Adding support for blueprintable layered moves through a new system that splits up Layered Move Logic and Data. This allows Layered Move Logic to be made natively and in Blueprints in an effort to expand blueprint support in Mover. These new layered moves exist alongside the old ones at the moment with the intention of these new moves eventually replacing the old ones. Currently, only Launch Layered move has been converted to the new system.

#virtualized

diff --git a/Engine/Plugins/Experimental/Mover/Source/Mover/Private/DefaultMovementSet/LayeredMoves/LaunchMove.cpp b/Engine/Plugins/Experimental/Mover/Source/Mover/Private/DefaultMovementSet/LayeredMoves/LaunchMove.cpp
--- a/Engine/Plugins/Experimental/Mover/Source/Mover/Private/DefaultMovementSet/LayeredMoves/LaunchMove.cpp
+++ b/Engine/Plugins/Experimental/Mover/Source/Mover/Private/DefaultMovementSet/LayeredMoves/LaunchMove.cpp
@@ -2,11 +2,59 @@
 
 #include "DefaultMovementSet/LayeredMoves/LaunchMove.h"
 #include "MoverComponent.h"
-#include "MoverTypes.h"
-#include "MoverLog.h"
 
 #include UE_INLINE_GENERATED_CPP_BY_NAME(LaunchMove)
 
+void FLaunchMoveData::ActivateFromContext(const FLayeredMoveActivationParams* ActivationParams)
+{
+	if (ActivationParams)
+	{
+		if (const FLaunchMoveActivationParams* LaunchMoveActivationParams = static_cast<const FLaunchMoveActivationParams*>(ActivationParams))
+		{
+			LaunchVelocity = LaunchMoveActivationParams->LaunchVelocity;
+			DurationMs = LaunchMoveActivationParams->DurationMs;
+			ForceMovementMode = LaunchMoveActivationParams->ForceMovementMode;
+		}
+	}
+}
+
+void FLaunchMoveData::NetSerialize(FArchive& Ar)
+{
+	Super::NetSerialize(Ar);
+
+	SerializePackedVector<10, 16>(LaunchVelocity, Ar);
+
+	bool bUsingForcedMovementMode = !ForceMovementMode.IsNone();
+
+	Ar.SerializeBits(&bUsingForcedMovementMode, 1);
+
+	if (bUsingForcedMovementMode)
+	{
+		Ar << ForceMovementMode;
+	}
+	else
+	{
+		ForceMovementMode = NAME_None;
+	}
+}
+
+ULaunchMoveLogic::ULaunchMoveLogic()
+{
+	DefaultDurationMs = 0.f;
+	MixMode = EMoveMixMode::OverrideVelocity;
+	InstancedDataStructType = FLaunchMoveData::StaticStruct();
+}
+
+bool ULaunchMoveLogic::GenerateMove_Implementation(const FMoverTimeStep& TimeStep, UMoverBlackboard* SimBlackboard, const FMoverTickStartData& StartState, FProposedMove& OutProposedMove)
+{
+	const FLaunchMoveData& LaunchMoveData = AccessExecutionMoveData<FLaunchMoveData>();
+	
+	OutProposedMove.MixMode = MixMode;
+	OutProposedMove.LinearVelocity = LaunchMoveData.LaunchVelocity;
+	OutProposedMove.PreferredMode = LaunchMoveData.ForceMovementMode;
+
+	return true;
+}
 
 FLayeredMove_Launch::FLayeredMove_Launch()
 {
diff --git a/Engine/Plugins/Experimental/Mover/Source/Mover/Private/LayeredMoveBase.cpp b/Engine/Plugins/Experimental/Mover/Source/Mover/Private/LayeredMoveBase.cpp
new file mode 100644
--- /dev/null
+++ b/Engine/Plugins/Experimental/Mover/Source/Mover/Private/LayeredMoveBase.cpp
@@ -0,0 +1,419 @@
+﻿// Copyright Epic Games, Inc. All Rights Reserved.
+
+#include "LayeredMoveBase.h"
+#include "MoverLog.h"
+#include "MoverComponent.h"
+#include "MoverSimulationTypes.h"
+#include "Blueprint/BlueprintExceptionInfo.h"
+#include "MoveLibrary/MovementMixer.h"
+
+#include UE_INLINE_GENERATED_CPP_BY_NAME(LayeredMoveBase)
+
+#define LOCTEXT_NAMESPACE "Mover"
+
+//////////////////////////////////////////////////////////////////////////
+// FLayeredMoveInstancedData
+
+bool FLayeredMoveInstancedData::operator==(const FLayeredMoveInstancedData& Other) const
+{
+	//@todo DanH: Compare logic class compatibility?
+	if (GetScriptStruct() == Other.GetScriptStruct())
+	{
+		return Equals(Other);
+	}
+	return false;
+}
+
+bool FLayeredMoveInstancedData::Equals(const FLayeredMoveInstancedData& OtherData) const
+{
+	return StartSimTimeMs == OtherData.StartSimTimeMs
+		&& DurationMs == OtherData.DurationMs;
+}
+
+void FLayeredMoveInstancedData::ActivateFromContext(const FLayeredMoveActivationParams* ActivationParams)
+{
+	
+}
+
+void FLayeredMoveInstancedData::NetSerialize(FArchive& Ar)
+{
+	Ar << StartSimTimeMs;
+	Ar << DurationMs;
+}
+
+//////////////////////////////////////////////////////////////////////////
+// ULayeredMoveLogic
+
+ULayeredMoveLogic::ULayeredMoveLogic()
+	: InstancedDataStructType(FLayeredMoveInstancedData::StaticStruct())
+{
+}
+
+bool ULayeredMoveLogic::IsFinished_Implementation(const FMoverTimeStep& TimeStep, const UMoverBlackboard* SimBlackboard)
+{
+	const FLayeredMoveInstancedData& ExecData = AccessExecutionMoveData();
+	const double DurationMs = ExecData.DurationMs;
+	if (DurationMs < 0.f)
+	{
+		return false;
+	}
+	
+	const double StartSimTimeMs = ExecData.GetStartSimTimeMs();
+	
+	const bool bHasStarted = StartSimTimeMs >= 0.f;
+	const bool bTimeExpired = StartSimTimeMs + DurationMs <= TimeStep.BaseSimTimeMs;
+	return bHasStarted && bTimeExpired;	
+}
+
+bool ULayeredMoveLogic::K2_GetActiveMoveData(ULayeredMoveLogic* MoveLogic, FLayeredMoveInstancedData& OutMoveData)
+{
+	// This will never be called, the exec version below will be hit instead
+	checkNoEntry();
+	return false;
+}
+
+
+DEFINE_FUNCTION(ULayeredMoveLogic::execK2_GetActiveMoveData)
+{
+	P_GET_OBJECT(ULayeredMoveLogic, MoveLogic);
+	
+	Stack.MostRecentPropertyAddress = nullptr;
+	Stack.MostRecentPropertyContainer = nullptr;
+	Stack.StepCompiledIn<FStructProperty>(nullptr);
+
+	const FStructProperty* MoveDataProperty = CastField<FStructProperty>(Stack.MostRecentProperty);
+	uint8* OutMoveDataPtr = Stack.MostRecentPropertyAddress;
+	
+	P_FINISH
+
+	const bool bHasValidMoveData = ValidateMoveDataGetSet(P_THIS, MoveLogic, MoveDataProperty, OutMoveDataPtr, Stack);
+	if (bHasValidMoveData)
+	{
+		// Write the active move data to the function output
+		P_NATIVE_BEGIN;
+		MoveDataProperty->Struct->CopyScriptStruct(OutMoveDataPtr, MoveLogic->CurrentInstancedData.Get());
+		P_NATIVE_END;
+	}
+	
+	*(bool*)RESULT_PARAM = bHasValidMoveData;
+}
+
+void ULayeredMoveLogic::K2_SetActiveMoveData(ULayeredMoveLogic* MoveLogic, const FLayeredMoveInstancedData& OutMoveData)
+{
+	// This will never be called, the exec version below will be hit instead
+	checkNoEntry();
+}
+
+DEFINE_FUNCTION(ULayeredMoveLogic::execK2_SetActiveMoveData)
+{
+	P_GET_OBJECT(ULayeredMoveLogic, MoveLogic);
+	
+	Stack.MostRecentPropertyAddress = nullptr;
+	Stack.MostRecentPropertyContainer = nullptr;
+	Stack.StepCompiledIn<FStructProperty>(nullptr);
+
+	const FStructProperty* MoveDataProperty = CastField<FStructProperty>(Stack.MostRecentProperty);
+	const uint8* MoveDataPtr = Stack.MostRecentPropertyAddress;
+
+	P_FINISH
+
+	if (ValidateMoveDataGetSet(P_THIS, MoveLogic, MoveDataProperty, MoveDataPtr, Stack))
+	{
+		// Overwrite the contents of the current move data with that provided
+		P_NATIVE_BEGIN;
+		MoveDataProperty->Struct->CopyScriptStruct(MoveLogic->CurrentInstancedData.Get(), MoveDataPtr);
+		P_NATIVE_END;
+	}
+}
+
+FLayeredMoveInstance::FLayeredMoveInstance()
+	: InstanceMoveData(nullptr)
+	, MoveLogic(nullptr)
+{
+}
+
+FLayeredMoveInstance::FLayeredMoveInstance(const FLayeredMoveInstance& OtherActiveLayeredMove)
+	: InstanceMoveData(OtherActiveLayeredMove.InstanceMoveData)
+	, MoveLogic(OtherActiveLayeredMove.MoveLogic)
+{
+}
+
+FLayeredMoveInstance::FLayeredMoveInstance(const TSharedRef<FLayeredMoveInstancedData>& InMoveData, ULayeredMoveLogic* InMoveLogic)
+	: InstanceMoveData(InMoveData)
+	, MoveLogic(InMoveLogic)
+{
+}
+
+bool ULayeredMoveLogic::ValidateMoveDataGetSet(const UObject* ObjectValidatingData, const ULayeredMoveLogic* MoveLogic, const FStructProperty* MoveDataProperty, const uint8* MoveDataPtr, FFrame& StackFrame)
+{
+	FText BlueprintExceptionText;
+	bool bIsValid = true;
+	if (!MoveLogic)
+	{
+		BlueprintExceptionText = LOCTEXT("ValidateMoveDataNoMoveLogic", "No MoveLogic was present.");
+		bIsValid = false;
+	}
+	else if (!MoveLogic->CurrentInstancedData)
+	{
+		BlueprintExceptionText = LOCTEXT("ValidateMoveDataNoCurrentActiveMoveData", "No CurrentActiveMoveData on the MoveLogic was present.");
+		bIsValid = false;
+	}
+	else if (!MoveDataProperty)
+	{
+		BlueprintExceptionText = LOCTEXT("ValidateMoveDataNoMoveData", "No MoveData was present.");
+		bIsValid = false;
+	}
+	else if (!MoveDataPtr)
+	{
+		BlueprintExceptionText = LOCTEXT("ValidateMoveDataNoMoveData", "No MoveData was present.");
+		bIsValid = false;
+	}
+	else if (MoveDataProperty->Struct != MoveLogic->InstancedDataStructType.Get())
+	{
+		BlueprintExceptionText = FText::Format(LOCTEXT("ValidateMoveDataMisMatchData", "MoveData passed in did not match MoveLogic Active mode data. Expected: {0}. Found: {1}."),
+			FText::FromString(MoveDataProperty->Struct.GetName()),
+			FText::FromString(MoveLogic->InstancedDataStructType.GetName())
+		);
+		bIsValid = false;
+	}
+
+	if (!bIsValid)
+	{
+		const FBlueprintExceptionInfo ExceptionInfo(
+			EBlueprintExceptionType::NonFatalError,
+			BlueprintExceptionText
+		);
+		
+		FBlueprintCoreDelegates::ThrowScriptException(ObjectValidatingData, StackFrame, ExceptionInfo);
+	}
+	
+	return bIsValid;
+}
+
+//////////////////////////////////////////////////////////////////////////
+// FLayeredMoveInstance
+
+/**
+ * Scoped wrapper that is the only means of calling the virtual functions in ULayeredMoveLogic
+ * that depend on/expect access to valid active move data.
+ */
+class FScopedMoveLogicExecContext
+{
+public:	
+	FScopedMoveLogicExecContext(ULayeredMoveLogic& MoveLogic, const TSharedRef<FLayeredMoveInstancedData>& MoveData)
+		: LogicObj(MoveLogic)
+	{
+		MoveLogic.CurrentInstancedData = MoveData.ToSharedPtr();
+	}
+
+	template <typename... ArgsT>
+    void OnStart(ArgsT&&... Args)
+    {
+    	LogicObj.OnStart(Args...);
+    }
+	
+	template <typename... ArgsT>
+	void OnEnd(ArgsT&&... Args)
+	{
+		LogicObj.OnEnd(Args...);
+	}
+	
+	template <typename... ArgsT>
+	bool GenerateMove(ArgsT&&... Args)
+	{
+		return LogicObj.GenerateMove(Args...);
+	}
+
+	template <typename... ArgsT>
+	bool IsFinished(ArgsT&&... Args)
+	{
+		return LogicObj.IsFinished(Args...);
+	}
+	
+	~FScopedMoveLogicExecContext()
+	{
+		LogicObj.CurrentInstancedData.Reset();
+	}
+
+private:
+	ULayeredMoveLogic& LogicObj;
+};
+
+void FLayeredMoveInstance::StartMove(const FMoverTimeStep& TimeStep, UMoverBlackboard* SimBlackboard) const
+{
+	if (ensure(MoveLogic))
+	{
+		FScopedMoveLogicExecContext LogicExecContext(*MoveLogic, InstanceMoveData.ToSharedRef());
+		LogicExecContext.OnStart(TimeStep, SimBlackboard);
+		
+		InstanceMoveData->StartSimTimeMs = TimeStep.BaseSimTimeMs;
+	}
+}
+
+bool FLayeredMoveInstance::GenerateMove(const FMoverTickStartData& StartState, const FMoverTimeStep& TimeStep, UMoverBlackboard* SimBlackboard, FProposedMove& OutProposedMove) const
+{
+	if (ensure(MoveLogic))
+	{
+		FScopedMoveLogicExecContext LogicExecContext(*MoveLogic, InstanceMoveData.ToSharedRef());
+		if (LogicExecContext.GenerateMove(TimeStep, SimBlackboard, StartState, OutProposedMove))
+		{
+			//@todo DanH: Feels like the preferred mode from moves should come from OnStart if that's the only time we listen to it
+			// Wipe the preferred mode if this wasn't the first tick of the mode
+			if (InstanceMoveData->GetStartSimTimeMs() < TimeStep.BaseSimTimeMs)
+			{
+				OutProposedMove.PreferredMode = NAME_None;
+			}
+			return true;
+		}
+	}
+	
+	return false;
+}
+
+void FLayeredMoveInstance::EndMove(const FMoverTimeStep& TimeStep, UMoverBlackboard* SimBlackboard) const
+{
+	if (ensure(MoveLogic))
+	{
+		FScopedMoveLogicExecContext LogicExecContext(*MoveLogic, InstanceMoveData.ToSharedRef());
+		LogicExecContext.OnEnd(TimeStep, SimBlackboard);
+	}
+}
+
+bool FLayeredMoveInstance::IsFinished(const FMoverTimeStep& TimeStep, const UMoverBlackboard* SimBlackboard) const
+{
+	if (ensure(MoveLogic))
+	{
+		FScopedMoveLogicExecContext LogicExecContext(*MoveLogic, InstanceMoveData.ToSharedRef());
+		return LogicExecContext.IsFinished(TimeStep, SimBlackboard);
+	}
+	return true;
+}
+
+const FLayeredMoveFinishVelocitySettings& FLayeredMoveInstance::GetFinishVelocitySettings() const
+{
+	if (MoveLogic)
+	{
+		return MoveLogic->GetFinishVelocitySettings();
+	}
+	
+	static const FLayeredMoveFinishVelocitySettings Defaults;
+	return Defaults;
+}
+
+EMoveMixMode FLayeredMoveInstance::GetMixMode() const
+{
+	if (MoveLogic)
+	{
+		return MoveLogic->GetMixMode();
+	}
+	return EMoveMixMode::AdditiveVelocity;
+}
+
+FLayeredMoveInstance FLayeredMoveInstance::Clone() const
+{
+	return FLayeredMoveInstance(TSharedRef<FLayeredMoveInstancedData>(InstanceMoveData->Clone()), MoveLogic.Get());
+}
+
+void FLayeredMoveInstance::NetSerialize(FArchive& Ar)
+{
+	// Step carefully! When writing from an archive, this move may have been created via TArray::AddZeroed, which does not call the constructor
+	// Therefore, this is the one spot where we are actually ok with MoveData being null, but that violates the rules of TSharedRef.
+	// So the ref is converted first to a TSharedPtr so we can safely, albeit cheekily, check it for validity.
+	TSharedPtr<FLayeredMoveInstancedData> PossiblyNullData = InstanceMoveData;
+	TCheckedObjPtr<UScriptStruct> DataStructType = PossiblyNullData ? PossiblyNullData->GetScriptStruct() : nullptr;
+	UScriptStruct* ExistingDataType = DataStructType.Get();
+	Ar << DataStructType;
+
+	UClass* CurrentMoveLogicClassType = MoveLogic ? MoveLogic->GetClass() : nullptr;
+	
+	Ar << CurrentMoveLogicClassType;
+	if (!MoveLogic)
+	{
+		MoveLogicClassType = CurrentMoveLogicClassType;
+	}
+	
+	if (DataStructType.IsValid())
+	{
+		//@todo DanH: Fwiw FLayeredMoveInstanceGroup is not actually ever expected to be sent in a server RPC right?
+		// Restrict replication to derived classes of FLayeredMoveInstancedData for security reasons:
+		// If FLayeredMoveInstanceGroup is replicated through a Server RPC, we need to prevent clients from sending us
+		// arbitrary ScriptStructs due to the allocation/reliance on GetCppStructOps below which could trigger a server crash
+		// for invalid structs.
+		if (DataStructType->IsChildOf(FLayeredMoveInstancedData::StaticStruct()))
+		{
+			// If the struct type at this index should be different from the one that's already here, we need to change it
+			if (DataStructType.Get() != ExistingDataType && ensure(Ar.IsLoading()))
+			{
+				//@todo DanH: There was a comment here about doing it this way "for now", but it was copy-pasted from FRootMotionSourceGroup::NetSerializeRMSArray. Is it valid?
+				FLayeredMoveInstancedData* NewMove = (FLayeredMoveInstancedData*)FMemory::Malloc(DataStructType->GetCppStructOps()->GetSize());
+				DataStructType->InitializeStruct(NewMove);
+
+				struct FSerializationAllocatedLayeredMoveDataDeleter
+				{
+					FORCEINLINE void operator()(FLayeredMoveInstancedData* MoveData) const
+					{
+						check(MoveData);
+						UScriptStruct* ScriptStruct = MoveData->GetScriptStruct();
+						check(ScriptStruct);
+						ScriptStruct->DestroyStruct(MoveData);
+						FMemory::Free(MoveData);
+					}
+				};
+				InstanceMoveData = TSharedRef<FLayeredMoveInstancedData>(NewMove, FSerializationAllocatedLayeredMoveDataDeleter());
+			}
+
+			InstanceMoveData->NetSerialize(Ar);
+		}
+		else
+		{
+			UE_LOG(LogMover, Error, TEXT("FLayeredMoveInstanceGroup::NetSerialize: ScriptStruct [%s] not derived from FLayeredMoveInstancedData attempted to serialize."),
+				*DataStructType->GetName());
+			Ar.SetError();
+		}
+	}
+	else if (DataStructType.IsError())
+	{
+		UE_LOG(LogMover, Error, TEXT("FLayeredMoveInstanceGroup::NetSerialize: Invalid ScriptStruct serialized."));
+		Ar.SetError();
+	}
+}
+
+const UClass* FLayeredMoveInstance::GetSerializedMoveLogicClass() const
+{
+	return MoveLogicClassType;
+}
+
+bool FLayeredMoveInstance::PopulateMissingActiveMoveLogic(const TArray<TObjectPtr<ULayeredMoveLogic>>& RegisteredMoves)
+{
+	if (!HasLogic())
+	{
+		bool bFoundMissingLogic = false;
+		
+		if (MoveLogicClassType)
+		{
+			for (const TObjectPtr<ULayeredMoveLogic>& RegisteredMoveLogic : RegisteredMoves)
+			{
+				if (RegisteredMoveLogic->GetClass() == MoveLogicClassType)
+				{
+					MoveLogic = RegisteredMoveLogic;
+					
+					bFoundMissingLogic = true;
+					break;
+				}
+			}
+
+			if (!bFoundMissingLogic)
+			{
+				UE_LOG(LogMover, Warning, TEXT("Active Layered Move couldn't find it's serialized MoveLogicClass (%s) among registered MoveLogic"), *MoveLogicClassType->GetName());
+			}
+		}
+		else
+		{
+			UE_LOG(LogMover, Warning, TEXT("Active Layered Move didn't have a valid logic class or class type to search for"));
+		}
+
+		return bFoundMissingLogic;
+	}
+
+	return true;
+}
diff --git a/Engine/Plugins/Experimental/Mover/Source/Mover/Private/LayeredMoveGroup.cpp b/Engine/Plugins/Experimental/Mover/Source/Mover/Private/LayeredMoveGroup.cpp
new file mode 100644
--- /dev/null
+++ b/Engine/Plugins/Experimental/Mover/Source/Mover/Private/LayeredMoveGroup.cpp
@@ -0,0 +1,354 @@
+﻿// Copyright Epic Games, Inc. All Rights Reserved.
+
+#include "LayeredMoveGroup.h"
+#include "LayeredMoveBase.h"
+#include "MoverLog.h"
+#include "MoveLibrary/MovementMixer.h"
+
+FLayeredMoveInstanceGroup::FLayeredMoveInstanceGroup()
+	: ResidualClamping(-1.f)
+	, bApplyResidualVelocity(false)
+	, ResidualVelocity(FVector::Zero())
+{
+}
+
+FLayeredMoveInstanceGroup& FLayeredMoveInstanceGroup::operator=(const FLayeredMoveInstanceGroup& Other)
+{
+	if (this != &Other)
+	{
+		const auto DeepCopyMovesFunc = [](const TArray<TSharedPtr<FLayeredMoveInstance>>& From, TArray<TSharedPtr<FLayeredMoveInstance>>& To)
+		{
+			To.Empty(From.Num());
+			for (const TSharedPtr<FLayeredMoveInstance>& Move : From)
+			{
+				To.Add(Move);
+			}
+		};
+		DeepCopyMovesFunc(Other.ActiveMoves, ActiveMoves);
+		DeepCopyMovesFunc(Other.QueuedMoves, QueuedMoves);
+	}
+
+	return *this;
+}
+
+bool FLayeredMoveInstanceGroup::operator==(const FLayeredMoveInstanceGroup& Other) const
+{
+	// Todo: Deep move-by-move comparison
+	if (ActiveMoves.Num() != Other.ActiveMoves.Num())
+	{
+		return false;
+	}
+	if (QueuedMoves.Num() != Other.QueuedMoves.Num())
+	{
+		return false;
+	}
+	
+	return true;
+}
+
+void FLayeredMoveInstanceGroup::QueueLayeredMove(const TSharedPtr<FLayeredMoveInstance>& Move)
+{
+	if (ensure(Move.IsValid() && Move->HasLogic()))
+	{
+		QueuedMoves.Add(Move);
+		// TODO NS: Add simple string support and re add this log
+		//UE_LOG(LogMover, VeryVerbose, TEXT("LayeredMove queued move (%s)"), *Move->ToSimpleString());
+	}
+}
+bool FLayeredMoveInstanceGroup::HasSameContents(const FLayeredMoveInstanceGroup& Other) const
+{
+	// Only compare the types of moves contained, not the state
+	if (ActiveMoves.Num() != Other.ActiveMoves.Num() ||
+		QueuedMoves.Num() != Other.QueuedMoves.Num())
+	{
+		return false;
+	}
+	for (int32 i = 0; i < ActiveMoves.Num(); ++i)
+	{
+		if (ActiveMoves[i]->GetDataStructType() != Other.ActiveMoves[i]->GetDataStructType())
+		{
+			return false;
+		}
+	}
+	for (int32 i = 0; i < QueuedMoves.Num(); ++i)
+	{
+		if (QueuedMoves[i]->GetDataStructType() != Other.QueuedMoves[i]->GetDataStructType())
+		{
+			return false;
+		}
+	}
+	return true;
+}
+void FLayeredMoveInstanceGroup::ApplyResidualVelocity(FProposedMove& ProposedMove)
+{
+	if (bApplyResidualVelocity)
+	{
+		ProposedMove.LinearVelocity = ResidualVelocity;
+	}
+	if (ResidualClamping >= 0.0f)
+	{
+		ProposedMove.LinearVelocity = ProposedMove.LinearVelocity.GetClampedToMaxSize(ResidualClamping);
+	}
+	ResetResidualVelocity();
+}
+bool FLayeredMoveInstanceGroup::GenerateMixedMove(const FMoverTickStartData& StartState, const FMoverTimeStep& TimeStep, UMovementMixer& MovementMixer, UMoverBlackboard* SimBlackboard, FProposedMove& OutMixedMove)
+{
+	// Tick and accumulate all active moves
+	// Gather all proposed moves and distill this into a cumulative movement report. May include separate additive vs override moves.
+	// TODO: may want to sort by priority or other factors
+	bool bHasLayeredMoveContributions = false;
+	for (const TSharedPtr<FLayeredMoveInstance>& Move : ActiveMoves)
+	{
+		FProposedMove MoveStep;
+		if (Move->GenerateMove(StartState, TimeStep, SimBlackboard, MoveStep))
+		{	
+			bHasLayeredMoveContributions = true;
+			
+			MovementMixer.MixLayeredMove(*Move, MoveStep, OutMixedMove);
+		}
+	}
+	
+	return bHasLayeredMoveContributions;
+}
+
+void FLayeredMoveInstanceGroup::NetSerialize(FArchive& Ar, uint8 MaxNumMovesToSerialize/* = MAX_uint8*/)
+{
+	const auto NetSerializeMovesArrayFunc = [&Ar, this](TArray<TSharedPtr<FLayeredMoveInstance>>& MovesArray, uint8 MaxArraySize)
+	{
+		uint8 NumMovesToSerialize;
+		if (Ar.IsSaving())
+		{
+			UE_CLOG(MovesArray.Num() > MaxArraySize, LogMover, Warning, TEXT("Too many Layered Moves (%d!) to net serialize. Clamping to %d"), MovesArray.Num(), MaxArraySize);
+			NumMovesToSerialize = FMath::Min<int32>(MovesArray.Num(), MaxArraySize);
+		}
+
+		Ar << NumMovesToSerialize;
+
+		if (Ar.IsLoading())
+		{
+			// Note that any instances of FLayeredMoveInstance added this way won't have their constructor called
+			// They are not safe to use until after the NetSerialize() that immediately follows below (so don't add anything in between!) 
+			MovesArray.SetNumZeroed(NumMovesToSerialize);
+			
+			for (int32 MoveIdx = 0; MoveIdx < NumMovesToSerialize && !Ar.IsError(); ++MoveIdx)
+			{
+				MovesArray[MoveIdx] = MakeShared<FLayeredMoveInstance>(MakeShared<FLayeredMoveInstancedData>(), nullptr);
+				FLayeredMoveInstance* Move = MovesArray[MoveIdx].Get();
+				Move->NetSerialize(Ar);
+			}
+		}
+		else
+		{
+			for (int32 MoveIdx = 0; MoveIdx < NumMovesToSerialize && !Ar.IsError(); ++MoveIdx)
+			{
+				FLayeredMoveInstance* Move = MovesArray[MoveIdx].Get();
+				Move->NetSerialize(Ar);
+			}
+		}
+	};
+
+	NetSerializeMovesArrayFunc(ActiveMoves, MaxNumMovesToSerialize);
+	
+	const uint8 MaxNumQueuedMovesToSerialize = FMath::Max(MaxNumMovesToSerialize - ActiveMoves.Num(), 0);
+	NetSerializeMovesArrayFunc(QueuedMoves, MaxNumQueuedMovesToSerialize);
+}
+
+
+void FLayeredMoveInstanceGroup::AddStructReferencedObjects(FReferenceCollector& Collector) const
+{
+	for (const TSharedPtr<FLayeredMoveInstance>& Move : ActiveMoves)
+	{
+		Move->AddReferencedObjects(Collector);
+	}
+	for (const TSharedPtr<FLayeredMoveInstance>& Move : QueuedMoves)
+	{
+		Move->AddReferencedObjects(Collector);
+	}
+}
+
+void FLayeredMoveInstanceGroup::ResetResidualVelocity()
+{
+	bApplyResidualVelocity = false;
+	ResidualVelocity = FVector::ZeroVector;
+	ResidualClamping = -1.f;
+}
+
+void FLayeredMoveInstanceGroup::Reset()
+{
+	ResetResidualVelocity();
+	QueuedMoves.Empty();
+	ActiveMoves.Empty();
+}
+
+void FLayeredMoveInstanceGroup::PopulateMissingActiveMoveLogic(const TArray<TObjectPtr<ULayeredMoveLogic>>& RegisteredMoves)
+{
+	for (TSharedPtr<FLayeredMoveInstance> ActiveMove : ActiveMoves)
+	{
+		if (ActiveMove && !ActiveMove->HasLogic())
+		{
+			bool bPopulatedActiveMoveLogic = ActiveMove->PopulateMissingActiveMoveLogic(RegisteredMoves);
+		}
+	}
+
+	for (TSharedPtr<FLayeredMoveInstance> QueuedMove : QueuedMoves)
+	{
+		if (QueuedMove && !QueuedMove->HasLogic())
+		{
+			bool bPopulatedActiveMoveLogic = QueuedMove->PopulateMissingActiveMoveLogic(RegisteredMoves);
+		}
+	}
+}
+
+FString FLayeredMoveInstanceGroup::ToSimpleString() const
+{
+	return FString::Printf(TEXT("FLayeredMoveInstanceGroup. Active: %i Queued: %i"), ActiveMoves.Num(), QueuedMoves.Num());
+}
+
+const FLayeredMoveInstance* FLayeredMoveInstanceGroup::PrivateFindActiveMove(const TSubclassOf<ULayeredMoveLogic>& MoveLogicClass) const
+{
+	for (const TSharedPtr<FLayeredMoveInstance>& Move : ActiveMoves)
+	{
+		if (Move->GetLogicClass() == MoveLogicClass)
+		{
+			return Move.Get();
+		}
+	}
+
+	return nullptr;
+}
+
+const FLayeredMoveInstance* FLayeredMoveInstanceGroup::PrivateFindActiveMove(const UScriptStruct* MoveDataType) const
+{
+	for (const TSharedPtr<FLayeredMoveInstance>& Move : ActiveMoves)
+	{
+		if (Move->GetDataStructType() == MoveDataType)
+		{
+			return Move.Get();
+		}
+	}
+
+	return nullptr;
+}
+
+const FLayeredMoveInstance* FLayeredMoveInstanceGroup::PrivateFindQueuedMove(const TSubclassOf<ULayeredMoveLogic>& MoveLogicClass) const
+{
+	for (const TSharedPtr<FLayeredMoveInstance>& Move : QueuedMoves)
+	{
+		if (Move->GetLogicClass() == MoveLogicClass)
+		{
+			return Move.Get();
+		}
+	}
+
+	return nullptr;
+}
+
+const FLayeredMoveInstance* FLayeredMoveInstanceGroup::PrivateFindQueuedMove(const UScriptStruct* MoveDataType) const
+{
+	for (const TSharedPtr<FLayeredMoveInstance>& Move : QueuedMoves)
+	{
+		if (Move->GetDataStructType() == MoveDataType)
+		{
+			return Move.Get();
+		}
+	}
+
+	return nullptr;
+}
+
+void FLayeredMoveInstanceGroup::FlushMoveArrays(const FMoverTimeStep& TimeStep, UMoverBlackboard* SimBlackboard)
+{
+	{
+		bool bResidualVelocityOverridden = false;
+		bool bClampVelocityOverridden = false;
+	
+		// Process completion of any moves that are finished
+		ActiveMoves.RemoveAll([&TimeStep, &SimBlackboard, &bResidualVelocityOverridden, &bClampVelocityOverridden, this]
+			(const TSharedPtr<FLayeredMoveInstance>& Move)
+			{
+				if (Move->IsFinished(TimeStep, SimBlackboard))
+				{
+					ProcessFinishedMove(*Move, bResidualVelocityOverridden, bClampVelocityOverridden);
+					Move->EndMove(TimeStep, SimBlackboard);
+					return true;
+				}
+				return false;
+			});	
+	}
+
+	for (TSharedPtr<FLayeredMoveInstance>& QueuedMove : QueuedMoves)
+	{
+		if (QueuedMove.IsValid())
+		{
+			if (QueuedMove->HasLogic())
+			{
+				ActiveMoves.Add(QueuedMove);
+				QueuedMove->StartMove(TimeStep, SimBlackboard);
+			}
+			else
+			{
+				// We should've populated missing logic before this so let's just clear this and warn about it
+				UE_LOG(LogMover, Warning, TEXT("Queued Active Move (%s) logic was not present. Move will not be activated."), *QueuedMove->GetDataStructType()->GetName());
+			}
+		}
+	}
+	
+	QueuedMoves.Empty();
+}
+
+void FLayeredMoveInstanceGroup::ProcessFinishedMove(const FLayeredMoveInstance& Move, bool& bResidualVelocityOverridden, bool& bClampVelocityOverridden)
+{
+	const FLayeredMoveFinishVelocitySettings& FinishVelocitySettings = Move.GetFinishVelocitySettings();
+	const EMoveMixMode MixMode = Move.GetMixMode();
+	
+	if (FinishVelocitySettings.FinishVelocityMode == ELayeredMoveFinishVelocityMode::SetVelocity)
+	{
+		bApplyResidualVelocity = true;
+
+		switch (MixMode)
+		{
+		case EMoveMixMode::AdditiveVelocity:
+			if (!bResidualVelocityOverridden)
+			{
+				ResidualVelocity += FinishVelocitySettings.SetVelocity;
+			}
+			break;
+		case EMoveMixMode::OverrideVelocity:
+		case EMoveMixMode::OverrideAll:
+			{
+				UE_CLOG(bClampVelocityOverridden, LogMover, Log, TEXT("Multiple LayeredMove residual settings have a MixMode that overrides. Only one will take effect."));
+				bResidualVelocityOverridden = true;
+				ResidualVelocity = FinishVelocitySettings.SetVelocity;
+			}
+			break;
+		}
+	}
+	else if (FinishVelocitySettings.FinishVelocityMode == ELayeredMoveFinishVelocityMode::ClampVelocity)
+	{
+		switch (MixMode)
+		{
+		case EMoveMixMode::AdditiveVelocity:
+			if (!bClampVelocityOverridden)
+			{
+				if (ResidualClamping < 0)
+				{
+					ResidualClamping = FinishVelocitySettings.ClampVelocity;
+				}
+				else if (ResidualClamping > FinishVelocitySettings.ClampVelocity)
+				{
+					// No way to really add clamping so we instead apply it if it was smaller
+					ResidualClamping = FinishVelocitySettings.ClampVelocity;
+				}
+			}
+			break;
+		case EMoveMixMode::OverrideVelocity:
+		case EMoveMixMode::OverrideAll:
+			{
+				UE_CLOG(bClampVelocityOverridden, LogMover, Log, TEXT("Multiple LayeredMove residual settings have a MixMode that overrides. Only one will take effect."));
+				bClampVelocityOverridden = true;
+				ResidualClamping = FinishVelocitySettings.ClampVelocity;
+			}
+			break;
+		}
+	}
+}
diff --git a/Engine/Plugins/Experimental/Mover/Source/Mover/Private/MoveLibrary/MovementMixer.cpp b/Engine/Plugins/Experimental/Mover/Source/Mover/Private/MoveLibrary/MovementMixer.cpp
--- a/Engine/Plugins/Experimental/Mover/Source/Mover/Private/MoveLibrary/MovementMixer.cpp
+++ b/Engine/Plugins/Experimental/Mover/Source/Mover/Private/MoveLibrary/MovementMixer.cpp
@@ -2,6 +2,7 @@
 
 #include "MoveLibrary/MovementMixer.h"
 #include "LayeredMove.h"
+#include "LayeredMoveBase.h"
 #include "MoverLog.h"
 
 #include UE_INLINE_GENERATED_CPP_BY_NAME(MovementMixer)
@@ -95,6 +96,89 @@ void UMovementMixer::MixLayeredMove(const FLayeredMoveBase& ActiveMove, const FP
 	}
 }
 
+void UMovementMixer::MixLayeredMove(const FLayeredMoveInstance& ActiveMove, const FProposedMove& MoveStep, FProposedMove& OutCumulativeMove)
+{
+	if (OutCumulativeMove.PreferredMode != MoveStep.PreferredMode && !OutCumulativeMove.PreferredMode.IsNone() && !MoveStep.PreferredMode.IsNone())
+	{
+		UE_LOG(LogMover, Log, TEXT("Multiple LayeredMoves are conflicting with preferred moves. %s will override %s"),
+			*MoveStep.PreferredMode.ToString(), *OutCumulativeMove.PreferredMode.ToString());
+	}
+
+	if (MoveStep.bHasDirIntent && OutCumulativeMove.MixMode != EMoveMixMode::OverrideAll && ActiveMove.GetPriority() >= CurrentHighestPriority)
+	{
+		if (OutCumulativeMove.bHasDirIntent)
+		{
+			UE_LOG(LogMover, Log, TEXT("Multiple LayeredMoves are setting direction intent and the layered move with highest priority will be used."));
+		}
+				
+		OutCumulativeMove.bHasDirIntent = MoveStep.bHasDirIntent;
+		OutCumulativeMove.DirectionIntent = MoveStep.DirectionIntent;
+	}
+
+	if (MoveStep.MixMode == EMoveMixMode::OverrideVelocity)
+	{
+		if (CheckPriority(ActiveMove.GetPriority(), ActiveMove.GetStartingTimeMs(), CurrentHighestPriority, CurrentLayeredMoveStartTime))
+		{
+			if (OutCumulativeMove.MixMode == EMoveMixMode::OverrideVelocity || OutCumulativeMove.MixMode == EMoveMixMode::OverrideAll)
+			{
+				UE_LOG(LogMover, Log, TEXT("Multiple LayeredMoves with Override mix mode are active simultaneously. Layered move with the highest priority will take effect."));
+			}
+
+			if (!MoveStep.PreferredMode.IsNone() && OutCumulativeMove.MixMode != EMoveMixMode::OverrideAll)
+			{
+				OutCumulativeMove.PreferredMode = MoveStep.PreferredMode;
+			}
+				
+			OutCumulativeMove.MixMode = EMoveMixMode::OverrideVelocity;
+			OutCumulativeMove.LinearVelocity  = MoveStep.LinearVelocity;
+			OutCumulativeMove.AngularVelocity = MoveStep.AngularVelocity;
+		}
+	}
+	else if (MoveStep.MixMode == EMoveMixMode::AdditiveVelocity)
+	{
+		if (OutCumulativeMove.MixMode != EMoveMixMode::OverrideVelocity && OutCumulativeMove.MixMode != EMoveMixMode::OverrideAll)
+		{
+			if (!MoveStep.PreferredMode.IsNone())
+			{
+				OutCumulativeMove.PreferredMode = MoveStep.PreferredMode;
+			}
+
+			OutCumulativeMove.LinearVelocity += MoveStep.LinearVelocity;
+			OutCumulativeMove.AngularVelocity += MoveStep.AngularVelocity;
+		}
+	}
+	else if (MoveStep.MixMode == EMoveMixMode::OverrideAll)
+	{
+		if (CheckPriority(ActiveMove.GetPriority(), ActiveMove.GetStartingTimeMs(), CurrentHighestPriority, CurrentLayeredMoveStartTime))
+		{
+			if (OutCumulativeMove.MixMode == EMoveMixMode::OverrideVelocity || OutCumulativeMove.MixMode == EMoveMixMode::OverrideAll)
+			{
+				UE_LOG(LogMover, Log, TEXT("Multiple LayeredMoves with Override mix mode are active simultaneously. Layered move with the highest priority will take effect."));
+			}
+				
+			OutCumulativeMove = MoveStep;
+			OutCumulativeMove.MixMode = EMoveMixMode::OverrideAll;
+		}
+	}
+	else if (MoveStep.MixMode == EMoveMixMode::OverrideAllExceptVerticalVelocity)
+	{
+		if (CheckPriority(ActiveMove.GetPriority(), ActiveMove.GetStartingTimeMs(), CurrentHighestPriority, CurrentLayeredMoveStartTime))
+		{
+			if (OutCumulativeMove.MixMode == EMoveMixMode::OverrideVelocity || OutCumulativeMove.MixMode == EMoveMixMode::OverrideAll || OutCumulativeMove.MixMode == EMoveMixMode::OverrideAllExceptVerticalVelocity)
+			{
+				UE_LOG(LogMover, Log, TEXT("Multiple LayeredMoves with Override mix mode are active simultaneously. Layered move with the highest priority will take effect."));
+			}
+
+			OutCumulativeMove = MoveStep;
+			OutCumulativeMove.MixMode = EMoveMixMode::OverrideAllExceptVerticalVelocity;
+		}
+	}
+	else
+	{
+		ensureMsgf(false, TEXT("Unhandled move mix mode was found."));
+	}
+}
+
 void UMovementMixer::MixProposedMoves(const FProposedMove& MoveToMix, FVector UpDirection, FProposedMove& OutCumulativeMove)
 {
 	if (MoveToMix.bHasDirIntent && OutCumulativeMove.MixMode != EMoveMixMode::OverrideAll)
@@ -155,3 +239,20 @@ bool UMovementMixer::CheckPriority(const FLayeredMoveBase* LayeredMove, uint8& I
 
 	return false;
 }
+
+bool UMovementMixer::CheckPriority(const uint8 LayeredMovePriority, const double LayeredMoveStartTimeMs, uint8& InOutHighestPriority, double& InOutCurrentLayeredMoveStartTimeMs)
+{
+	if (LayeredMovePriority > InOutHighestPriority)
+	{
+		InOutHighestPriority = LayeredMovePriority;
+		InOutCurrentLayeredMoveStartTimeMs = LayeredMoveStartTimeMs;
+		return true;
+	}
+	if (LayeredMovePriority == InOutHighestPriority && LayeredMoveStartTimeMs < InOutCurrentLayeredMoveStartTimeMs)
+	{
+		InOutCurrentLayeredMoveStartTimeMs = LayeredMoveStartTimeMs;
+		return true;
+	}
+
+	return false;
+}
diff --git a/Engine/Plugins/Experimental/Mover/Source/Mover/Private/MovementModeStateMachine.cpp b/Engine/Plugins/Experimental/Mover/Source/Mover/Private/MovementModeStateMachine.cpp
--- a/Engine/Plugins/Experimental/Mover/Source/Mover/Private/MovementModeStateMachine.cpp
+++ b/Engine/Plugins/Experimental/Mover/Source/Mover/Private/MovementModeStateMachine.cpp
@@ -207,6 +207,10 @@ void UMovementModeStateMachine::OnSimulationTick(USceneComponent* UpdatedCompone
 		// Transfer any queued moves into the starting state. They'll be started during the move generation.
 		FlushQueuedMovesToGroup(WorkingSubstepStartData.SyncState.LayeredMoves);
 		OutputState.SyncState.LayeredMoves = WorkingSubstepStartData.SyncState.LayeredMoves;
+		
+		ActivateQueuedMoves(WorkingSubstepStartData.SyncState.LayeredMoveInstances);
+		WorkingSubstepStartData.SyncState.LayeredMoveInstances.PopulateMissingActiveMoveLogic(*MoverComp->GetRegisteredMoves());
+		OutputState.SyncState.LayeredMoveInstances = WorkingSubstepStartData.SyncState.LayeredMoveInstances;
 
 		FlushQueuedModifiersToGroup(WorkingSubstepStartData.SyncState.MovementModifiers);
 		OutputState.SyncState.MovementModifiers = WorkingSubstepStartData.SyncState.MovementModifiers;
@@ -249,6 +253,7 @@ void UMovementModeStateMachine::OnSimulationTick(USceneComponent* UpdatedCompone
 		}
 		
 		FLayeredMoveGroup& CurrentLayeredMoves = OutputState.SyncState.LayeredMoves;
+		FLayeredMoveInstanceGroup& CurrentActiveLayeredMoves = OutputState.SyncState.LayeredMoveInstances;
 
 		// Gather any layered move contributions
 		FProposedMove CombinedLayeredMove;
@@ -257,6 +262,8 @@ void UMovementModeStateMachine::OnSimulationTick(USceneComponent* UpdatedCompone
 		MoverComp->MovementMixer->ResetMixerState();
 		
 		TArray<TSharedPtr<FLayeredMoveBase>> ActiveMoves = CurrentLayeredMoves.GenerateActiveMoves(SubTimeStep, MoverComp, SimBlackboard);
+		CurrentActiveLayeredMoves.FlushMoveArrays(SubTimeStep, SimBlackboard);
+		bHasLayeredMoveContributions = CurrentActiveLayeredMoves.GenerateMixedMove(WorkingSubstepStartData, SubTimeStep, *MoverComp->MovementMixer, SimBlackboard, CombinedLayeredMove);
 
 		// Tick and accumulate all active moves
 		// Gather all proposed moves and distill this into a cumulative movement report. May include separate additive vs override moves.
@@ -478,6 +485,11 @@ void UMovementModeStateMachine::OnSimulationRollback(const FMoverSyncState* Sync
 		QueuedLayeredMoves.Empty();
 	}
 
+	{
+		FRWScopeLock QueueLock(LayeredMoveQueueLock, SLT_Write);
+		QueuedLayeredMoveInstances.Empty();
+	}
+	
 	{
 		FRWScopeLock QueueLock(InstantEffectsQueueLock, SLT_Write);
 		QueuedInstantEffects.Empty();
@@ -522,6 +534,11 @@ void UMovementModeStateMachine::QueueInstantMovementEffect_Internal(const FSched
 	QueueInstantMovementEffect(ScheduledInstantMovementEffect);
 }
 
+void UMovementModeStateMachine::QueueActiveLayeredMove(const TSharedPtr<FLayeredMoveInstance>& LayeredMove)
+{
+	QueuedLayeredMoveInstances.Add(LayeredMove);
+}
+
 void UMovementModeStateMachine::QueueInstantMovementEffect(const FScheduledInstantMovementEffect& ScheduledInstantMovementEffect)
 {
 	FRWScopeLock QueueLock(InstantEffectsQueueLock, SLT_Write);
@@ -673,6 +690,21 @@ void UMovementModeStateMachine::FlushQueuedMovesToGroup(FLayeredMoveGroup& Group
 	}
 }
 
+void UMovementModeStateMachine::ActivateQueuedMoves(FLayeredMoveInstanceGroup& Group)
+{
+	FRWScopeLock QueueLock(LayeredMoveQueueLock, SLT_Write);
+
+	if (!QueuedLayeredMoveInstances.IsEmpty())
+	{
+		for (TSharedPtr<FLayeredMoveInstance>& QueuedMove : QueuedLayeredMoveInstances)
+		{
+			Group.QueueLayeredMove(QueuedMove);
+		}
+		
+		QueuedLayeredMoveInstances.Empty();
+	}
+}
+
 void UMovementModeStateMachine::FlushQueuedModifiersToGroup(FMovementModifierGroup& ModifierGroup)
 {
 	FRWScopeLock QueueLock(ModifiersQueueLock, SLT_Write);
diff --git a/Engine/Plugins/Experimental/Mover/Source/Mover/Private/MoverComponent.cpp b/Engine/Plugins/Experimental/Mover/Source/Mover/Private/MoverComponent.cpp
--- a/Engine/Plugins/Experimental/Mover/Source/Mover/Private/MoverComponent.cpp
+++ b/Engine/Plugins/Experimental/Mover/Source/Mover/Private/MoverComponent.cpp
@@ -520,10 +520,7 @@ void UMoverComponent::SimulationTick(const FMoverTimeStep& InTimeStep, const FMo
 		ProcessFirstSimTickAfterRollback(InTimeStep);
 	}
 
-	if (OnPreSimulationTick.IsBound())
-	{
-		OnPreSimulationTick.Broadcast(MoverTimeStep, SimInput.InputCmd);
-	}
+	PreSimulationTick(MoverTimeStep, SimInput.InputCmd);
 
 	if (!ModeFSM)
 	{
@@ -859,6 +856,37 @@ void UMoverComponent::RemoveGameplayTags(const FGameplayTagContainer& TagsToRemo
 	ExternalGameplayTags.RemoveTags(TagsToRemove);
 }
 
+void UMoverComponent::PreSimulationTick(const FMoverTimeStep& TimeStep, const FMoverInputCmdContext& InputCmd)
+{
+	if (OnPreSimulationTick.IsBound())
+	{
+		OnPreSimulationTick.Broadcast(TimeStep, InputCmd);
+	}
+
+	for (const TSubclassOf<ULayeredMoveLogic>& PendingRegistrantClass : MovesPendingRegistration)
+	{
+		TObjectPtr<ULayeredMoveLogic> RegisteredMove = NewObject<ULayeredMoveLogic>(this, PendingRegistrantClass);
+		RegisteredMoves.Add(RegisteredMove);
+	}
+
+	for (const TSubclassOf<ULayeredMoveLogic>& PendingUnregistrantClass : MovesPendingUnregistration)
+	{
+		RegisteredMoves.RemoveAll([&PendingUnregistrantClass, this]
+			(const TObjectPtr<ULayeredMoveLogic>& MoveLogic)
+			{
+				if (MoveLogic->GetClass() == PendingUnregistrantClass)
+				{
+					return true;
+				}
+				
+				return false;
+			});
+	}
+	
+	MovesPendingRegistration.Empty();
+	MovesPendingUnregistration.Empty();
+}
+
 void UMoverComponent::UpdateCachedFrameState(const FMoverSyncState* SyncState, const FMoverAuxStateContext* AuxState)
 {
 	// TODO integrate dirty tracking
@@ -1465,6 +1493,87 @@ void UMoverComponent::RefreshSharedSettings()
 }
 
 
+const TArray<TObjectPtr<ULayeredMoveLogic>>* UMoverComponent::GetRegisteredMoves() const
+{
+	return &RegisteredMoves;
+}
+
+void UMoverComponent::K2_RegisterMove(TSubclassOf<ULayeredMoveLogic> MoveClass)
+{
+	MovesPendingUnregistration.Remove(MoveClass);
+	if (!MovesPendingRegistration.Contains(MoveClass))
+	{
+		const bool bAlreadyRegistered = RegisteredMoves.ContainsByPredicate([MoveClass](const TObjectPtr<ULayeredMoveLogic>& Move) { return Move->GetClass() == MoveClass; });
+		if (!bAlreadyRegistered)
+		{
+			MovesPendingRegistration.AddUnique(MoveClass);
+		}
+	}
+}
+
+void UMoverComponent::K2_RegisterMoves(TArray<TSubclassOf<ULayeredMoveLogic>> MoveClasses)
+{
+	for (const TSubclassOf<ULayeredMoveLogic>& MoveClass : MoveClasses)
+	{
+		K2_RegisterMove(MoveClass);
+	}
+}
+
+void UMoverComponent::K2_UnregisterMove(TSubclassOf<ULayeredMoveLogic> MoveClass)
+{
+	MovesPendingRegistration.Remove(MoveClass);
+	if (!MovesPendingUnregistration.Contains(MoveClass))
+	{
+		const bool bAlreadyUnregistered = RegisteredMoves.ContainsByPredicate([MoveClass](const TObjectPtr<ULayeredMoveLogic>& Move) { return Move->GetClass() == MoveClass; });
+		if (!bAlreadyUnregistered)
+		{
+			MovesPendingUnregistration.AddUnique(MoveClass);
+		}
+	}
+}
+
+bool UMoverComponent::K2_QueueLayeredMoveActivationWithContext(TSubclassOf<ULayeredMoveLogic> MoveLogicClass, const int32& MoveAsRawData)
+{
+	// This will never be called, the exec version below will be hit instead
+	checkNoEntry();
+	return false;
+}
+
+DEFINE_FUNCTION(UMoverComponent::execK2_QueueLayeredMoveActivationWithContext)
+{
+	P_GET_OBJECT(UClass, MoveLogicClass);
+	
+	Stack.MostRecentPropertyAddress = nullptr;
+	Stack.MostRecentPropertyContainer = nullptr;
+	Stack.StepCompiledIn<FStructProperty>(nullptr);
+
+	const FStructProperty* MoveActivationProperty = CastField<FStructProperty>(Stack.MostRecentProperty);
+	uint8* MoveActivationPtr = Stack.MostRecentPropertyAddress;
+	
+	P_FINISH
+	
+	P_NATIVE_BEGIN;
+
+	// TODO NS: throw some helpful warnings of what wasn't valid
+	const bool bHasValidActivationStructProp = MoveActivationProperty && MoveActivationProperty->Struct && MoveActivationProperty->Struct->IsChildOf(FLayeredMoveActivationParams::StaticStruct());
+
+	bool bHasValidMoveData = MoveLogicClass && bHasValidActivationStructProp; 
+	if (bHasValidMoveData)
+	{
+		const FLayeredMoveActivationParams* MoveActivationContext = reinterpret_cast<FLayeredMoveActivationParams*>(MoveActivationPtr);
+		bHasValidMoveData = P_THIS->MakeAndQueueLayeredMove(MoveLogicClass, MoveActivationContext);
+	}
+	
+	*(bool*)RESULT_PARAM = bHasValidMoveData;
+	
+	P_NATIVE_END;
+}
+
+bool UMoverComponent::QueueLayeredMoveActivation(TSubclassOf<ULayeredMoveLogic> MoveLogicClass)
+{
+	return MakeAndQueueLayeredMove(MoveLogicClass, nullptr);
+}
+
 void UMoverComponent::K2_QueueLayeredMove(const int32& MoveAsRawData)
 {
 	// This will never be called, the exec version below will be hit instead
@@ -1610,7 +1719,6 @@ DEFINE_FUNCTION(UMoverComponent::execK2_ScheduleInstantMovementEffect)
 void UMoverComponent::ScheduleInstantMovementEffect(TSharedPtr<FInstantMovementEffect> InstantMovementEffect)
 {
 	ensureMsgf(IsInGameThread(), TEXT("UMoverComponent::ScheduleInstantMovementEffect should only be called from the game thread. Inspect code for incorrect calls."));
-	ensure(IsInGameThread());
 	FMoverTimeStep TimeStep;
 	if (ensureMsgf(BackendLiaisonComp, TEXT("UMoverComponent::ScheduleInstantMovementEffect was unexpectedly called with a null backend liaison component. The instant movement effect will be ignored.")))
 	{
@@ -2390,6 +2498,72 @@ UBaseMovementMode* UMoverComponent::GetActiveModeInternal(TSubclassOf<UBaseMovem
 	return nullptr;
 }
 
+bool UMoverComponent::MakeAndQueueLayeredMove(const TSubclassOf<ULayeredMoveLogic>& MoveLogicClass, const FLayeredMoveActivationParams* ActivationParams)
+{
+	// Find registered type for class passed in
+	TObjectPtr<ULayeredMoveLogic> FoundRegisteredMoveLogic = nullptr;
+	for (TObjectPtr<ULayeredMoveLogic> RegisteredMoveLogic : RegisteredMoves)
+	{
+		if (RegisteredMoveLogic.GetClass()->IsChildOf(MoveLogicClass))
+		{
+			FoundRegisteredMoveLogic = RegisteredMoveLogic;
+			break;
+		}
+	}
+			
+	ULayeredMoveLogic* ActiveMoveLogic;
+	TSharedPtr<FLayeredMoveInstancedData> QueuedInstancedData;
+
+	if (FoundRegisteredMoveLogic)
+	{
+		ActiveMoveLogic = FoundRegisteredMoveLogic;
+
+		const UScriptStruct* InstancedDataType = FoundRegisteredMoveLogic->GetInstancedDataType();
+		if (InstancedDataType && InstancedDataType->IsChildOf(FLayeredMoveInstancedData::StaticStruct()))
+		{
+			TCheckedObjPtr<UScriptStruct> DataStructType = FoundRegisteredMoveLogic->GetInstancedDataType();
+			FLayeredMoveInstancedData* NewMove = (FLayeredMoveInstancedData*)FMemory::Malloc(DataStructType->GetCppStructOps()->GetSize());
+			DataStructType->InitializeStruct(NewMove);
+
+			struct FAllocatedLayeredMoveDataDeleter
+			{
+				FORCEINLINE void operator()(FLayeredMoveInstancedData* MoveData) const
+				{
+					check(MoveData);
+					UScriptStruct* ScriptStruct = MoveData->GetScriptStruct();
+					check(ScriptStruct);
+					ScriptStruct->DestroyStruct(MoveData);
+					FMemory::Free(MoveData);
+				}
+			};
+				
+			QueuedInstancedData = TSharedRef<FLayeredMoveInstancedData>(NewMove, FAllocatedLayeredMoveDataDeleter());
+			QueuedInstancedData->ActivateFromContext(ActivationParams);
+		}
+		else
+		{
+			UE_LOG(LogMover, Warning, TEXT("%s activation was queued on %s but the move was NOT queued since it did not have valid data. InstancedDataStructType on Move Logic needs to be a FLayeredMoveInstancedData or child struct of."),
+				*MoveLogicClass->GetName(),
+				*GetOwner()->GetName());
+			
+			return false;
+		}
+	}
+	else
+	{
+		UE_LOG(LogMover, Warning, TEXT("%s activation was queued on %s and the move was not registered. Any move activated on a MoverComponent Needs to be Registered with the MoverCompoent. The layered move will not be queued for activation."),
+			*MoveLogicClass->GetName(),
+			*GetOwner()->GetName());
+			
+		return false;
+	}
+	
+	const TSharedPtr<FLayeredMoveInstance> ActiveMoveToQueue = MakeShared<FLayeredMoveInstance>(QueuedInstancedData.ToSharedRef(), ActiveMoveLogic);
+	ModeFSM->QueueActiveLayeredMove(ActiveMoveToQueue);
+	
+	return true;
+}
+
 void UMoverComponent::SetSimulationOutput(const FMoverTimeStep& TimeStep, const UE::Mover::FSimulationOutputData& OutputData)
 {
 	CachedLastSimTickTimeStep = TimeStep;
diff --git a/Engine/Plugins/Experimental/Mover/Source/Mover/Public/DefaultMovementSet/LayeredMoves/LaunchMove.h b/Engine/Plugins/Experimental/Mover/Source/Mover/Public/DefaultMovementSet/LayeredMoves/LaunchMove.h
--- a/Engine/Plugins/Experimental/Mover/Source/Mover/Public/DefaultMovementSet/LayeredMoves/LaunchMove.h
+++ b/Engine/Plugins/Experimental/Mover/Source/Mover/Public/DefaultMovementSet/LayeredMoves/LaunchMove.h
@@ -3,10 +3,69 @@
 #pragma once
 
 #include "LayeredMove.h"
+#include "LayeredMoveBase.h"
 #include "LaunchMove.generated.h"
 
 #define UE_API MOVER_API
 
+USTRUCT(Blueprintable)
+struct FLaunchMoveActivationParams : public FLayeredMoveActivationParams
+{
+	GENERATED_BODY()
+
+	/** Velocity to apply to the updated component. Could be additive or overriding depending on MixMode setting. */
+	UPROPERTY(EditAnywhere, BlueprintReadWrite, Category = Mover, meta=(ForceUnits="cm/s"))
+	FVector LaunchVelocity = FVector::ZeroVector;
+
+	// Optional movement mode name to force the actor into before applying the impulse velocity.
+	UPROPERTY(EditAnywhere, BlueprintReadWrite, Category = Mover)
+	FName ForceMovementMode = NAME_None;
+	
+};
+
+USTRUCT(Blueprintable)
+struct FLaunchMoveData : public FLayeredMoveInstancedData
+{
+	GENERATED_BODY()
+
+	//@todo DanH: This is boilerplate begging for a macro
+	virtual FLayeredMoveInstancedData* Clone() const override { return new FLaunchMoveData(*this); }
+	virtual UScriptStruct* GetScriptStruct() const override { return StaticStruct(); }
+	
+	/** Velocity to apply to the updated component. Could be additive or overriding depending on MixMode setting. */
+	UPROPERTY(EditAnywhere, BlueprintReadWrite, Category = Mover, meta=(ForceUnits="cm/s"))
+	FVector LaunchVelocity = FVector::ZeroVector;
+
+	/** Optional movement mode name to force the actor into before applying the impulse velocity. */
+	UPROPERTY(EditAnywhere, BlueprintReadWrite, Category = Mover)
+	FName ForceMovementMode;
+	
+	virtual void ActivateFromContext(const FLayeredMoveActivationParams* ActivationParams) override;
+	
+	virtual void NetSerialize(FArchive& Ar) override;
+};
+
+// TODO: Create data for this? Is it not needed?!
+UCLASS()
+class ULaunchMoveLogic : public ULayeredMoveLogic
+{
+	GENERATED_BODY()
+
+public:
+	UE_API ULaunchMoveLogic();
+
+	/** Velocity to apply to the updated component. Could be additive or overriding depending on MixMode setting. */
+	UPROPERTY(EditAnywhere, BlueprintReadWrite, Category = Mover, meta=(ForceUnits="cm/s"))
+	FVector LaunchVelocity = FVector::ZeroVector;
+
+	//@todo DanH: Should forcing a mode be an option at the root ULayeredMoveLogic?
+	/** Optional movement mode name to force the actor into before applying the impulse velocity. */
+	UPROPERTY(EditAnywhere, BlueprintReadWrite, Category = Mover)
+	FName ForceMovementMode;
+	
+protected:
+	UE_API virtual bool GenerateMove_Implementation(const FMoverTimeStep& TimeStep, UMoverBlackboard* SimBlackboard, const FMoverTickStartData& StartState, FProposedMove& OutProposedMove) override;
+};
 
 /** Launch Move: provides an impulse velocity to the actor after (optionally) forcing them into a particular movement mode */
 USTRUCT(BlueprintType)
diff --git a/Engine/Plugins/Experimental/Mover/Source/Mover/Public/LayeredMoveBase.h b/Engine/Plugins/Experimental/Mover/Source/Mover/Public/LayeredMoveBase.h
new file mode 100644
--- /dev/null
+++ b/Engine/Plugins/Experimental/Mover/Source/Mover/Public/LayeredMoveBase.h
@@ -0,0 +1,301 @@
+﻿// Copyright Epic Games, Inc. All Rights Reserved.
+
+#pragma once
+
+#include "GameplayTagContainer.h"
+#include "MoveLibrary/MovementUtilsTypes.h"
+#include "Templates/SubclassOf.h"
+#include "LayeredMove.h"
+#include "MoverLog.h"
+
+#include "LayeredMoveBase.generated.h"
+
+//////////////////////////////////////////////////////////////////////////
+// TODO: Consider relocating this class and functionality, a lot of what's here belongs in the original LayeredMove.h/.cpp as replacements for FLayeredMoveBase and such
+//		 but we wanted to explore this implementation of Blueprintable layered moves without disrupting the old system so it lives here for now.
+//////////////////////////////////////////////////////////////////////////
+
+#define UE_API MOVER_API
+#define USING_PARAMS(ActivationParamsClass) using ActivationParamsType = ActivationParamsClass;
+#define USING_MOVE_DATA(MoveDataClass) using MoveDataType = MoveDataClass;
+
+class UMovementMixer;
+struct FMoverTickStartData;
+struct FMoverTimeStep;
+class UMoverBlackboard;
+class UMoverComponent;
+class ULayeredMoveLogic;
+
+/**
+ * Packaged params struct for initializing a corresponding FLayeredMoveInstancedData
+ * Allows BP to do "templated" move data creation. Optional in C++, where params can be forwarded to the FLayeredMoveInstancedData ctor directly
+ * The base class can also be used on any activation to just use default values
+ */
+USTRUCT(Blueprintable)
+struct FLayeredMoveActivationParams
+{
+	GENERATED_BODY()
+	
+	/**
+ 	 * This move will expire after a set amount of time if > 0. If 0, it will be ticked only once, regardless of time step. It will need to be manually ended if < 0.
+ 	 * Note: If changed after starting to a value beneath the current lifetime of the move, it will immediately finish (so if your move finishes early, setting this to 0 is equivalent to returning true from IsFinished())
+ 	 */
+	UPROPERTY(EditAnywhere, BlueprintReadWrite, Category = Mover)
+	double DurationMs = 0.f;
+};
+
+/** Instanced data created and replicated for each activation of a layered move */
+USTRUCT(Blueprintable)
+struct FLayeredMoveInstancedData
+{
+	GENERATED_BODY()
+
+	//@todo DanH: If we're already using a macro, might as well put the boilerplate virtuals in here
+	USING_PARAMS(FLayeredMoveActivationParams);
+	
+	FLayeredMoveInstancedData() {}
+	virtual ~FLayeredMoveInstancedData() {}
+	// Checks if 
+	bool operator==(const FLayeredMoveInstancedData& Other) const;
+	bool operator!=(const FLayeredMoveInstancedData& Other) const { return !operator==(Other); }
+
+	/** @return A newly allocated copy of this FLayeredMoveInstancedData. Must be overridden by child classes */
+	virtual FLayeredMoveInstancedData* Clone() const { return new FLayeredMoveInstancedData(*this); }
+
+	/** @return The UScriptStruct describing this struct. Must be overridden by child classes */
+	virtual UScriptStruct* GetScriptStruct() const { return StaticStruct(); }
+	
+	/** @return True if this move data is identical to OtherData. OtherData is guaranteed to be safe to cast to the implementing type. */
+	virtual bool Equals(const FLayeredMoveInstancedData& OtherData) const;
+	
+	/** Called when a queued layered move is activated. Provides an opportunity to initialize Layered MoveData */
+	virtual void ActivateFromContext(const FLayeredMoveActivationParams* ActivationParams);
+	
+    virtual void NetSerialize(FArchive& Ar);
+
+    virtual void AddReferencedObjects(FReferenceCollector& Collector) {}
+
+	/** Is this move considered to "have" a given gameplay tag? */
+	virtual bool HasGameplayTag(FGameplayTag TagToFind, bool bExactMatch) const { return false; }
+	
+	double GetStartSimTimeMs() const { return StartSimTimeMs; }
+
+	/**
+	 * This move will expire after a set amount of time if > 0. If 0, it will be ticked only once, regardless of time step. It will need to be manually ended if < 0.
+	 * Note: If changed after starting to a value beneath the current lifetime of the move, it will immediately finish (so if your move finishes early, setting this to 0 is equivalent to returning true from IsFinished())
+	 */
+	UPROPERTY(BlueprintReadWrite, Category = Mover)
+	double DurationMs = -1.f;
+	
+	/** The simulation time this move first ticked (< 0 means it hasn't started yet) */
+    UPROPERTY(BlueprintReadOnly, Category = Mover)
+    double StartSimTimeMs = -UE_BIG_NUMBER;
+
+	friend class ULayeredMoveLogic;
+};
+
+template<>
+struct TStructOpsTypeTraits<FLayeredMoveInstancedData> : public TStructOpsTypeTraitsBase2<FLayeredMoveInstancedData>
+{
+	enum
+	{
+		//WithNetSerializer = true,
+		WithCopy = true
+	};
+};
+
+/**
+ * Base class for all layered move logic that operates in tandem with instanced FLayeredMoveInstancedData.
+ * The logic object itself is not meant to ever be replicated, and a maximum of one instance of each logic class
+ * need ever exist on a given MoverComponent. Repeated and/or simultaneous activations of the same move on a component
+ * are represented, tracked, and replicated through instances of the FLayeredMoveInstancedData struct type
+ * that the logic class indicates in the ActiveMoveDataStructType property.
+ *
+ * The virtual methods on this class are invoked in a special and strict pattern that guarantees AccessExecutionMoveData
+ * will return the valid data instance relevant to that function execution.
+ *
+ * Refer to [Examples when they exist] for implementation examples
+ */
+UCLASS(Abstract, BlueprintType, Blueprintable)
+class UE_API ULayeredMoveLogic : public UObject
+{
+	GENERATED_BODY()
+
+public:
+	USING_MOVE_DATA(FLayeredMoveInstancedData)
+	
+	ULayeredMoveLogic();
+
+	UScriptStruct* GetInstancedDataType() const { return InstancedDataStructType; }
+	
+	const FLayeredMoveFinishVelocitySettings& GetFinishVelocitySettings() const { return FinishVelocitySettings; }
+	EMoveMixMode GetMixMode() const { return MixMode; }
+	uint8 GetPriority() const { return Priority; }
+
+	// Helper function for validating move data when passing data and logic to/from BP
+	static bool ValidateMoveDataGetSet(const UObject* ObjectValidatingData, const ULayeredMoveLogic* MoveLogic, const FStructProperty* MoveDataProperty, const uint8* MoveDataPtr, FFrame& StackFrame);
+	
+private:
+	// These functions all assume validity of CurrentActiveMoveData to perform logic from the perspective of that chunk of properties
+	// To that end, the only external entity capable of invoking them is the FScopedMoveLogicExecContext, which guarantees the active move data has been set
+	friend class FScopedMoveLogicExecContext;
+	
+protected:
+	/** Called when this move is initially activated  */
+	UFUNCTION(BlueprintNativeEvent)
+	void OnStart(const FMoverTimeStep& TimeStep, UMoverBlackboard* SimBlackboard);
+
+	/** Called when this move has ended  */
+	UFUNCTION(BlueprintNativeEvent)
+	void OnEnd(const FMoverTimeStep& TimeStep, UMoverBlackboard* SimBlackboard);
+
+	/** Generate a movement that will be combined with other sources */
+	UFUNCTION(BlueprintNativeEvent)
+	bool GenerateMove(const FMoverTimeStep& TimeStep, UMoverBlackboard* SimBlackboard, UPARAM(ref) const FMoverTickStartData& StartState, UPARAM(ref) FProposedMove& OutProposedMove);
+	
+	//@todo: Will need to cache whether to treat the instance data as const (i.e. whether to disregard & complain if someone tries to set the data in BP during that window) 
+	UFUNCTION(BlueprintNativeEvent)
+	bool IsFinished(const FMoverTimeStep& TimeStep, const UMoverBlackboard* SimBlackboard);
+	
+	virtual void OnStart_Implementation(const FMoverTimeStep& TimeStep, UMoverBlackboard* SimBlackboard) {}
+	virtual void OnEnd_Implementation(const FMoverTimeStep& TimeStep, UMoverBlackboard* SimBlackboard) {}
+	virtual bool GenerateMove_Implementation(const FMoverTimeStep& TimeStep, UMoverBlackboard* SimBlackboard, const FMoverTickStartData& StartState, FProposedMove& OutProposedMove) { return false; }
+	virtual bool IsFinished_Implementation(const FMoverTimeStep& TimeStep, const UMoverBlackboard* SimBlackboard);
+	
+	/** Accessor to the FLayeredMoveInstancedData established for the execution of a virtual move logic function */
+	template <typename MoveDataT = FLayeredMoveInstancedData UE_REQUIRES(std::is_base_of_v<FLayeredMoveInstancedData, MoveDataT>)>
+	MoveDataT& AccessExecutionMoveData() const
+	{
+#if UE_BUILD_SHIPPING || UE_BUILD_TEST
+		if (!ensure(InstancedDataStructType))
+        {
+			UE_LOG(LogMover, Error, TEXT("Move Logic needs an active data struct type. If no data is needed consider using the default move data type."))
+        	static MoveDataT GarbageData;
+        	return GarbageData;
+        }
+#endif
+		
+		check(InstancedDataStructType);
+		return static_cast<MoveDataT&>(*CurrentInstancedData);
+	}
+	
+	/**
+	 * Gets active move data that is tied to this active move. Use SetActiveMoveData to write active move data.
+	 * Note: Make sure break node is the correct data type specified by the layered move logic.
+	 */
+	UFUNCTION(BlueprintCallable, CustomThunk, Category = "Mover | Layered Move", DisplayName = "Get Active Move Data", meta = (BlueprintProtected, BlueprintInternalUseOnly, CustomStructureParam = "OutMoveData", DefaultToSelf = "MoveLogic"))
+	static bool K2_GetActiveMoveData(ULayeredMoveLogic* MoveLogic, FLayeredMoveInstancedData& OutMoveData);
+	DECLARE_FUNCTION(execK2_GetActiveMoveData);
+	
+	/**
+	 * Sets active move data that is tied to this active move. Use GetActiveMoveData to read current active move data.
+	 * Note: Make sure make node is the correct data type specified by the layered move logic.
+	 */
+	UFUNCTION(BlueprintCallable, CustomThunk, Category = "Mover | Layered Move", DisplayName = "Set Active Move Data", meta = (BlueprintProtected, BlueprintInternalUseOnly, CustomStructureParam = "OutMoveData", DefaultToSelf = "MoveLogic"))
+	static void K2_SetActiveMoveData(ULayeredMoveLogic* MoveLogic, const FLayeredMoveInstancedData& OutMoveData);
+	DECLARE_FUNCTION(execK2_SetActiveMoveData);
+	
+	/**
+	 * This move will expire after a set amount of time if > 0. If 0, it will be ticked only once, regardless of time step. It will need to be manually ended if < 0.
+	 * Note: If changed after starting to a value beneath the current lifetime of the move, it will immediately finish (so if your move finishes early, setting this to 0 is equivalent to returning true from IsFinished())
+	 */
+	UPROPERTY(BlueprintReadWrite, Category = Mover)
+	double DefaultDurationMs = -1.f;
+	
+	/** Determines how this object's movement contribution should be mixed with others */
+	UPROPERTY(EditAnywhere, BlueprintReadOnly, Category = Mover)
+	EMoveMixMode MixMode = EMoveMixMode::AdditiveVelocity;
+
+	/** Determines if this layered move should take priority over other layered moves when different moves have conflicting overrides - higher numbers taking precedent. */
+	UPROPERTY(EditAnywhere, BlueprintReadOnly, Category = Mover)
+	uint8 Priority = 0;
+
+	//@todo: Does this need to be a separate struct?
+	/** Settings related to velocity applied to the actor after the move has finished */
+	UPROPERTY(EditAnywhere, BlueprintReadOnly, Category = Mover)
+	FLayeredMoveFinishVelocitySettings FinishVelocitySettings;
+
+	//@todo DanH: Fail validation if this isn't set or isn't an FLayeredMoveInstancedData
+	UPROPERTY(EditAnywhere, NoClear, BlueprintReadOnly, Category = Mover, meta=(MetaStruct="/Script/Mover.LayeredMoveInstancedData"))
+	TObjectPtr<UScriptStruct> InstancedDataStructType;
+	
+private:
+	/**
+	 * The FLayeredMoveInstancedData provided to each of the base virtual move functions, valid only for the duration a single virtual function execution.
+	 * Direct access is only for internal plumbing - use AccessExecutionMoveData() in virtuals to obtain a typed reference to this.
+	 */
+	TSharedPtr<FLayeredMoveInstancedData> CurrentInstancedData;
+};
+
+/**
+ * Wrapper to encapsulate the split implementation of a move between a stateless ULayeredMoveLogic* object and an instance of FLayeredMoveInstancedData
+ * Those two pieces, in tandem, represent a "whole" functional Layered Move.
+ */
+USTRUCT()
+struct FLayeredMoveInstance
+{
+	GENERATED_BODY()
+	
+	FLayeredMoveInstance();
+	FLayeredMoveInstance(const FLayeredMoveInstance& OtherLayeredMoveInstance);
+	FLayeredMoveInstance(const TSharedRef<FLayeredMoveInstancedData>& InMoveData, ULayeredMoveLogic* InMoveLogic = nullptr);
+	
+	/** TArray::FindByKey() enablers */
+	bool operator==(const TSubclassOf<ULayeredMoveLogic>& LogicClass) const { return MoveLogic && MoveLogic->IsA(LogicClass); }
+	bool operator==(const UScriptStruct* MoveDataType) const { return InstanceMoveData->GetScriptStruct()->IsChildOf(MoveDataType); }
+
+	bool HasLogic() const { return MoveLogic != nullptr; }
+	const UClass* GetLogicClass() const { return MoveLogic->GetClass(); }
+	
+	void StartMove(const FMoverTimeStep& TimeStep, UMoverBlackboard* SimBlackboard) const;
+	bool GenerateMove(const FMoverTickStartData& StartState, const FMoverTimeStep& TimeStep, UMoverBlackboard* SimBlackboard, FProposedMove& OutProposedMove) const;
+	void EndMove(const FMoverTimeStep& TimeStep, UMoverBlackboard* SimBlackboard) const;
+	bool IsFinished(const FMoverTimeStep& TimeStep, const UMoverBlackboard* SimBlackboard) const;
+
+	const FLayeredMoveFinishVelocitySettings& GetFinishVelocitySettings() const;
+	EMoveMixMode GetMixMode() const;
+	uint8 GetPriority() const { return MoveLogic->GetPriority(); }
+	double GetStartingTimeMs() const { return InstanceMoveData->GetStartSimTimeMs(); }
+	
+	FLayeredMoveInstance Clone() const;
+	UScriptStruct* GetDataStructType() const { return InstanceMoveData->GetScriptStruct(); }
+	void AddReferencedObjects(FReferenceCollector& Collector) const { InstanceMoveData->AddReferencedObjects(Collector); } 
+
+	void NetSerialize(FArchive& Ar);
+
+	const UClass* GetSerializedMoveLogicClass() const;
+
+	/**
+	 * This function populates the MoveLogic reference of ActiveMoves that don't have any logic classes. This is necessary as
+	 * active move data received from NetSerialize doesn't necessarily have the Logic class it was activated with.
+	 * Note: Currently we are NetSerializing the MoveLogic class so that we can then populate Logic from using the registered move array
+	 * on the MoverComponent but eventually we should use a ID mapping or something better to avoid serializing the MoveLogic class.
+	 */
+	bool PopulateMissingActiveMoveLogic(const TArray<TObjectPtr<ULayeredMoveLogic>>& RegisteredMoves);
+	
+private:
+	TSharedPtr<FLayeredMoveInstancedData> InstanceMoveData;
+	
+	/**
+	 * Is used in PopulateMissingActiveMoveLogic to help populate logic classes in active moves that were NetSerialized since MoveLogic isn't NetSerialized
+	 * TODO: Eventually replace this with some ID/Mapping system so we can NetSerialize that instead of a full UClass
+	 */
+	UPROPERTY(Transient)
+	TSubclassOf<ULayeredMoveLogic> MoveLogicClassType;
+	
+	//@todo DanH: This exists on the PT for the chaos mover - can/should this still be a UProperty TObjectPtr?
+	UPROPERTY(Transient)
+	TObjectPtr<ULayeredMoveLogic> MoveLogic;
+};
+
+template<>
+struct TStructOpsTypeTraits<FLayeredMoveInstance> : public TStructOpsTypeTraitsBase2<FLayeredMoveInstance>
+{
+	enum
+	{
+		//WithNetSerializer = true,
+		WithCopy = true
+	};
+};
+
+#undef UE_API
diff --git a/Engine/Plugins/Experimental/Mover/Source/Mover/Public/LayeredMoveGroup.h b/Engine/Plugins/Experimental/Mover/Source/Mover/Public/LayeredMoveGroup.h
new file mode 100644
--- /dev/null
+++ b/Engine/Plugins/Experimental/Mover/Source/Mover/Public/LayeredMoveGroup.h
@@ -0,0 +1,148 @@
+﻿// Copyright Epic Games, Inc. All Rights Reserved.
+
+#pragma once
+
+#include "MoveLibrary/MovementUtilsTypes.h"
+#include "LayeredMoveGroup.generated.h"
+
+#define UE_API MOVER_API
+
+struct FLayeredMoveInstance;
+struct FLayeredMoveInstancedData;
+struct FProposedMove;
+class UMovementMixer;
+struct FMoverTickStartData;
+struct FMoverTimeStep;
+class UMoverBlackboard;
+class ULayeredMoveLogic;
+
+/**
+ * The group of information about currently active and queued moves.
+ * This replicates info for FLayeredMoveInstancedData only - it is expected that the corresponding ULayeredMoveLogic is 
+ * already registered with the mover component.
+ */
+USTRUCT(BlueprintType)
+struct UE_API FLayeredMoveInstanceGroup
+{
+	GENERATED_BODY()
+
+	FLayeredMoveInstanceGroup();
+	FLayeredMoveInstanceGroup& operator=(const FLayeredMoveInstanceGroup& Other);
+	bool operator==(const FLayeredMoveInstanceGroup& Other) const;
+	bool operator!=(const FLayeredMoveInstanceGroup& Other) const { return !operator==(Other); }
+
+	/** Checks only whether there are matching LayeredMoves, but NOT necessarily identical states of each move */
+	bool HasSameContents(const FLayeredMoveInstanceGroup& Other) const;
+	
+	bool GenerateMixedMove(const FMoverTickStartData& StartState, const FMoverTimeStep& TimeStep, UMovementMixer& MovementMixer, UMoverBlackboard* SimBlackboard, FProposedMove& OutMixedMove);
+	void ApplyResidualVelocity(FProposedMove& InOutProposedMove);
+	void NetSerialize(FArchive& Ar, uint8 MaxNumMovesToSerialize = MAX_uint8);
+	void AddStructReferencedObjects(FReferenceCollector& Collector) const;
+	void ResetResidualVelocity();
+	void Reset();
+
+	/**
+	 * Loops through all Queued and Active moves and populates any missing MoveLogic using FLayeredMoveInstance::PopulateMissingActiveMoveLogic.
+	 * See FLayeredMoveInstance::PopulateMissingActiveMoveLogic function for more details.
+	 */
+	void PopulateMissingActiveMoveLogic(const TArray<TObjectPtr<ULayeredMoveLogic>>& RegisteredMoves);
+	
+	/** Adds the active move to the queued array of the move group */
+	void QueueLayeredMove(const TSharedPtr<FLayeredMoveInstance>& Move);
+	
+	/** @return True if there are any active or queued moves in this group */
+	bool HasAnyMoves() const { return (!ActiveMoves.IsEmpty() || !QueuedMoves.IsEmpty()); }
+	
+	/** @return True if there is at least one layered move that's either active or queued and is associated with the provided logic or data type */
+	template <typename MoveElementT UE_REQUIRES(std::is_base_of_v<FLayeredMoveInstancedData, MoveElementT> || std::is_base_of_v<ULayeredMoveLogic, MoveElementT>)>
+	bool HasMove() const
+	{
+		return FindActiveMove<MoveElementT>() || FindQueuedMove<MoveElementT>();
+	}
+	
+	/** Get a simplified string representation of this group. Typically for debugging. */
+	FString ToSimpleString() const;
+	
+	/** Returns the first active layered move associated with logic of the specified type, if one exists */
+	template <typename MoveLogicT = ULayeredMoveLogic UE_REQUIRES(std::is_base_of_v<ULayeredMoveLogic, MoveLogicT>)>
+	const FLayeredMoveInstance* FindActiveMove(TSubclassOf<ULayeredMoveLogic> MoveLogicClass = MoveLogicT::StaticClass()) const
+	{
+		return PrivateFindActiveMove(MoveLogicClass);
+	}
+	
+	/** Returns the first active layered move using data of the specified type, if one exists */
+	template <typename MoveDataT = FLayeredMoveInstancedData UE_REQUIRES(std::is_base_of_v<FLayeredMoveInstancedData, MoveDataT>)>
+	const FLayeredMoveInstance* FindActiveMove(const UScriptStruct* MoveDataType = MoveDataT::StaticStruct()) const
+	{
+		return PrivateFindActiveMove(MoveDataType);
+	}
+	
+	/** Returns the first queued layered move associated with logic of the specified type, if one exists */
+	template <typename MoveLogicT = ULayeredMoveLogic UE_REQUIRES(std::is_base_of_v<ULayeredMoveLogic, MoveLogicT>)>
+	const FLayeredMoveInstance* FindQueuedMove(TSubclassOf<ULayeredMoveLogic> MoveLogicClass = MoveLogicT::StaticClass()) const
+	{
+		return PrivateFindQueuedMove(MoveLogicClass);
+	}
+
+	/** Returns the first queued layered move using data of the specified type, if one exists */
+	template <typename MoveDataT = FLayeredMoveInstancedData UE_REQUIRES(std::is_base_of_v<FLayeredMoveInstancedData, MoveDataT>)>
+	const FLayeredMoveInstance* FindQueuedMove(const UScriptStruct* MoveDataType = MoveDataT::StaticStruct()) const
+	{
+		return PrivateFindQueuedMove(MoveDataType);
+	}
+
+	// Clears out any finished or invalid active moves and adds any queued moves to the active moves
+	void FlushMoveArrays(const FMoverTimeStep& TimeStep, UMoverBlackboard* SimBlackboard);
+	
+protected:
+	// Helper function for gathering any residual velocity settings from layered moves that just ended
+	void ProcessFinishedMove(const FLayeredMoveInstance& Move, bool& bResidualVelocityOverriden, bool& bClampVelocityOverriden);
+	
+	/** Moves that are currently active in this group */
+	TArray<TSharedPtr<FLayeredMoveInstance>> ActiveMoves;
+
+	/** Moves that are queued to become active next sim frame */
+	TArray<TSharedPtr<FLayeredMoveInstance>> QueuedMoves;
+
+private:
+	const FLayeredMoveInstance* PrivateFindActiveMove(const TSubclassOf<ULayeredMoveLogic>& MoveLogicClass) const;
+	const FLayeredMoveInstance* PrivateFindActiveMove(const UScriptStruct* MoveDataType) const;
+	const FLayeredMoveInstance* PrivateFindQueuedMove(const TSubclassOf<ULayeredMoveLogic>& MoveLogicClass) const;
+	const FLayeredMoveInstance* PrivateFindQueuedMove(const UScriptStruct* MoveDataType) const;
+	
+	//@todo DanH: Maybe these should be grouped in a struct?
+	/**
+	 * Clamps an actors velocity to this value when a layered move ends. This expects Value >= 0.
+	 * Note: This is set automatically when a layered move has ended based off of end velocity settings - it is not meant to be set by the user see @FLayeredMoveFinishVelocitySettings
+	 */
+	UPROPERTY(NotReplicated, VisibleAnywhere, BlueprintReadOnly, Category = Mover, meta = (AllowPrivateAccess))
+	float ResidualClamping;
+
+	/**
+	 * If true ResidualVelocity will be the next velocity used for this actor
+	 * Note: This is set automatically when a layered move has ended based off of end velocity settings - it is not meant to be set by the user see @FLayeredMoveFinishVelocitySettings
+	 */
+	UPROPERTY(NotReplicated, VisibleAnywhere, BlueprintReadOnly, Category = Mover, meta = (AllowPrivateAccess))
+	bool bApplyResidualVelocity;
+
+	/**
+	 * If bApplyResidualVelocity is true this actors velocity will be set to this.
+	 * Note: This is set automatically when a layered move has ended based off of end velocity settings - it is not meant to be set by the user see @FLayeredMoveFinishVelocitySettings
+	 */
+	UPROPERTY(NotReplicated, VisibleAnywhere, BlueprintReadOnly, Category = Mover, meta = (AllowPrivateAccess))
+	FVector ResidualVelocity;
+};
+
+template<>
+struct TStructOpsTypeTraits<FLayeredMoveInstanceGroup> : public TStructOpsTypeTraitsBase2<FLayeredMoveInstanceGroup>
+{
+	enum
+	{
+		WithCopy = true,
+		//WithNetSerializer = true,
+		WithIdenticalViaEquality = true,
+		WithAddStructReferencedObjects = true,
+	};
+};
+
+#undef UE_API
diff --git a/Engine/Plugins/Experimental/Mover/Source/Mover/Public/MoveLibrary/MovementMixer.h b/Engine/Plugins/Experimental/Mover/Source/Mover/Public/MoveLibrary/MovementMixer.h
--- a/Engine/Plugins/Experimental/Mover/Source/Mover/Public/MoveLibrary/MovementMixer.h
+++ b/Engine/Plugins/Experimental/Mover/Source/Mover/Public/MoveLibrary/MovementMixer.h
@@ -8,6 +8,7 @@
 #define UE_API MOVER_API
 
 struct FLayeredMoveBase;
+struct FLayeredMoveInstance;
 struct FProposedMove;
 
 /**
@@ -24,6 +25,9 @@ public:
 	/** In charge of mixing Layered Move proposed moves into a cumulative proposed move based on mix mode and priority.*/
 	UE_API virtual void MixLayeredMove(const FLayeredMoveBase& ActiveMove, const FProposedMove& MoveStep, FProposedMove& OutCumulativeMove);
 
+	/** In charge of mixing Layered Move proposed moves into a cumulative proposed move based on mix mode and priority.*/
+	UE_API virtual void MixLayeredMove(const FLayeredMoveInstance& ActiveMove, const FProposedMove& MoveStep, FProposedMove& OutCumulativeMove);
+
 	/** In charge of mixing proposed moves together. Is similar to MixLayeredMove but is only responsible for mixing proposed moves instead of layered moves. */
 	UE_API virtual void MixProposedMoves(const FProposedMove& MoveToMix, FVector UpDirection, FProposedMove& OutCumulativeMove);
 
@@ -42,6 +46,7 @@ protected:
 	 * Returns true if this layered move should take priority given current HighestPriority and CurrentLayeredMoveStartTimeMs
 	 */
 	static UE_API bool CheckPriority(const FLayeredMoveBase* LayeredMove, uint8& InOutHighestPriority, double& InOutCurrentLayeredMoveStartTimeMs);
+	static UE_API bool CheckPriority(uint8 LayeredMovePriority, double LayeredMoveStartTimeMs, uint8& InOutHighestPriority, double& InOutCurrentLayeredMoveStartTimeMs);
 };
 
 #undef UE_API
diff --git a/Engine/Plugins/Experimental/Mover/Source/Mover/Public/MovementModeStateMachine.h b/Engine/Plugins/Experimental/Mover/Source/Mover/Public/MovementModeStateMachine.h
--- a/Engine/Plugins/Experimental/Mover/Source/Mover/Public/MovementModeStateMachine.h
+++ b/Engine/Plugins/Experimental/Mover/Source/Mover/Public/MovementModeStateMachine.h
@@ -54,6 +54,8 @@ public:
 
 	UE_API void QueueLayeredMove(TSharedPtr<FLayeredMoveBase> Move);
 	
+	UE_API void QueueActiveLayeredMove(const TSharedPtr<FLayeredMoveInstance>& LayeredMove);
+
  	UE_API FMovementModifierHandle QueueMovementModifier(TSharedPtr<FMovementModifierBase> Modifier);
 
  	UE_API void CancelModifierFromHandle(FMovementModifierHandle ModifierHandle);
@@ -88,6 +90,9 @@ protected:
 	/** Moves that are queued to be added to the simulation at the start of the next sim subtick. Access covered by lock. */
 	TArray<TSharedPtr<FLayeredMoveBase>> QueuedLayeredMoves;
 
+	/** Moves that are queued to be added to the simulation at the start of the next sim subtick. Access covered by lock. */
+ 	TArray<TSharedPtr<FLayeredMoveInstance>> QueuedLayeredMoveInstances;
+ 	
  	/** Effects that are queued to be applied to the simulation at the start of the next sim subtick or at the end of this tick.  Access covered by lock. */
  	TArray<FScheduledInstantMovementEffect> QueuedInstantEffects;
 
@@ -111,6 +116,8 @@ private:
 	UE_API void ConstructDefaultModes();
 	UE_API void AdvanceToNextMode();
 	UE_API void FlushQueuedMovesToGroup(FLayeredMoveGroup& Group);
+	// Flushes queued ActiveLayeredMoves to FLayeredMoveInstanceGroup for this frame
+ 	UE_API void ActivateQueuedMoves(FLayeredMoveInstanceGroup& Group);
  	UE_API void FlushQueuedModifiersToGroup(FMovementModifierGroup& ModifierGroup);
  	UE_API void FlushModifierCancellationsToGroup(FMovementModifierGroup& ActiveModifierGroup);
  	UE_API void RollbackModifiers(const FMoverSyncState* InvalidSyncState, const FMoverSyncState* SyncState, const FMoverAuxStateContext* InvalidAuxState, const FMoverAuxStateContext* AuxState);
diff --git a/Engine/Plugins/Experimental/Mover/Source/Mover/Public/MoverComponent.h b/Engine/Plugins/Experimental/Mover/Source/Mover/Public/MoverComponent.h
--- a/Engine/Plugins/Experimental/Mover/Source/Mover/Public/MoverComponent.h
+++ b/Engine/Plugins/Experimental/Mover/Source/Mover/Public/MoverComponent.h
@@ -10,6 +10,7 @@
 #include "MovementMode.h"
 #include "MoverTypes.h"
 #include "LayeredMove.h"
+#include "LayeredMoveBase.h"
 #include "MoveLibrary/BasedMovementUtils.h"
 #include "MoveLibrary/ConstrainedMoveUtils.h"
 #if UE_ENABLE_INCLUDE_ORDER_DEPRECATED_IN_5_4
@@ -223,6 +224,58 @@ public:
 	UPROPERTY(BlueprintReadWrite, EditAnywhere, Category = Mover)
 	TObjectPtr<UMovementMixer> MovementMixer;
 
+	const TArray<TObjectPtr<ULayeredMoveLogic>>* GetRegisteredMoves() const;
+	
+	/** Registers layered move logic */
+	template <typename MoveT UE_REQUIRES(std::is_base_of_v<ULayeredMoveLogic, MoveT>)>
+	void RegisterMove(TSubclassOf<MoveT> MoveClass = MoveT::StaticClass())
+	{
+		K2_RegisterMove(MoveClass);
+	}
+
+	/** Registers layered move logic */
+	UFUNCTION(BlueprintCallable, Category = Mover, DisplayName = "Register Move")
+	void K2_RegisterMove(TSubclassOf<ULayeredMoveLogic> MoveClass);
+
+	/** Registers an array of layered move logic classes */
+	UFUNCTION(BlueprintCallable, Category = Mover, DisplayName = "Register Moves")
+	void K2_RegisterMoves(TArray<TSubclassOf<ULayeredMoveLogic>> MoveClasses);
+	
+	/** Unregisters layered move logic */
+	template <typename MoveT UE_REQUIRES(std::is_base_of_v<ULayeredMoveLogic, MoveT>)>
+	void UnregisterMove(TSubclassOf<MoveT> MoveClass = MoveT::StaticClass())
+	{
+		K2_UnregisterMove(MoveClass);
+	}
+
+	/** Unregisters layered move logic */
+	UFUNCTION(BlueprintCallable, Category = Mover, DisplayName = "Unregister Move")
+	void K2_UnregisterMove(TSubclassOf<ULayeredMoveLogic> MoveClass);
+	
+	template <typename MoveT, typename ActivationParamsT UE_REQUIRES(std::is_base_of_v<ULayeredMoveLogic, MoveT> && std::is_base_of_v<typename MoveT::MoveDataType::ActivationParamsType, ActivationParamsT>)>
+	bool QueueLayeredMoveActivationWithContext(const ActivationParamsT& ActivationParams, TSubclassOf<MoveT> MoveClass = MoveT::StaticClass())
+	{
+		return MakeAndQueueLayeredMove(MoveClass, &ActivationParams);
+	}
+	
+	/**
+	 * Queues a layered move for activation.
+	 * Takes a Activation Context which provides context to set Layered Move Data.
+	 * Make sure Activation Context type matches layered Move Data
+	 */
+	UFUNCTION(BlueprintCallable, CustomThunk, Category = Mover, meta = (CustomStructureParam = "MoveAsRawData", AllowAbstract = "false"), DisplayName = "Queue Layered Move Activation With Context")
+	bool K2_QueueLayeredMoveActivationWithContext(TSubclassOf<ULayeredMoveLogic> MoveLogicClass, UPARAM(DisplayName="Layered Move Activation Context") const int32& MoveAsRawData);
+	DECLARE_FUNCTION(execK2_QueueLayeredMoveActivationWithContext);
+
+	/**
+ 	 * Queues a layered move for activation.
+ 	 * Takes NO Activation Context meaning the layered move will be activated using default Move Data.
+ 	 * Note: Changing Move Data is still possible in the layered move logic itself
+ 	 * See QueueLayeredMoveActivationWithContext for activating a layered move with context
+ 	 */
+	UFUNCTION(BlueprintCallable, Category = Mover, meta = (AllowAbstract = "false"))
+	bool QueueLayeredMoveActivation(TSubclassOf<ULayeredMoveLogic> MoveLogicClass);
+	
 	/**
 	 * Queue a layered move to start during the next simulation frame. This will clone whatever move you pass in, so you'll need to fully set it up before queuing.
 	 * @param LayeredMove			The move to queue, which must be a LayeredMoveBase sub-type. 
@@ -288,6 +341,7 @@ public:
 	MOVER_API void ClearQueuedInstantMovementEffects();
 	// Queue an instant movement effect in async mode. Do not use on the game thread.
 	void QueueInstantMovementEffect_Internal(const FMoverTimeStep& TimeStep, TSharedPtr<FInstantMovementEffect> InstantMovementEffect);
+
 protected:
 	// Queue a scheduled instant movement effect. Thread safe, can be used outside the game thread.
 	void QueueInstantMovementEffect(const FScheduledInstantMovementEffect& ScheduledInstantMovementEffect);
@@ -614,6 +668,9 @@ public:	// Queries
 	
 protected:
 
+	// Called before each simulation tick. Broadcasts OnPreSimulationTick delegate.
+	void PreSimulationTick(const FMoverTimeStep& TimeStep, const FMoverInputCmdContext& InputCmd);
+	
 	/** Makes this component and owner actor reflect the state of a particular frame snapshot. This occurs after simulation ticking, as well as during a rollback before we resimulate forward.
 	  @param bRebaseBasedState	If true and the state was using based movement, it will use the current game world base pos/rot instead of the captured one. This is necessary during rollbacks.
 	*/
@@ -726,7 +783,21 @@ protected:
 	UPROPERTY(Transient)
 	TObjectPtr<URollbackBlackboard_InternalWrapper> RollbackBlackboard_InternalWrapper = nullptr;	// This is a thin layer for use only by in-simulation users
 
+	/**
+	 * Layered moves registered on this component that can be activated regardless of the current mode
+	 * Changes to this array or its contents occur ONLY during PreSimulationTick to ensure threadsafe access during async simulations.
+	 */
+	UPROPERTY(Transient)
+	TArray<TObjectPtr<ULayeredMoveLogic>> RegisteredMoves;
+
+	UPROPERTY(Transient)
+	TArray<TSubclassOf<ULayeredMoveLogic>> MovesPendingRegistration;
+
+	UPROPERTY(Transient)
+	TArray<TSubclassOf<ULayeredMoveLogic>> MovesPendingUnregistration;
 	
+	/** Helper function for making Queueing and making Layered Moves */
+	bool MakeAndQueueLayeredMove(const TSubclassOf<ULayeredMoveLogic>& MoveLogicClass, const FLayeredMoveActivationParams* ActivationParams);
 	
 private:
 	/** Collection of settings objects that are shared between movement modes. This list is automatically managed based on the @MovementModes contents. */
diff --git a/Engine/Plugins/Experimental/Mover/Source/Mover/Public/MoverSimulationTypes.h b/Engine/Plugins/Experimental/Mover/Source/Mover/Public/MoverSimulationTypes.h
--- a/Engine/Plugins/Experimental/Mover/Source/Mover/Public/MoverSimulationTypes.h
+++ b/Engine/Plugins/Experimental/Mover/Source/Mover/Public/MoverSimulationTypes.h
@@ -8,6 +8,7 @@
 #include "MoverTypes.h"
 #include "MoveLibrary/MovementRecord.h"
 #include "LayeredMove.h"
+#include "LayeredMoveGroup.h"
 #include "MovementModifier.h"
 #include "MoverDataModelTypes.h"
 #include "InstantMovementEffect.h"
@@ -201,6 +202,10 @@ public:
 	UPROPERTY(BlueprintReadOnly, VisibleAnywhere, Category = Mover)
 	FLayeredMoveGroup LayeredMoves;
 
+	// Additional moves influencing our proposed motion
+	UPROPERTY(BlueprintReadOnly, VisibleAnywhere, Category = Mover)
+	FLayeredMoveInstanceGroup LayeredMoveInstances;
+
 	// Additional modifiers influencing our simulation
 	UPROPERTY(BlueprintReadOnly, VisibleAnywhere, Category = Mover)
 	FMovementModifierGroup MovementModifiers;
@@ -217,6 +222,7 @@ public:
 	{
 		return MovementMode == Other.MovementMode &&
 			LayeredMoves.HasSameContents(Other.LayeredMoves) &&
+			LayeredMoveInstances.HasSameContents(Other.LayeredMoveInstances) &&
 			MovementModifiers.HasSameContents(Other.MovementModifiers) &&
 			SyncStateCollection.HasSameContents(Other.SyncStateCollection);
 	}
@@ -228,6 +234,7 @@ public:
 	{
 		P.Ar << MovementMode;
 		LayeredMoves.NetSerialize(P.Ar);
+		LayeredMoveInstances.NetSerialize(P.Ar);
 		MovementModifiers.NetSerialize(P.Ar);
 
 		bool bIgnoredResult(false);
@@ -238,6 +245,7 @@ public:
 	{
 		Out.Appendf("MovementMode: %s\n", TCHAR_TO_ANSI(*MovementMode.ToString()));
 		Out.Appendf("Layered Moves: %s\n", TCHAR_TO_ANSI(*LayeredMoves.ToSimpleString()));
+		Out.Appendf("Layered Moves: %s\n", TCHAR_TO_ANSI(*LayeredMoveInstances.ToSimpleString()));
 		Out.Appendf("Movement Modifiers: %s\n", TCHAR_TO_ANSI(*MovementModifiers.ToSimpleString()));
 		SyncStateCollection.ToString(Out);
 	}
@@ -253,6 +261,7 @@ public:
 	{
 		MovementMode = To->MovementMode;
 		LayeredMoves = To->LayeredMoves;
+		LayeredMoveInstances = To->LayeredMoveInstances;
 		MovementModifiers = To->MovementModifiers;
 
 		SyncStateCollection.Interpolate(From->SyncStateCollection, To->SyncStateCollection, Pct);
@@ -265,6 +274,7 @@ public:
 		MovementMode = NAME_None;
 		SyncStateCollection.Empty();
 		LayeredMoves.Reset();
+		LayeredMoveInstances.Reset();
 		MovementModifiers.Reset();
 	}
 };
