[Backout] - CL45147400
#rnx
Original CL Desc
-----------------------------------------------------------------
Mover: Adding support for blueprintable layered moves through a new system that splits up Layered Move Logic and Data. This allows Layered Move Logic to be made natively and in Blueprints in an effort to expand blueprint support in Mover. These new layered moves exist alongside the old ones at the moment with the intention of these new moves eventually replacing the old ones. Currently, only Launch Layered move has been converted to the new system.

#virtualized

diff --git a/Engine/Plugins/Experimental/Mover/Source/Mover/Private/DefaultMovementSet/LayeredMoves/LaunchMove.cpp b/Engine/Plugins/Experimental/Mover/Source/Mover/Private/DefaultMovementSet/LayeredMoves/LaunchMove.cpp
--- a/Engine/Plugins/Experimental/Mover/Source/Mover/Private/DefaultMovementSet/LayeredMoves/LaunchMove.cpp
+++ b/Engine/Plugins/Experimental/Mover/Source/Mover/Private/DefaultMovementSet/LayeredMoves/LaunchMove.cpp
@@ -2,59 +2,11 @@
 
 #include "DefaultMovementSet/LayeredMoves/LaunchMove.h"
 #include "MoverComponent.h"
+#include "MoverTypes.h"
+#include "MoverLog.h"
 
 #include UE_INLINE_GENERATED_CPP_BY_NAME(LaunchMove)
 
-void FLaunchMoveData::ActivateFromContext(const FLayeredMoveActivationParams* ActivationParams)
-{
-	if (ActivationParams)
-	{
-		if (const FLaunchMoveActivationParams* LaunchMoveActivationParams = static_cast<const FLaunchMoveActivationParams*>(ActivationParams))
-		{
-			LaunchVelocity = LaunchMoveActivationParams->LaunchVelocity;
-			DurationMs = LaunchMoveActivationParams->DurationMs;
-			ForceMovementMode = LaunchMoveActivationParams->ForceMovementMode;
-		}
-	}
-}
-
-void FLaunchMoveData::NetSerialize(FArchive& Ar)
-{
-	Super::NetSerialize(Ar);
-
-	SerializePackedVector<10, 16>(LaunchVelocity, Ar);
-
-	bool bUsingForcedMovementMode = !ForceMovementMode.IsNone();
-
-	Ar.SerializeBits(&bUsingForcedMovementMode, 1);
-
-	if (bUsingForcedMovementMode)
-	{
-		Ar << ForceMovementMode;
-	}
-	else
-	{
-		ForceMovementMode = NAME_None;
-	}
-}
-
-ULaunchMoveLogic::ULaunchMoveLogic()
-{
-	DefaultDurationMs = 0.f;
-	MixMode = EMoveMixMode::OverrideVelocity;
-	InstancedDataStructType = FLaunchMoveData::StaticStruct();
-}
-
-bool ULaunchMoveLogic::GenerateMove_Implementation(const FMoverTimeStep& TimeStep, UMoverBlackboard* SimBlackboard, const FMoverTickStartData& StartState, FProposedMove& OutProposedMove)
-{
-	const FLaunchMoveData& LaunchMoveData = AccessExecutionMoveData<FLaunchMoveData>();
-	
-	OutProposedMove.MixMode = MixMode;
-	OutProposedMove.LinearVelocity = LaunchMoveData.LaunchVelocity;
-	OutProposedMove.PreferredMode = LaunchMoveData.ForceMovementMode;
-
-	return true;
-}
 
 FLayeredMove_Launch::FLayeredMove_Launch()
 {
diff --git a/Engine/Plugins/Experimental/Mover/Source/Mover/Private/LayeredMoveBase.cpp b/Engine/Plugins/Experimental/Mover/Source/Mover/Private/LayeredMoveBase.cpp
deleted file mode 100644
diff --git a/Engine/Plugins/Experimental/Mover/Source/Mover/Private/LayeredMoveGroup.cpp b/Engine/Plugins/Experimental/Mover/Source/Mover/Private/LayeredMoveGroup.cpp
deleted file mode 100644
diff --git a/Engine/Plugins/Experimental/Mover/Source/Mover/Private/MoveLibrary/MovementMixer.cpp b/Engine/Plugins/Experimental/Mover/Source/Mover/Private/MoveLibrary/MovementMixer.cpp
--- a/Engine/Plugins/Experimental/Mover/Source/Mover/Private/MoveLibrary/MovementMixer.cpp
+++ b/Engine/Plugins/Experimental/Mover/Source/Mover/Private/MoveLibrary/MovementMixer.cpp
@@ -2,7 +2,6 @@
 
 #include "MoveLibrary/MovementMixer.h"
 #include "LayeredMove.h"
-#include "LayeredMoveBase.h"
 #include "MoverLog.h"
 
 #include UE_INLINE_GENERATED_CPP_BY_NAME(MovementMixer)
@@ -96,89 +95,6 @@ void UMovementMixer::MixLayeredMove(const FLayeredMoveBase& ActiveMove, const FP
 	}
 }
 
-void UMovementMixer::MixLayeredMove(const FLayeredMoveInstance& ActiveMove, const FProposedMove& MoveStep, FProposedMove& OutCumulativeMove)
-{
-	if (OutCumulativeMove.PreferredMode != MoveStep.PreferredMode && !OutCumulativeMove.PreferredMode.IsNone() && !MoveStep.PreferredMode.IsNone())
-	{
-		UE_LOG(LogMover, Log, TEXT("Multiple LayeredMoves are conflicting with preferred moves. %s will override %s"),
-			*MoveStep.PreferredMode.ToString(), *OutCumulativeMove.PreferredMode.ToString());
-	}
-
-	if (MoveStep.bHasDirIntent && OutCumulativeMove.MixMode != EMoveMixMode::OverrideAll && ActiveMove.GetPriority() >= CurrentHighestPriority)
-	{
-		if (OutCumulativeMove.bHasDirIntent)
-		{
-			UE_LOG(LogMover, Log, TEXT("Multiple LayeredMoves are setting direction intent and the layered move with highest priority will be used."));
-		}
-				
-		OutCumulativeMove.bHasDirIntent = MoveStep.bHasDirIntent;
-		OutCumulativeMove.DirectionIntent = MoveStep.DirectionIntent;
-	}
-
-	if (MoveStep.MixMode == EMoveMixMode::OverrideVelocity)
-	{
-		if (CheckPriority(ActiveMove.GetPriority(), ActiveMove.GetStartingTimeMs(), CurrentHighestPriority, CurrentLayeredMoveStartTime))
-		{
-			if (OutCumulativeMove.MixMode == EMoveMixMode::OverrideVelocity || OutCumulativeMove.MixMode == EMoveMixMode::OverrideAll)
-			{
-				UE_LOG(LogMover, Log, TEXT("Multiple LayeredMoves with Override mix mode are active simultaneously. Layered move with the highest priority will take effect."));
-			}
-
-			if (!MoveStep.PreferredMode.IsNone() && OutCumulativeMove.MixMode != EMoveMixMode::OverrideAll)
-			{
-				OutCumulativeMove.PreferredMode = MoveStep.PreferredMode;
-			}
-				
-			OutCumulativeMove.MixMode = EMoveMixMode::OverrideVelocity;
-			OutCumulativeMove.LinearVelocity  = MoveStep.LinearVelocity;
-			OutCumulativeMove.AngularVelocity = MoveStep.AngularVelocity;
-		}
-	}
-	else if (MoveStep.MixMode == EMoveMixMode::AdditiveVelocity)
-	{
-		if (OutCumulativeMove.MixMode != EMoveMixMode::OverrideVelocity && OutCumulativeMove.MixMode != EMoveMixMode::OverrideAll)
-		{
-			if (!MoveStep.PreferredMode.IsNone())
-			{
-				OutCumulativeMove.PreferredMode = MoveStep.PreferredMode;
-			}
-
-			OutCumulativeMove.LinearVelocity += MoveStep.LinearVelocity;
-			OutCumulativeMove.AngularVelocity += MoveStep.AngularVelocity;
-		}
-	}
-	else if (MoveStep.MixMode == EMoveMixMode::OverrideAll)
-	{
-		if (CheckPriority(ActiveMove.GetPriority(), ActiveMove.GetStartingTimeMs(), CurrentHighestPriority, CurrentLayeredMoveStartTime))
-		{
-			if (OutCumulativeMove.MixMode == EMoveMixMode::OverrideVelocity || OutCumulativeMove.MixMode == EMoveMixMode::OverrideAll)
-			{
-				UE_LOG(LogMover, Log, TEXT("Multiple LayeredMoves with Override mix mode are active simultaneously. Layered move with the highest priority will take effect."));
-			}
-				
-			OutCumulativeMove = MoveStep;
-			OutCumulativeMove.MixMode = EMoveMixMode::OverrideAll;
-		}
-	}
-	else if (MoveStep.MixMode == EMoveMixMode::OverrideAllExceptVerticalVelocity)
-	{
-		if (CheckPriority(ActiveMove.GetPriority(), ActiveMove.GetStartingTimeMs(), CurrentHighestPriority, CurrentLayeredMoveStartTime))
-		{
-			if (OutCumulativeMove.MixMode == EMoveMixMode::OverrideVelocity || OutCumulativeMove.MixMode == EMoveMixMode::OverrideAll || OutCumulativeMove.MixMode == EMoveMixMode::OverrideAllExceptVerticalVelocity)
-			{
-				UE_LOG(LogMover, Log, TEXT("Multiple LayeredMoves with Override mix mode are active simultaneously. Layered move with the highest priority will take effect."));
-			}
-
-			OutCumulativeMove = MoveStep;
-			OutCumulativeMove.MixMode = EMoveMixMode::OverrideAllExceptVerticalVelocity;
-		}
-	}
-	else
-	{
-		ensureMsgf(false, TEXT("Unhandled move mix mode was found."));
-	}
-}
-
 void UMovementMixer::MixProposedMoves(const FProposedMove& MoveToMix, FVector UpDirection, FProposedMove& OutCumulativeMove)
 {
 	if (MoveToMix.bHasDirIntent && OutCumulativeMove.MixMode != EMoveMixMode::OverrideAll)
@@ -239,20 +155,3 @@ bool UMovementMixer::CheckPriority(const FLayeredMoveBase* LayeredMove, uint8& I
 
 	return false;
 }
-
-bool UMovementMixer::CheckPriority(const uint8 LayeredMovePriority, const double LayeredMoveStartTimeMs, uint8& InOutHighestPriority, double& InOutCurrentLayeredMoveStartTimeMs)
-{
-	if (LayeredMovePriority > InOutHighestPriority)
-	{
-		InOutHighestPriority = LayeredMovePriority;
-		InOutCurrentLayeredMoveStartTimeMs = LayeredMoveStartTimeMs;
-		return true;
-	}
-	if (LayeredMovePriority == InOutHighestPriority && LayeredMoveStartTimeMs < InOutCurrentLayeredMoveStartTimeMs)
-	{
-		InOutCurrentLayeredMoveStartTimeMs = LayeredMoveStartTimeMs;
-		return true;
-	}
-
-	return false;
-}
diff --git a/Engine/Plugins/Experimental/Mover/Source/Mover/Private/MovementModeStateMachine.cpp b/Engine/Plugins/Experimental/Mover/Source/Mover/Private/MovementModeStateMachine.cpp
--- a/Engine/Plugins/Experimental/Mover/Source/Mover/Private/MovementModeStateMachine.cpp
+++ b/Engine/Plugins/Experimental/Mover/Source/Mover/Private/MovementModeStateMachine.cpp
@@ -207,10 +207,6 @@ void UMovementModeStateMachine::OnSimulationTick(USceneComponent* UpdatedCompone
 		// Transfer any queued moves into the starting state. They'll be started during the move generation.
 		FlushQueuedMovesToGroup(WorkingSubstepStartData.SyncState.LayeredMoves);
 		OutputState.SyncState.LayeredMoves = WorkingSubstepStartData.SyncState.LayeredMoves;
-		
-		ActivateQueuedMoves(WorkingSubstepStartData.SyncState.LayeredMoveInstances);
-		WorkingSubstepStartData.SyncState.LayeredMoveInstances.PopulateMissingActiveMoveLogic(*MoverComp->GetRegisteredMoves());
-		OutputState.SyncState.LayeredMoveInstances = WorkingSubstepStartData.SyncState.LayeredMoveInstances;
 
 		FlushQueuedModifiersToGroup(WorkingSubstepStartData.SyncState.MovementModifiers);
 		OutputState.SyncState.MovementModifiers = WorkingSubstepStartData.SyncState.MovementModifiers;
@@ -253,7 +249,6 @@ void UMovementModeStateMachine::OnSimulationTick(USceneComponent* UpdatedCompone
 		}
 		
 		FLayeredMoveGroup& CurrentLayeredMoves = OutputState.SyncState.LayeredMoves;
-		FLayeredMoveInstanceGroup& CurrentActiveLayeredMoves = OutputState.SyncState.LayeredMoveInstances;
 
 		// Gather any layered move contributions
 		FProposedMove CombinedLayeredMove;
@@ -262,8 +257,6 @@ void UMovementModeStateMachine::OnSimulationTick(USceneComponent* UpdatedCompone
 		MoverComp->MovementMixer->ResetMixerState();
 		
 		TArray<TSharedPtr<FLayeredMoveBase>> ActiveMoves = CurrentLayeredMoves.GenerateActiveMoves(SubTimeStep, MoverComp, SimBlackboard);
-		CurrentActiveLayeredMoves.FlushMoveArrays(SubTimeStep, SimBlackboard);
-		bHasLayeredMoveContributions = CurrentActiveLayeredMoves.GenerateMixedMove(WorkingSubstepStartData, SubTimeStep, *MoverComp->MovementMixer, SimBlackboard, CombinedLayeredMove);
 
 		// Tick and accumulate all active moves
 		// Gather all proposed moves and distill this into a cumulative movement report. May include separate additive vs override moves.
@@ -485,11 +478,6 @@ void UMovementModeStateMachine::OnSimulationRollback(const FMoverSyncState* Sync
 		QueuedLayeredMoves.Empty();
 	}
 
-	{
-		FRWScopeLock QueueLock(LayeredMoveQueueLock, SLT_Write);
-		QueuedLayeredMoveInstances.Empty();
-	}
-	
 	{
 		FRWScopeLock QueueLock(InstantEffectsQueueLock, SLT_Write);
 		QueuedInstantEffects.Empty();
@@ -534,11 +522,6 @@ void UMovementModeStateMachine::QueueInstantMovementEffect_Internal(const FSched
 	QueueInstantMovementEffect(ScheduledInstantMovementEffect);
 }
 
-void UMovementModeStateMachine::QueueActiveLayeredMove(const TSharedPtr<FLayeredMoveInstance>& LayeredMove)
-{
-	QueuedLayeredMoveInstances.Add(LayeredMove);
-}
-
 void UMovementModeStateMachine::QueueInstantMovementEffect(const FScheduledInstantMovementEffect& ScheduledInstantMovementEffect)
 {
 	FRWScopeLock QueueLock(InstantEffectsQueueLock, SLT_Write);
@@ -690,21 +673,6 @@ void UMovementModeStateMachine::FlushQueuedMovesToGroup(FLayeredMoveGroup& Group
 	}
 }
 
-void UMovementModeStateMachine::ActivateQueuedMoves(FLayeredMoveInstanceGroup& Group)
-{
-	FRWScopeLock QueueLock(LayeredMoveQueueLock, SLT_Write);
-
-	if (!QueuedLayeredMoveInstances.IsEmpty())
-	{
-		for (TSharedPtr<FLayeredMoveInstance>& QueuedMove : QueuedLayeredMoveInstances)
-		{
-			Group.QueueLayeredMove(QueuedMove);
-		}
-		
-		QueuedLayeredMoveInstances.Empty();
-	}
-}
-
 void UMovementModeStateMachine::FlushQueuedModifiersToGroup(FMovementModifierGroup& ModifierGroup)
 {
 	FRWScopeLock QueueLock(ModifiersQueueLock, SLT_Write);
diff --git a/Engine/Plugins/Experimental/Mover/Source/Mover/Private/MoverComponent.cpp b/Engine/Plugins/Experimental/Mover/Source/Mover/Private/MoverComponent.cpp
--- a/Engine/Plugins/Experimental/Mover/Source/Mover/Private/MoverComponent.cpp
+++ b/Engine/Plugins/Experimental/Mover/Source/Mover/Private/MoverComponent.cpp
@@ -520,7 +520,10 @@ void UMoverComponent::SimulationTick(const FMoverTimeStep& InTimeStep, const FMo
 		ProcessFirstSimTickAfterRollback(InTimeStep);
 	}
 
-	PreSimulationTick(MoverTimeStep, SimInput.InputCmd);
+	if (OnPreSimulationTick.IsBound())
+	{
+		OnPreSimulationTick.Broadcast(MoverTimeStep, SimInput.InputCmd);
+	}
 
 	if (!ModeFSM)
 	{
@@ -856,37 +859,6 @@ void UMoverComponent::RemoveGameplayTags(const FGameplayTagContainer& TagsToRemo
 	ExternalGameplayTags.RemoveTags(TagsToRemove);
 }
 
-void UMoverComponent::PreSimulationTick(const FMoverTimeStep& TimeStep, const FMoverInputCmdContext& InputCmd)
-{
-	if (OnPreSimulationTick.IsBound())
-	{
-		OnPreSimulationTick.Broadcast(TimeStep, InputCmd);
-	}
-
-	for (const TSubclassOf<ULayeredMoveLogic>& PendingRegistrantClass : MovesPendingRegistration)
-	{
-		TObjectPtr<ULayeredMoveLogic> RegisteredMove = NewObject<ULayeredMoveLogic>(this, PendingRegistrantClass);
-		RegisteredMoves.Add(RegisteredMove);
-	}
-
-	for (const TSubclassOf<ULayeredMoveLogic>& PendingUnregistrantClass : MovesPendingUnregistration)
-	{
-		RegisteredMoves.RemoveAll([&PendingUnregistrantClass, this]
-			(const TObjectPtr<ULayeredMoveLogic>& MoveLogic)
-			{
-				if (MoveLogic->GetClass() == PendingUnregistrantClass)
-				{
-					return true;
-				}
-				
-				return false;
-			});
-	}
-	
-	MovesPendingRegistration.Empty();
-	MovesPendingUnregistration.Empty();
-}
-
 void UMoverComponent::UpdateCachedFrameState(const FMoverSyncState* SyncState, const FMoverAuxStateContext* AuxState)
 {
 	// TODO integrate dirty tracking
@@ -1493,87 +1465,6 @@ void UMoverComponent::RefreshSharedSettings()
 }
 
 
-const TArray<TObjectPtr<ULayeredMoveLogic>>* UMoverComponent::GetRegisteredMoves() const
-{
-	return &RegisteredMoves;
-}
-
-void UMoverComponent::K2_RegisterMove(TSubclassOf<ULayeredMoveLogic> MoveClass)
-{
-	MovesPendingUnregistration.Remove(MoveClass);
-	if (!MovesPendingRegistration.Contains(MoveClass))
-	{
-		const bool bAlreadyRegistered = RegisteredMoves.ContainsByPredicate([MoveClass](const TObjectPtr<ULayeredMoveLogic>& Move) { return Move->GetClass() == MoveClass; });
-		if (!bAlreadyRegistered)
-		{
-			MovesPendingRegistration.AddUnique(MoveClass);
-		}
-	}
-}
-
-void UMoverComponent::K2_RegisterMoves(TArray<TSubclassOf<ULayeredMoveLogic>> MoveClasses)
-{
-	for (const TSubclassOf<ULayeredMoveLogic>& MoveClass : MoveClasses)
-	{
-		K2_RegisterMove(MoveClass);
-	}
-}
-
-void UMoverComponent::K2_UnregisterMove(TSubclassOf<ULayeredMoveLogic> MoveClass)
-{
-	MovesPendingRegistration.Remove(MoveClass);
-	if (!MovesPendingUnregistration.Contains(MoveClass))
-	{
-		const bool bAlreadyUnregistered = RegisteredMoves.ContainsByPredicate([MoveClass](const TObjectPtr<ULayeredMoveLogic>& Move) { return Move->GetClass() == MoveClass; });
-		if (!bAlreadyUnregistered)
-		{
-			MovesPendingUnregistration.AddUnique(MoveClass);
-		}
-	}
-}
-
-bool UMoverComponent::K2_QueueLayeredMoveActivationWithContext(TSubclassOf<ULayeredMoveLogic> MoveLogicClass, const int32& MoveAsRawData)
-{
-	// This will never be called, the exec version below will be hit instead
-	checkNoEntry();
-	return false;
-}
-
-DEFINE_FUNCTION(UMoverComponent::execK2_QueueLayeredMoveActivationWithContext)
-{
-	P_GET_OBJECT(UClass, MoveLogicClass);
-	
-	Stack.MostRecentPropertyAddress = nullptr;
-	Stack.MostRecentPropertyContainer = nullptr;
-	Stack.StepCompiledIn<FStructProperty>(nullptr);
-
-	const FStructProperty* MoveActivationProperty = CastField<FStructProperty>(Stack.MostRecentProperty);
-	uint8* MoveActivationPtr = Stack.MostRecentPropertyAddress;
-	
-	P_FINISH
-	
-	P_NATIVE_BEGIN;
-
-	// TODO NS: throw some helpful warnings of what wasn't valid
-	const bool bHasValidActivationStructProp = MoveActivationProperty && MoveActivationProperty->Struct && MoveActivationProperty->Struct->IsChildOf(FLayeredMoveActivationParams::StaticStruct());
-
-	bool bHasValidMoveData = MoveLogicClass && bHasValidActivationStructProp; 
-	if (bHasValidMoveData)
-	{
-		const FLayeredMoveActivationParams* MoveActivationContext = reinterpret_cast<FLayeredMoveActivationParams*>(MoveActivationPtr);
-		bHasValidMoveData = P_THIS->MakeAndQueueLayeredMove(MoveLogicClass, MoveActivationContext);
-	}
-	
-	*(bool*)RESULT_PARAM = bHasValidMoveData;
-	
-	P_NATIVE_END;
-}
-
-bool UMoverComponent::QueueLayeredMoveActivation(TSubclassOf<ULayeredMoveLogic> MoveLogicClass)
-{
-	return MakeAndQueueLayeredMove(MoveLogicClass, nullptr);
-}
-
 void UMoverComponent::K2_QueueLayeredMove(const int32& MoveAsRawData)
 {
 	// This will never be called, the exec version below will be hit instead
@@ -1719,6 +1610,7 @@ DEFINE_FUNCTION(UMoverComponent::execK2_ScheduleInstantMovementEffect)
 void UMoverComponent::ScheduleInstantMovementEffect(TSharedPtr<FInstantMovementEffect> InstantMovementEffect)
 {
 	ensureMsgf(IsInGameThread(), TEXT("UMoverComponent::ScheduleInstantMovementEffect should only be called from the game thread. Inspect code for incorrect calls."));
+	ensure(IsInGameThread());
 	FMoverTimeStep TimeStep;
 	if (ensureMsgf(BackendLiaisonComp, TEXT("UMoverComponent::ScheduleInstantMovementEffect was unexpectedly called with a null backend liaison component. The instant movement effect will be ignored.")))
 	{
@@ -2498,72 +2390,6 @@ UBaseMovementMode* UMoverComponent::GetActiveModeInternal(TSubclassOf<UBaseMovem
 	return nullptr;
 }
 
-bool UMoverComponent::MakeAndQueueLayeredMove(const TSubclassOf<ULayeredMoveLogic>& MoveLogicClass, const FLayeredMoveActivationParams* ActivationParams)
-{
-	// Find registered type for class passed in
-	TObjectPtr<ULayeredMoveLogic> FoundRegisteredMoveLogic = nullptr;
-	for (TObjectPtr<ULayeredMoveLogic> RegisteredMoveLogic : RegisteredMoves)
-	{
-		if (RegisteredMoveLogic.GetClass()->IsChildOf(MoveLogicClass))
-		{
-			FoundRegisteredMoveLogic = RegisteredMoveLogic;
-			break;
-		}
-	}
-			
-	ULayeredMoveLogic* ActiveMoveLogic;
-	TSharedPtr<FLayeredMoveInstancedData> QueuedInstancedData;
-
-	if (FoundRegisteredMoveLogic)
-	{
-		ActiveMoveLogic = FoundRegisteredMoveLogic;
-
-		const UScriptStruct* InstancedDataType = FoundRegisteredMoveLogic->GetInstancedDataType();
-		if (InstancedDataType && InstancedDataType->IsChildOf(FLayeredMoveInstancedData::StaticStruct()))
-		{
-			TCheckedObjPtr<UScriptStruct> DataStructType = FoundRegisteredMoveLogic->GetInstancedDataType();
-			FLayeredMoveInstancedData* NewMove = (FLayeredMoveInstancedData*)FMemory::Malloc(DataStructType->GetCppStructOps()->GetSize());
-			DataStructType->InitializeStruct(NewMove);
-
-			struct FAllocatedLayeredMoveDataDeleter
-			{
-				FORCEINLINE void operator()(FLayeredMoveInstancedData* MoveData) const
-				{
-					check(MoveData);
-					UScriptStruct* ScriptStruct = MoveData->GetScriptStruct();
-					check(ScriptStruct);
-					ScriptStruct->DestroyStruct(MoveData);
-					FMemory::Free(MoveData);
-				}
-			};
-				
-			QueuedInstancedData = TSharedRef<FLayeredMoveInstancedData>(NewMove, FAllocatedLayeredMoveDataDeleter());
-			QueuedInstancedData->ActivateFromContext(ActivationParams);
-		}
-		else
-		{
-			UE_LOG(LogMover, Warning, TEXT("%s activation was queued on %s but the move was NOT queued since it did not have valid data. InstancedDataStructType on Move Logic needs to be a FLayeredMoveInstancedData or child struct of."),
-				*MoveLogicClass->GetName(),
-				*GetOwner()->GetName());
-			
-			return false;
-		}
-	}
-	else
-	{
-		UE_LOG(LogMover, Warning, TEXT("%s activation was queued on %s and the move was not registered. Any move activated on a MoverComponent Needs to be Registered with the MoverCompoent. The layered move will not be queued for activation."),
-			*MoveLogicClass->GetName(),
-			*GetOwner()->GetName());
-			
-		return false;
-	}
-	
-	const TSharedPtr<FLayeredMoveInstance> ActiveMoveToQueue = MakeShared<FLayeredMoveInstance>(QueuedInstancedData.ToSharedRef(), ActiveMoveLogic);
-	ModeFSM->QueueActiveLayeredMove(ActiveMoveToQueue);
-	
-	return true;
-}
-
 void UMoverComponent::SetSimulationOutput(const FMoverTimeStep& TimeStep, const UE::Mover::FSimulationOutputData& OutputData)
 {
 	CachedLastSimTickTimeStep = TimeStep;
diff --git a/Engine/Plugins/Experimental/Mover/Source/Mover/Public/DefaultMovementSet/LayeredMoves/LaunchMove.h b/Engine/Plugins/Experimental/Mover/Source/Mover/Public/DefaultMovementSet/LayeredMoves/LaunchMove.h
--- a/Engine/Plugins/Experimental/Mover/Source/Mover/Public/DefaultMovementSet/LayeredMoves/LaunchMove.h
+++ b/Engine/Plugins/Experimental/Mover/Source/Mover/Public/DefaultMovementSet/LayeredMoves/LaunchMove.h
@@ -3,69 +3,10 @@
 #pragma once
 
 #include "LayeredMove.h"
-#include "LayeredMoveBase.h"
 #include "LaunchMove.generated.h"
 
 #define UE_API MOVER_API
 
-USTRUCT(Blueprintable)
-struct FLaunchMoveActivationParams : public FLayeredMoveActivationParams
-{
-	GENERATED_BODY()
-
-	/** Velocity to apply to the updated component. Could be additive or overriding depending on MixMode setting. */
-	UPROPERTY(EditAnywhere, BlueprintReadWrite, Category = Mover, meta=(ForceUnits="cm/s"))
-	FVector LaunchVelocity = FVector::ZeroVector;
-
-	// Optional movement mode name to force the actor into before applying the impulse velocity.
-	UPROPERTY(EditAnywhere, BlueprintReadWrite, Category = Mover)
-	FName ForceMovementMode = NAME_None;
-	
-};
-
-USTRUCT(Blueprintable)
-struct FLaunchMoveData : public FLayeredMoveInstancedData
-{
-	GENERATED_BODY()
-
-	//@todo DanH: This is boilerplate begging for a macro
-	virtual FLayeredMoveInstancedData* Clone() const override { return new FLaunchMoveData(*this); }
-	virtual UScriptStruct* GetScriptStruct() const override { return StaticStruct(); }
-	
-	/** Velocity to apply to the updated component. Could be additive or overriding depending on MixMode setting. */
-	UPROPERTY(EditAnywhere, BlueprintReadWrite, Category = Mover, meta=(ForceUnits="cm/s"))
-	FVector LaunchVelocity = FVector::ZeroVector;
-
-	/** Optional movement mode name to force the actor into before applying the impulse velocity. */
-	UPROPERTY(EditAnywhere, BlueprintReadWrite, Category = Mover)
-	FName ForceMovementMode;
-	
-	virtual void ActivateFromContext(const FLayeredMoveActivationParams* ActivationParams) override;
-	
-	virtual void NetSerialize(FArchive& Ar) override;
-};
-
-// TODO: Create data for this? Is it not needed?!
-UCLASS()
-class ULaunchMoveLogic : public ULayeredMoveLogic
-{
-	GENERATED_BODY()
-
-public:
-	UE_API ULaunchMoveLogic();
-
-	/** Velocity to apply to the updated component. Could be additive or overriding depending on MixMode setting. */
-	UPROPERTY(EditAnywhere, BlueprintReadWrite, Category = Mover, meta=(ForceUnits="cm/s"))
-	FVector LaunchVelocity = FVector::ZeroVector;
-
-	//@todo DanH: Should forcing a mode be an option at the root ULayeredMoveLogic?
-	/** Optional movement mode name to force the actor into before applying the impulse velocity. */
-	UPROPERTY(EditAnywhere, BlueprintReadWrite, Category = Mover)
-	FName ForceMovementMode;
-	
-protected:
-	UE_API virtual bool GenerateMove_Implementation(const FMoverTimeStep& TimeStep, UMoverBlackboard* SimBlackboard, const FMoverTickStartData& StartState, FProposedMove& OutProposedMove) override;
-};
 
 /** Launch Move: provides an impulse velocity to the actor after (optionally) forcing them into a particular movement mode */
 USTRUCT(BlueprintType)
diff --git a/Engine/Plugins/Experimental/Mover/Source/Mover/Public/LayeredMoveBase.h b/Engine/Plugins/Experimental/Mover/Source/Mover/Public/LayeredMoveBase.h
deleted file mode 100644
diff --git a/Engine/Plugins/Experimental/Mover/Source/Mover/Public/LayeredMoveGroup.h b/Engine/Plugins/Experimental/Mover/Source/Mover/Public/LayeredMoveGroup.h
deleted file mode 100644
diff --git a/Engine/Plugins/Experimental/Mover/Source/Mover/Public/MoveLibrary/MovementMixer.h b/Engine/Plugins/Experimental/Mover/Source/Mover/Public/MoveLibrary/MovementMixer.h
--- a/Engine/Plugins/Experimental/Mover/Source/Mover/Public/MoveLibrary/MovementMixer.h
+++ b/Engine/Plugins/Experimental/Mover/Source/Mover/Public/MoveLibrary/MovementMixer.h
@@ -8,7 +8,6 @@
 #define UE_API MOVER_API
 
 struct FLayeredMoveBase;
-struct FLayeredMoveInstance;
 struct FProposedMove;
 
 /**
@@ -25,9 +24,6 @@ public:
 	/** In charge of mixing Layered Move proposed moves into a cumulative proposed move based on mix mode and priority.*/
 	UE_API virtual void MixLayeredMove(const FLayeredMoveBase& ActiveMove, const FProposedMove& MoveStep, FProposedMove& OutCumulativeMove);
 
-	/** In charge of mixing Layered Move proposed moves into a cumulative proposed move based on mix mode and priority.*/
-	UE_API virtual void MixLayeredMove(const FLayeredMoveInstance& ActiveMove, const FProposedMove& MoveStep, FProposedMove& OutCumulativeMove);
-
 	/** In charge of mixing proposed moves together. Is similar to MixLayeredMove but is only responsible for mixing proposed moves instead of layered moves. */
 	UE_API virtual void MixProposedMoves(const FProposedMove& MoveToMix, FVector UpDirection, FProposedMove& OutCumulativeMove);
 
@@ -46,7 +42,6 @@ protected:
 	 * Returns true if this layered move should take priority given current HighestPriority and CurrentLayeredMoveStartTimeMs
 	 */
 	static UE_API bool CheckPriority(const FLayeredMoveBase* LayeredMove, uint8& InOutHighestPriority, double& InOutCurrentLayeredMoveStartTimeMs);
-	static UE_API bool CheckPriority(uint8 LayeredMovePriority, double LayeredMoveStartTimeMs, uint8& InOutHighestPriority, double& InOutCurrentLayeredMoveStartTimeMs);
 };
 
 #undef UE_API
diff --git a/Engine/Plugins/Experimental/Mover/Source/Mover/Public/MovementModeStateMachine.h b/Engine/Plugins/Experimental/Mover/Source/Mover/Public/MovementModeStateMachine.h
--- a/Engine/Plugins/Experimental/Mover/Source/Mover/Public/MovementModeStateMachine.h
+++ b/Engine/Plugins/Experimental/Mover/Source/Mover/Public/MovementModeStateMachine.h
@@ -54,8 +54,6 @@ public:
 
 	UE_API void QueueLayeredMove(TSharedPtr<FLayeredMoveBase> Move);
 	
-	UE_API void QueueActiveLayeredMove(const TSharedPtr<FLayeredMoveInstance>& LayeredMove);
-
  	UE_API FMovementModifierHandle QueueMovementModifier(TSharedPtr<FMovementModifierBase> Modifier);
 
  	UE_API void CancelModifierFromHandle(FMovementModifierHandle ModifierHandle);
@@ -90,9 +88,6 @@ protected:
 	/** Moves that are queued to be added to the simulation at the start of the next sim subtick. Access covered by lock. */
 	TArray<TSharedPtr<FLayeredMoveBase>> QueuedLayeredMoves;
 
-	/** Moves that are queued to be added to the simulation at the start of the next sim subtick. Access covered by lock. */
- 	TArray<TSharedPtr<FLayeredMoveInstance>> QueuedLayeredMoveInstances;
- 	
  	/** Effects that are queued to be applied to the simulation at the start of the next sim subtick or at the end of this tick.  Access covered by lock. */
  	TArray<FScheduledInstantMovementEffect> QueuedInstantEffects;
 
@@ -116,8 +111,6 @@ private:
 	UE_API void ConstructDefaultModes();
 	UE_API void AdvanceToNextMode();
 	UE_API void FlushQueuedMovesToGroup(FLayeredMoveGroup& Group);
-	// Flushes queued ActiveLayeredMoves to FLayeredMoveInstanceGroup for this frame
- 	UE_API void ActivateQueuedMoves(FLayeredMoveInstanceGroup& Group);
  	UE_API void FlushQueuedModifiersToGroup(FMovementModifierGroup& ModifierGroup);
  	UE_API void FlushModifierCancellationsToGroup(FMovementModifierGroup& ActiveModifierGroup);
  	UE_API void RollbackModifiers(const FMoverSyncState* InvalidSyncState, const FMoverSyncState* SyncState, const FMoverAuxStateContext* InvalidAuxState, const FMoverAuxStateContext* AuxState);
diff --git a/Engine/Plugins/Experimental/Mover/Source/Mover/Public/MoverComponent.h b/Engine/Plugins/Experimental/Mover/Source/Mover/Public/MoverComponent.h
--- a/Engine/Plugins/Experimental/Mover/Source/Mover/Public/MoverComponent.h
+++ b/Engine/Plugins/Experimental/Mover/Source/Mover/Public/MoverComponent.h
@@ -10,7 +10,6 @@
 #include "MovementMode.h"
 #include "MoverTypes.h"
 #include "LayeredMove.h"
-#include "LayeredMoveBase.h"
 #include "MoveLibrary/BasedMovementUtils.h"
 #include "MoveLibrary/ConstrainedMoveUtils.h"
 #if UE_ENABLE_INCLUDE_ORDER_DEPRECATED_IN_5_4
@@ -224,58 +223,6 @@ public:
 	UPROPERTY(BlueprintReadWrite, EditAnywhere, Category = Mover)
 	TObjectPtr<UMovementMixer> MovementMixer;
 
-	const TArray<TObjectPtr<ULayeredMoveLogic>>* GetRegisteredMoves() const;
-	
-	/** Registers layered move logic */
-	template <typename MoveT UE_REQUIRES(std::is_base_of_v<ULayeredMoveLogic, MoveT>)>
-	void RegisterMove(TSubclassOf<MoveT> MoveClass = MoveT::StaticClass())
-	{
-		K2_RegisterMove(MoveClass);
-	}
-
-	/** Registers layered move logic */
-	UFUNCTION(BlueprintCallable, Category = Mover, DisplayName = "Register Move")
-	void K2_RegisterMove(TSubclassOf<ULayeredMoveLogic> MoveClass);
-
-	/** Registers an array of layered move logic classes */
-	UFUNCTION(BlueprintCallable, Category = Mover, DisplayName = "Register Moves")
-	void K2_RegisterMoves(TArray<TSubclassOf<ULayeredMoveLogic>> MoveClasses);
-	
-	/** Unregisters layered move logic */
-	template <typename MoveT UE_REQUIRES(std::is_base_of_v<ULayeredMoveLogic, MoveT>)>
-	void UnregisterMove(TSubclassOf<MoveT> MoveClass = MoveT::StaticClass())
-	{
-		K2_UnregisterMove(MoveClass);
-	}
-
-	/** Unregisters layered move logic */
-	UFUNCTION(BlueprintCallable, Category = Mover, DisplayName = "Unregister Move")
-	void K2_UnregisterMove(TSubclassOf<ULayeredMoveLogic> MoveClass);
-	
-	template <typename MoveT, typename ActivationParamsT UE_REQUIRES(std::is_base_of_v<ULayeredMoveLogic, MoveT> && std::is_base_of_v<typename MoveT::MoveDataType::ActivationParamsType, ActivationParamsT>)>
-	bool QueueLayeredMoveActivationWithContext(const ActivationParamsT& ActivationParams, TSubclassOf<MoveT> MoveClass = MoveT::StaticClass())
-	{
-		return MakeAndQueueLayeredMove(MoveClass, &ActivationParams);
-	}
-	
-	/**
-	 * Queues a layered move for activation.
-	 * Takes a Activation Context which provides context to set Layered Move Data.
-	 * Make sure Activation Context type matches layered Move Data
-	 */
-	UFUNCTION(BlueprintCallable, CustomThunk, Category = Mover, meta = (CustomStructureParam = "MoveAsRawData", AllowAbstract = "false"), DisplayName = "Queue Layered Move Activation With Context")
-	bool K2_QueueLayeredMoveActivationWithContext(TSubclassOf<ULayeredMoveLogic> MoveLogicClass, UPARAM(DisplayName="Layered Move Activation Context") const int32& MoveAsRawData);
-	DECLARE_FUNCTION(execK2_QueueLayeredMoveActivationWithContext);
-
-	/**
- 	 * Queues a layered move for activation.
- 	 * Takes NO Activation Context meaning the layered move will be activated using default Move Data.
- 	 * Note: Changing Move Data is still possible in the layered move logic itself
- 	 * See QueueLayeredMoveActivationWithContext for activating a layered move with context
- 	 */
-	UFUNCTION(BlueprintCallable, Category = Mover, meta = (AllowAbstract = "false"))
-	bool QueueLayeredMoveActivation(TSubclassOf<ULayeredMoveLogic> MoveLogicClass);
-	
 	/**
 	 * Queue a layered move to start during the next simulation frame. This will clone whatever move you pass in, so you'll need to fully set it up before queuing.
 	 * @param LayeredMove			The move to queue, which must be a LayeredMoveBase sub-type. 
@@ -341,7 +288,6 @@ public:
 	MOVER_API void ClearQueuedInstantMovementEffects();
 	// Queue an instant movement effect in async mode. Do not use on the game thread.
 	void QueueInstantMovementEffect_Internal(const FMoverTimeStep& TimeStep, TSharedPtr<FInstantMovementEffect> InstantMovementEffect);
-
 protected:
 	// Queue a scheduled instant movement effect. Thread safe, can be used outside the game thread.
 	void QueueInstantMovementEffect(const FScheduledInstantMovementEffect& ScheduledInstantMovementEffect);
@@ -668,9 +614,6 @@ public:	// Queries
 	
 protected:
 
-	// Called before each simulation tick. Broadcasts OnPreSimulationTick delegate.
-	void PreSimulationTick(const FMoverTimeStep& TimeStep, const FMoverInputCmdContext& InputCmd);
-	
 	/** Makes this component and owner actor reflect the state of a particular frame snapshot. This occurs after simulation ticking, as well as during a rollback before we resimulate forward.
 	  @param bRebaseBasedState	If true and the state was using based movement, it will use the current game world base pos/rot instead of the captured one. This is necessary during rollbacks.
 	*/
@@ -783,21 +726,7 @@ protected:
 	UPROPERTY(Transient)
 	TObjectPtr<URollbackBlackboard_InternalWrapper> RollbackBlackboard_InternalWrapper = nullptr;	// This is a thin layer for use only by in-simulation users
 
-	/**
-	 * Layered moves registered on this component that can be activated regardless of the current mode
-	 * Changes to this array or its contents occur ONLY during PreSimulationTick to ensure threadsafe access during async simulations.
-	 */
-	UPROPERTY(Transient)
-	TArray<TObjectPtr<ULayeredMoveLogic>> RegisteredMoves;
-
-	UPROPERTY(Transient)
-	TArray<TSubclassOf<ULayeredMoveLogic>> MovesPendingRegistration;
-
-	UPROPERTY(Transient)
-	TArray<TSubclassOf<ULayeredMoveLogic>> MovesPendingUnregistration;
 	
-	/** Helper function for making Queueing and making Layered Moves */
-	bool MakeAndQueueLayeredMove(const TSubclassOf<ULayeredMoveLogic>& MoveLogicClass, const FLayeredMoveActivationParams* ActivationParams);
 	
 private:
 	/** Collection of settings objects that are shared between movement modes. This list is automatically managed based on the @MovementModes contents. */
diff --git a/Engine/Plugins/Experimental/Mover/Source/Mover/Public/MoverSimulationTypes.h b/Engine/Plugins/Experimental/Mover/Source/Mover/Public/MoverSimulationTypes.h
--- a/Engine/Plugins/Experimental/Mover/Source/Mover/Public/MoverSimulationTypes.h
+++ b/Engine/Plugins/Experimental/Mover/Source/Mover/Public/MoverSimulationTypes.h
@@ -8,7 +8,6 @@
 #include "MoverTypes.h"
 #include "MoveLibrary/MovementRecord.h"
 #include "LayeredMove.h"
-#include "LayeredMoveGroup.h"
 #include "MovementModifier.h"
 #include "MoverDataModelTypes.h"
 #include "InstantMovementEffect.h"
@@ -202,10 +201,6 @@ public:
 	UPROPERTY(BlueprintReadOnly, VisibleAnywhere, Category = Mover)
 	FLayeredMoveGroup LayeredMoves;
 
-	// Additional moves influencing our proposed motion
-	UPROPERTY(BlueprintReadOnly, VisibleAnywhere, Category = Mover)
-	FLayeredMoveInstanceGroup LayeredMoveInstances;
-
 	// Additional modifiers influencing our simulation
 	UPROPERTY(BlueprintReadOnly, VisibleAnywhere, Category = Mover)
 	FMovementModifierGroup MovementModifiers;
@@ -222,7 +217,6 @@ public:
 	{
 		return MovementMode == Other.MovementMode &&
 			LayeredMoves.HasSameContents(Other.LayeredMoves) &&
-			LayeredMoveInstances.HasSameContents(Other.LayeredMoveInstances) &&
 			MovementModifiers.HasSameContents(Other.MovementModifiers) &&
 			SyncStateCollection.HasSameContents(Other.SyncStateCollection);
 	}
@@ -234,7 +228,6 @@ public:
 	{
 		P.Ar << MovementMode;
 		LayeredMoves.NetSerialize(P.Ar);
-		LayeredMoveInstances.NetSerialize(P.Ar);
 		MovementModifiers.NetSerialize(P.Ar);
 
 		bool bIgnoredResult(false);
@@ -245,7 +238,6 @@ public:
 	{
 		Out.Appendf("MovementMode: %s\n", TCHAR_TO_ANSI(*MovementMode.ToString()));
 		Out.Appendf("Layered Moves: %s\n", TCHAR_TO_ANSI(*LayeredMoves.ToSimpleString()));
-		Out.Appendf("Layered Moves: %s\n", TCHAR_TO_ANSI(*LayeredMoveInstances.ToSimpleString()));
 		Out.Appendf("Movement Modifiers: %s\n", TCHAR_TO_ANSI(*MovementModifiers.ToSimpleString()));
 		SyncStateCollection.ToString(Out);
 	}
@@ -261,7 +253,6 @@ public:
 	{
 		MovementMode = To->MovementMode;
 		LayeredMoves = To->LayeredMoves;
-		LayeredMoveInstances = To->LayeredMoveInstances;
 		MovementModifiers = To->MovementModifiers;
 
 		SyncStateCollection.Interpolate(From->SyncStateCollection, To->SyncStateCollection, Pct);
@@ -274,7 +265,6 @@ public:
 		MovementMode = NAME_None;
 		SyncStateCollection.Empty();
 		LayeredMoves.Reset();
-		LayeredMoveInstances.Reset();
 		MovementModifiers.Reset();
 	}
 };
