[MoveLibrary] Replace FRWLock in UMoverBlackboard with transactionally-safe RWLock

This code is called from the closed and use of a regular RWLock will crash.

diff --git a/Engine/Plugins/Experimental/Mover/Source/Mover/Private/MoveLibrary/MoverBlackboard.cpp b/Engine/Plugins/Experimental/Mover/Source/Mover/Private/MoveLibrary/MoverBlackboard.cpp
--- a/Engine/Plugins/Experimental/Mover/Source/Mover/Private/MoveLibrary/MoverBlackboard.cpp
+++ b/Engine/Plugins/Experimental/Mover/Source/Mover/Private/MoveLibrary/MoverBlackboard.cpp
@@ -9,7 +9,7 @@
 
 void UMoverBlackboard::Invalidate(FName ObjName)
 {
-	FRWScopeLock Lock(ObjectsMapLock, SLT_Write);
+	UE::TWriteScopeLock Lock(ObjectsMapLock);
 	ObjectsByName.Remove(ObjName);
 }
 
@@ -20,7 +20,7 @@ void UMoverBlackboard::Invalidate(EInvalidationReason Reason)
 		default:
 		case EInvalidationReason::FullReset:
 		{
-			FRWScopeLock Lock(ObjectsMapLock, SLT_Write);
+			UE::TWriteScopeLock Lock(ObjectsMapLock);
 			ObjectsByName.Empty();
 		}
 		break;
diff --git a/Engine/Plugins/Experimental/Mover/Source/Mover/Public/MoveLibrary/MoverBlackboard.h b/Engine/Plugins/Experimental/Mover/Source/Mover/Public/MoveLibrary/MoverBlackboard.h
--- a/Engine/Plugins/Experimental/Mover/Source/Mover/Public/MoveLibrary/MoverBlackboard.h
+++ b/Engine/Plugins/Experimental/Mover/Source/Mover/Public/MoveLibrary/MoverBlackboard.h
@@ -2,6 +2,7 @@
 
 #pragma once
 
+#include "Misc/TransactionallySafeRWLock.h"
 #include "UObject/Object.h"
 #include "Templates/SharedPointer.h"
 #include "Templates/Casts.h"
@@ -80,7 +81,7 @@ public:
  	template<typename T>
 	bool TryGet(FName ObjName, T& OutFoundValue) const
 	{
-		FRWScopeLock Lock(ObjectsMapLock, SLT_ReadOnly);
+		UE::TReadScopeLock Lock(ObjectsMapLock);
 
 		if (const TUniquePtr<BlackboardObject>* ExistingObject = ObjectsByName.Find(ObjName))
 		{
@@ -115,7 +116,7 @@ public:
 	/** Returns true/false to indicate if an object is stored with that name */
 	bool Contains(FName ObjName) const
 	{
-		FRWScopeLock Lock(ObjectsMapLock, SLT_ReadOnly);
+		UE::TReadScopeLock Lock(ObjectsMapLock);
 		return ObjectsByName.Contains(ObjName);
 	}
 
@@ -123,7 +124,7 @@ public:
 	template<typename T>
 	void Set(FName ObjName, T Obj)
 	{
-		FRWScopeLock Lock(ObjectsMapLock, SLT_Write);
+		UE::TWriteScopeLock Lock(ObjectsMapLock);
 		ObjectsByName.Emplace(ObjName, MakeUnique<BlackboardObject>(Obj));
 	}
 
@@ -141,7 +142,7 @@ public:
 	// End UObject interface
 
 private:
-	mutable FRWLock ObjectsMapLock;		// used internally when reading/writing to ObjectsByName map
+	mutable FTransactionallySafeRWLock ObjectsMapLock;		// used internally when reading/writing to ObjectsByName map
 	TMap<FName, TUniquePtr<BlackboardObject>> ObjectsByName;
 };
 
