Mover; Expose SimpleSpringWalkingMode API as blueprint native event
Add MaxSpeedOverride

diff --git a/Engine/Plugins/Experimental/Mover/Source/Mover/Private/DefaultMovementSet/Modes/SimpleSpringWalkingMode.cpp b/Engine/Plugins/Experimental/Mover/Source/Mover/Private/DefaultMovementSet/Modes/SimpleSpringWalkingMode.cpp
--- a/Engine/Plugins/Experimental/Mover/Source/Mover/Private/DefaultMovementSet/Modes/SimpleSpringWalkingMode.cpp
+++ b/Engine/Plugins/Experimental/Mover/Source/Mover/Private/DefaultMovementSet/Modes/SimpleSpringWalkingMode.cpp
@@ -21,15 +21,14 @@ void USimpleSpringWalkingMode::SimulationTick_Implementation(const FSimulationTi
 	}
 }
 
-void USimpleSpringWalkingMode::GenerateWalkMove(const FMoverTickStartData& StartState, float DeltaSeconds, const FVector& DesiredVelocity,
-	const FQuat& DesiredFacing, FVector& OutVelocity, FQuat& OutFacing) const
+void USimpleSpringWalkingMode::GenerateWalkMove_Implementation(FMoverTickStartData& StartState, float DeltaSeconds, const FVector& DesiredVelocity,
+	const FQuat& DesiredFacing, FVector& OutVelocity, FQuat& OutFacing)
 {
 	const FMoverDefaultSyncState* StartingSyncState = StartState.SyncState.SyncStateCollection.FindDataByType<FMoverDefaultSyncState>();
 	FVector CurrentPosition = StartingSyncState->GetLocation_WorldSpace();
 	FVector PrevPosition = CurrentPosition;
 
-	FMoverTickStartData* MutableStartData = const_cast<FMoverTickStartData*>(&StartState);
-	FSimpleSpringState& SpringState = MutableStartData->SyncState.SyncStateCollection.FindOrAddMutableDataByType<FSimpleSpringState>();
+	FSimpleSpringState& SpringState = StartState.SyncState.SyncStateCollection.FindOrAddMutableDataByType<FSimpleSpringState>();
 
 	// Linear //
 	SpringMath::SpringCharacterUpdate(CurrentPosition, SpringState.CurrentVelocity, SpringState.CurrentAccel, DesiredVelocity, VelocitySmoothingTime, DeltaSeconds, VelocityDeadzoneThreshold);
diff --git a/Engine/Plugins/Experimental/Mover/Source/Mover/Private/DefaultMovementSet/Modes/SimpleWalkingMode.cpp b/Engine/Plugins/Experimental/Mover/Source/Mover/Private/DefaultMovementSet/Modes/SimpleWalkingMode.cpp
--- a/Engine/Plugins/Experimental/Mover/Source/Mover/Private/DefaultMovementSet/Modes/SimpleWalkingMode.cpp
+++ b/Engine/Plugins/Experimental/Mover/Source/Mover/Private/DefaultMovementSet/Modes/SimpleWalkingMode.cpp
@@ -42,7 +42,7 @@ void USimpleWalkingMode::GenerateMove_Implementation(const FMoverTickStartData&
 	}
 	
 	
-	float MaxMoveSpeed = CommonLegacySettings->MaxSpeed;
+	float MaxMoveSpeed = MaxSpeedOverride >= 0.0f ? MaxSpeedOverride : CommonLegacySettings->MaxSpeed;
 
 	// Subtract vertical component but keep same magnitude
 	float DesiredVelMag = DesiredVelocity.Length();
@@ -92,7 +92,12 @@ void USimpleWalkingMode::GenerateMove_Implementation(const FMoverTickStartData&
 
 	FQuat OutputFacing = CurrentFacing;
 	OutProposedMove.LinearVelocity = StartingSyncState->GetVelocity_WorldSpace();
-	GenerateWalkMove(StartState, DeltaSeconds, DesiredVelocity, DesiredFacing, OutProposedMove.LinearVelocity, OutputFacing);
+	
+	// Hack const_cast stuff
+	// Why is this needed?
+	// Because some modes want to mutate their data inside the generate walk move
+	USimpleWalkingMode* MutableSimpleWalkMode  = const_cast<USimpleWalkingMode*>(this);
+	MutableSimpleWalkMode->GenerateWalkMove(const_cast<FMoverTickStartData&>(StartState), DeltaSeconds, DesiredVelocity, DesiredFacing, OutProposedMove.LinearVelocity, OutputFacing);
 
 	// Calc angular velocity
 	FQuat Diff = (PrevFacing.Inverse() * OutputFacing);
@@ -104,8 +109,8 @@ void USimpleWalkingMode::GenerateMove_Implementation(const FMoverTickStartData&
 	OutProposedMove.AngularVelocity = AngularVelocity;
 }
 
-void USimpleWalkingMode::GenerateWalkMove(const FMoverTickStartData& StartState, float DeltaSeconds, const FVector& DesiredVelocity,
-	const FQuat& DesiredFacing, FVector& OutVelocity, FQuat& OutFacing) const
+void USimpleWalkingMode::GenerateWalkMove_Implementation(FMoverTickStartData& StartState, float DeltaSeconds, const FVector& DesiredVelocity,
+	const FQuat& DesiredFacing, FVector& OutVelocity, FQuat& OutFacing)
 {
 	OutVelocity = DesiredVelocity;
 	OutFacing = DesiredFacing;
diff --git a/Engine/Plugins/Experimental/Mover/Source/Mover/Private/DefaultMovementSet/Modes/SmoothWalkingMode.cpp b/Engine/Plugins/Experimental/Mover/Source/Mover/Private/DefaultMovementSet/Modes/SmoothWalkingMode.cpp
--- a/Engine/Plugins/Experimental/Mover/Source/Mover/Private/DefaultMovementSet/Modes/SmoothWalkingMode.cpp
+++ b/Engine/Plugins/Experimental/Mover/Source/Mover/Private/DefaultMovementSet/Modes/SmoothWalkingMode.cpp
@@ -21,8 +21,8 @@ void USmoothWalkingMode::SimulationTick_Implementation(const FSimulationTickPara
 	}
 }
 
-void USmoothWalkingMode::GenerateWalkMove(const FMoverTickStartData& StartState, float DeltaSeconds, const FVector& DesiredVelocity,
-	const FQuat& DesiredFacing, FVector& InOutVelocity, FQuat& InOutFacing) const
+void USmoothWalkingMode::GenerateWalkMove_Implementation(FMoverTickStartData& StartState, float DeltaSeconds, const FVector& DesiredVelocity,
+	const FQuat& DesiredFacing, FVector& InOutVelocity, FQuat& InOutFacing) 
 {
 	if (DeltaSeconds <= FLT_EPSILON)
 	{
@@ -35,11 +35,9 @@ void USmoothWalkingMode::GenerateWalkMove(const FMoverTickStartData& StartState,
 		return;
 	}
 
-	FMoverTickStartData* MutableStartData = const_cast<FMoverTickStartData*>(&StartState);
-	
 	// Find or add a FSmoothWalkingState to the SyncState
 	bool bSmoothWalkingStateAdded = false;
-	FSmoothWalkingState& SpringState = MutableStartData->SyncState.SyncStateCollection.FindOrAddMutableDataByType<FSmoothWalkingState>(bSmoothWalkingStateAdded);
+	FSmoothWalkingState& SpringState = StartState.SyncState.SyncStateCollection.FindOrAddMutableDataByType<FSmoothWalkingState>(bSmoothWalkingStateAdded);
 
 	// If the state was not there already we need to initialize some of the intermediate state to whatever we have as the current state to avoid 
 	// a discontinuity. Unfortunately there is no way currently to initialize the angular velocities or accelerations right now as these are not 
diff --git a/Engine/Plugins/Experimental/Mover/Source/Mover/Public/DefaultMovementSet/Modes/SimpleSpringWalkingMode.h b/Engine/Plugins/Experimental/Mover/Source/Mover/Public/DefaultMovementSet/Modes/SimpleSpringWalkingMode.h
--- a/Engine/Plugins/Experimental/Mover/Source/Mover/Public/DefaultMovementSet/Modes/SimpleSpringWalkingMode.h
+++ b/Engine/Plugins/Experimental/Mover/Source/Mover/Public/DefaultMovementSet/Modes/SimpleSpringWalkingMode.h
@@ -19,8 +19,8 @@ class USimpleSpringWalkingMode : public USimpleWalkingMode
 
 public:
 	UE_API virtual void SimulationTick_Implementation(const FSimulationTickParams& Params, FMoverTickEndData& OutputState) override;
-	UE_API virtual void GenerateWalkMove(const FMoverTickStartData& StartState, float DeltaSeconds, const FVector& DesiredVelocity,
-									 const FQuat& DesiredFacing, FVector& OutVelocity, FQuat& OutFacing) const override;
+	UE_API virtual void GenerateWalkMove_Implementation(FMoverTickStartData& StartState, float DeltaSeconds, const FVector& DesiredVelocity,
+									 const FQuat& DesiredFacing, FVector& OutVelocity, FQuat& OutFacing) override;
 
 protected:
 	UPROPERTY(EditAnywhere, BlueprintReadWrite, Category = "Mover|Spring Settings", meta = (ClampMin = "0", UIMin = "0", ForceUnits = "s"))
diff --git a/Engine/Plugins/Experimental/Mover/Source/Mover/Public/DefaultMovementSet/Modes/SimpleWalkingMode.h b/Engine/Plugins/Experimental/Mover/Source/Mover/Public/DefaultMovementSet/Modes/SimpleWalkingMode.h
--- a/Engine/Plugins/Experimental/Mover/Source/Mover/Public/DefaultMovementSet/Modes/SimpleWalkingMode.h
+++ b/Engine/Plugins/Experimental/Mover/Source/Mover/Public/DefaultMovementSet/Modes/SimpleWalkingMode.h
@@ -21,8 +21,13 @@ public:
 	UE_API virtual void GenerateMove_Implementation(const FMoverTickStartData& StartState, const FMoverTimeStep& TimeStep, FProposedMove& OutProposedMove) const override;
 
 	// Override this to make a simple walking mode
-	UE_API virtual void GenerateWalkMove(const FMoverTickStartData& StartState, float DeltaSeconds, const FVector& DesiredVelocity,
-	                                     const FQuat& DesiredFacing, FVector& InOutVelocity, FQuat& InOutFacing) const;
+	UFUNCTION(BlueprintNativeEvent, meta = (DisplayName = "Generate Simple Walk Move"))
+	UE_API void GenerateWalkMove(UPARAM(ref) FMoverTickStartData& StartState, float DeltaSeconds, const FVector& DesiredVelocity,
+	                                     const FQuat& DesiredFacing,UPARAM(ref)  FVector& InOutVelocity, UPARAM(ref) FQuat& InOutFacing);
+
+	// If this value is greater or equal to 0, this will override the max speed read from the common legacy shared walk settings
+	UPROPERTY(EditAnywhere, BlueprintReadWrite, Category = "Mover", meta = (ForceUnits = "cm/s"))
+	float MaxSpeedOverride = -1.0f;
 };
 
 #undef UE_API
diff --git a/Engine/Plugins/Experimental/Mover/Source/Mover/Public/DefaultMovementSet/Modes/SmoothWalkingMode.h b/Engine/Plugins/Experimental/Mover/Source/Mover/Public/DefaultMovementSet/Modes/SmoothWalkingMode.h
--- a/Engine/Plugins/Experimental/Mover/Source/Mover/Public/DefaultMovementSet/Modes/SmoothWalkingMode.h
+++ b/Engine/Plugins/Experimental/Mover/Source/Mover/Public/DefaultMovementSet/Modes/SmoothWalkingMode.h
@@ -19,8 +19,8 @@ class USmoothWalkingMode : public USimpleWalkingMode
 
 public:
 	UE_API virtual void SimulationTick_Implementation(const FSimulationTickParams& Params, FMoverTickEndData& OutputState) override;
-	UE_API virtual void GenerateWalkMove(const FMoverTickStartData& StartState, float DeltaSeconds, const FVector& DesiredVelocity,
-									 const FQuat& DesiredFacing, FVector& InOutVelocity, FQuat& InOutFacing) const override;
+	UE_API virtual void GenerateWalkMove_Implementation(FMoverTickStartData& StartState, float DeltaSeconds, const FVector& DesiredVelocity,
+									 const FQuat& DesiredFacing, FVector& InOutVelocity, FQuat& InOutFacing) override;
 
 
 protected: // Velocity Controls
