Mover: address issues with weighting and rotator interpolaton, plus loctext mismatch

#rnx

diff --git a/Engine/Plugins/Experimental/Mover/Source/Mover/Private/BlueprintDataExt.cpp b/Engine/Plugins/Experimental/Mover/Source/Mover/Private/BlueprintDataExt.cpp
--- a/Engine/Plugins/Experimental/Mover/Source/Mover/Private/BlueprintDataExt.cpp
+++ b/Engine/Plugins/Experimental/Mover/Source/Mover/Private/BlueprintDataExt.cpp
@@ -5,8 +5,6 @@
 
 #include UE_INLINE_GENERATED_CPP_BY_NAME(BlueprintDataExt)
 
-#define LOCTEXT_NAMESPACE "MoverBPDataExt"
-
 
 static constexpr uint8 NumValueMaps = 6;
 
@@ -84,7 +82,7 @@ bool FMoverDictionaryData::NetSerialize(FArchive& Ar, UPackageMap* Map, bool& bO
 		NameValues.Empty();
 	}
 
-	bOutSuccess = true;
+	bOutSuccess = Ar.IsError();
 	return true;
 }
 
@@ -135,7 +133,7 @@ void FMoverDictionaryData::Interpolate(const FMoverDataStructBase& From, const F
 	const FMoverDictionaryData& SecondaryInfluence = Alpha < 0.5f ? ToTyped : FromTyped;
 
 	const float PrimaryWeight = Alpha < 0.5f ? 1.0f-Alpha : Alpha;
-	const float SecondaryWeight = 1.0f - Alpha;
+	const float SecondaryWeight = 1.0f - PrimaryWeight;
 
 	BoolValues = PrimaryInfluence.BoolValues;
 	IntValues  = PrimaryInfluence.IntValues;
@@ -187,7 +185,7 @@ void FMoverDictionaryData::Interpolate(const FMoverDataStructBase& From, const F
 		}
 		else
 		{
-			RotatorValues.Add(PrimaryPair.Key);
+			RotatorValues.Add(PrimaryPair);
 		}
 	}
 }
