Fix inversion on success value.

#rnx

diff --git a/Engine/Plugins/Experimental/Mover/Source/Mover/Private/BlueprintDataExt.cpp b/Engine/Plugins/Experimental/Mover/Source/Mover/Private/BlueprintDataExt.cpp
--- a/Engine/Plugins/Experimental/Mover/Source/Mover/Private/BlueprintDataExt.cpp
+++ b/Engine/Plugins/Experimental/Mover/Source/Mover/Private/BlueprintDataExt.cpp
@@ -82,7 +82,7 @@ bool FMoverDictionaryData::NetSerialize(FArchive& Ar, UPackageMap* Map, bool& bO
 		NameValues.Empty();
 	}
 
-	bOutSuccess = Ar.IsError();
+	bOutSuccess = !Ar.IsError();
 	return true;
 }
 
