Mover: User-Defined Struct support for Mover data collections (input cmds, sync state) so Blueprint-only users can include custom data without requiring native code
Includes Network Prediction Plugin changes to pass the package map along during NetSerialization

diff --git a/Engine/Plugins/Experimental/ChaosMover/Source/ChaosMover/Private/Backends/NetworkChaosMoverData.cpp b/Engine/Plugins/Experimental/ChaosMover/Source/ChaosMover/Private/Backends/NetworkChaosMoverData.cpp
--- a/Engine/Plugins/Experimental/ChaosMover/Source/ChaosMover/Private/Backends/NetworkChaosMoverData.cpp
+++ b/Engine/Plugins/Experimental/ChaosMover/Source/ChaosMover/Private/Backends/NetworkChaosMoverData.cpp
@@ -28,13 +28,13 @@ void FNetworkChaosMoverInputData::DecayData(float DecayAmount)
 	InputCmdContext.InputCollection.Decay(DecayAmount);
 }
 
-bool FNetworkChaosMoverInputData::NetSerialize(FArchive& Ar, class UPackageMap* PackageMap, bool& bOutSuccess)
+bool FNetworkChaosMoverInputData::NetSerialize(FArchive& Ar, UPackageMap* PackageMap, bool& bOutSuccess)
 {
 	SerializeFrames(Ar);
 	if (PackageMap)
 	{
-		InputCmdContext.NetSerialize(FNetSerializeParams(Ar));
-		bOutSuccess = true;
+		InputCmdContext.NetSerialize(FNetSerializeParams(Ar, PackageMap));
+		bOutSuccess = !Ar.IsError();
 	}
 	else
 	{
@@ -102,14 +102,14 @@ void FNetworkChaosMoverStateData::BuildData(const UActorComponent* NetworkCompon
 	}
 }
 
-bool FNetworkChaosMoverStateData::NetSerialize(FArchive& Ar, class UPackageMap* PackageMap, bool& bOutSuccess)
+bool FNetworkChaosMoverStateData::NetSerialize(FArchive& Ar, UPackageMap* PackageMap, bool& bOutSuccess)
 {
 	SerializeFrames(Ar);
 	if (PackageMap)
 	{
-		FNetSerializeParams Params(Ar);
+		FNetSerializeParams Params(Ar, PackageMap);
 		SyncState.NetSerialize(Params);
-		bOutSuccess = true;
+		bOutSuccess = !Ar.IsError();
 	}
 	else
 	{
diff --git a/Engine/Plugins/Experimental/Mover/Source/Mover/Private/Backends/MoverNetworkPhysicsLiaisonBase.cpp b/Engine/Plugins/Experimental/Mover/Source/Mover/Private/Backends/MoverNetworkPhysicsLiaisonBase.cpp
--- a/Engine/Plugins/Experimental/Mover/Source/Mover/Private/Backends/MoverNetworkPhysicsLiaisonBase.cpp
+++ b/Engine/Plugins/Experimental/Mover/Source/Mover/Private/Backends/MoverNetworkPhysicsLiaisonBase.cpp
@@ -51,13 +51,13 @@ void FNetworkPhysicsMoverInputs::DecayData(float DecayAmount)
 	InputCmdContext.InputCollection.Decay(DecayAmount);
 }
 
-bool FNetworkPhysicsMoverInputs::NetSerialize(FArchive& Ar, class UPackageMap* PackageMap, bool& bOutSuccess)
+bool FNetworkPhysicsMoverInputs::NetSerialize(FArchive& Ar, UPackageMap* PackageMap, bool& bOutSuccess)
 {
 	SerializeFrames(Ar);
 	if (PackageMap)
 	{
-		InputCmdContext.NetSerialize(FNetSerializeParams(Ar));
-		bOutSuccess = true;
+		InputCmdContext.NetSerialize(FNetSerializeParams(Ar, PackageMap));
+		bOutSuccess = !Ar.IsError();
 	}
 	else
 	{
@@ -131,14 +131,14 @@ void FNetworkPhysicsMoverState::BuildData(const UActorComponent* NetworkComponen
 	}
 }
 
-bool FNetworkPhysicsMoverState::NetSerialize(FArchive& Ar, class UPackageMap* PackageMap, bool& bOutSuccess)
+bool FNetworkPhysicsMoverState::NetSerialize(FArchive& Ar, UPackageMap* PackageMap, bool& bOutSuccess)
 {
 	SerializeFrames(Ar);
 	if (PackageMap)
 	{
-		FNetSerializeParams Params(Ar);
+		FNetSerializeParams Params(Ar, PackageMap);
 		SyncStateContext.NetSerialize(Params);
-		bOutSuccess = true;
+		bOutSuccess = !Ar.IsError();
 	}
 	else
 	{
diff --git a/Engine/Plugins/Experimental/Mover/Source/Mover/Private/MoverTypes.cpp b/Engine/Plugins/Experimental/Mover/Source/Mover/Private/MoverTypes.cpp
--- a/Engine/Plugins/Experimental/Mover/Source/Mover/Private/MoverTypes.cpp
+++ b/Engine/Plugins/Experimental/Mover/Source/Mover/Private/MoverTypes.cpp
@@ -5,7 +5,9 @@
 #include "Blueprint/BlueprintExceptionInfo.h"
 #include "MoverLog.h"
 #include "MoverModule.h"
+#include "StructUtils/UserDefinedStruct.h"
 #include "UObject/ObjectKey.h"
+#include "UserDefinedStructSupport.h"
 #include "HAL/IConsoleManager.h"
 
 #include UE_INLINE_GENERATED_CPP_BY_NAME(MoverTypes)
@@ -70,6 +72,11 @@ void FMoverDataStructBase::Merge(const FMoverDataStructBase& From)
 		__FUNCTION__, *GetScriptStruct()->GetName());
 }
 
+const UScriptStruct* FMoverDataStructBase::GetDataScriptStruct() const
+{ 
+	return GetScriptStruct(); 
+}
+
 
 FMoverDataCollection::FMoverDataCollection()
 {
@@ -77,7 +84,7 @@ FMoverDataCollection::FMoverDataCollection()
 
 bool FMoverDataCollection::NetSerialize(FArchive& Ar, UPackageMap* Map, bool& bOutSuccess)
 {
-	NetSerializeDataArray(Ar, DataArray);
+	NetSerializeDataArray(Ar, Map, DataArray);
 
 	if (Ar.IsError())
 	{
@@ -257,10 +264,11 @@ bool FMoverDataCollection::ShouldReconcile(const FMoverDataCollection& Other) co
 	for (int32 i = 0; i < DataArray.Num(); ++i)
 	{
 		const FMoverDataStructBase* DataElement = DataArray[i].Get();
-		const FMoverDataStructBase* OtherDataElement = Other.FindDataByType(DataElement->GetScriptStruct());
-		
-		if (OtherDataElement == nullptr || 
-		    DataElement->ShouldReconcile(*OtherDataElement))
+		const FMoverDataStructBase* OtherDataElement = Other.FindDataByType(DataElement->GetDataScriptStruct());
+
+		// Reconciliation is needed if there's no matching types, or if the element pair needs reconciliation
+		if (OtherDataElement == nullptr ||
+			DataElement->ShouldReconcile(*OtherDataElement))
 		{
 			return true;
 		}
@@ -271,15 +279,17 @@ bool FMoverDataCollection::ShouldReconcile(const FMoverDataCollection& Other) co
 
 void FMoverDataCollection::Interpolate(const FMoverDataCollection& From, const FMoverDataCollection& To, float Pct)
 {
-	TSet<TObjectKey<UScriptStruct>> InterpolatedTypes;
-	
+	// TODO: Consider an inline allocator to avoid dynamic memory allocations
+	TSet<TObjectKey<UScriptStruct>> AddedDataTypes;
+
 	// Piece-wise interpolation of matching data blocks
 	for (const TSharedPtr<FMoverDataStructBase>& FromElement : From.DataArray)
 	{
-		InterpolatedTypes.Add(FromElement->GetScriptStruct());
-		if (const FMoverDataStructBase* ToElement = To.FindDataByType(FromElement->GetScriptStruct()))
+		AddedDataTypes.Add(FromElement->GetDataScriptStruct());
+
+		if (const FMoverDataStructBase* ToElement = To.FindDataByType(FromElement->GetDataScriptStruct()))
 		{
-			FMoverDataStructBase* InterpElement = FindOrAddDataByType(FromElement->GetScriptStruct());
+			FMoverDataStructBase* InterpElement = FindOrAddDataByType(FromElement->GetDataScriptStruct());
 			InterpElement->Interpolate(*FromElement, *ToElement, Pct);
 		}
 		else
@@ -292,7 +302,7 @@ void FMoverDataCollection::Interpolate(const FMoverDataCollection& From, const F
 	// Add any types present only in To as well
 	for (const TSharedPtr<FMoverDataStructBase>& ToElement : To.DataArray)
 	{
-		if (!InterpolatedTypes.Contains(ToElement->GetScriptStruct()))
+		if (!AddedDataTypes.Contains(ToElement->GetDataScriptStruct()))
 		{
 			AddDataByCopy(ToElement.Get());
 		}
@@ -303,7 +313,7 @@ void FMoverDataCollection::Merge(const FMoverDataCollection& From)
 {
 	for (const TSharedPtr<FMoverDataStructBase>& FromElement : From.DataArray)
 	{
-		if (FMoverDataStructBase* ExistingElement = FindDataByType(FromElement->GetScriptStruct()))
+		if (FMoverDataStructBase* ExistingElement = FindDataByType(FromElement->GetDataScriptStruct()))
 		{
 			ExistingElement->Merge(*FromElement);
 		}
@@ -333,7 +343,7 @@ bool FMoverDataCollection::HasSameContents(const FMoverDataCollection& Other) co
 
 	for (int32 i = 0; i < DataArray.Num(); ++i)
 	{
-		if (DataArray[i]->GetScriptStruct() != Other.DataArray[i]->GetScriptStruct())
+		if (DataArray[i]->GetDataScriptStruct() != Other.DataArray[i]->GetDataScriptStruct())
 		{
 			return false;
 		}
@@ -374,6 +384,8 @@ TArray<TSharedPtr<FMoverDataStructBase>>::TConstIterator FMoverDataCollection::G
 //static 
 TSharedPtr<FMoverDataStructBase> FMoverDataCollection::CreateDataByType(const UScriptStruct* DataStructType)
 {
+	check(DataStructType->IsChildOf(FMoverDataStructBase::StaticStruct()));
+
 	FMoverDataStructBase* NewDataBlock = (FMoverDataStructBase*)FMemory::Malloc(DataStructType->GetCppStructOps()->GetSize());
 	DataStructType->InitializeStruct(NewDataBlock);
 
@@ -385,7 +397,18 @@ FMoverDataStructBase* FMoverDataCollection::AddDataByType(const UScriptStruct* D
 {
 	if (ensure(!FindDataByType(DataStructType)))
 	{
-		TSharedPtr<FMoverDataStructBase> NewDataInstance = CreateDataByType(DataStructType);
+		TSharedPtr<FMoverDataStructBase> NewDataInstance;
+
+		if (DataStructType->IsA<UUserDefinedStruct>())
+		{
+			NewDataInstance = CreateDataByType(FMoverUserDefinedDataStruct::StaticStruct());
+			static_cast<FMoverUserDefinedDataStruct*>(NewDataInstance.Get())->StructInstance.InitializeAs(DataStructType);
+		}
+		else
+		{
+			NewDataInstance = CreateDataByType(DataStructType);
+		}
+
 		DataArray.Add(NewDataInstance);
 		return NewDataInstance.Get();
 	}
@@ -396,7 +419,7 @@ FMoverDataStructBase* FMoverDataCollection::AddDataByType(const UScriptStruct* D
 
 void FMoverDataCollection::AddOrOverwriteData(const TSharedPtr<FMoverDataStructBase> DataInstance)
 {
-	RemoveDataByType(DataInstance->GetScriptStruct());
+	RemoveDataByType(DataInstance->GetDataScriptStruct());
 	DataArray.Add(DataInstance);
 }
 
@@ -405,12 +428,13 @@ void FMoverDataCollection::AddDataByCopy(const FMoverDataStructBase* DataInstanc
 {
 	check(DataInstanceToCopy);
 
-	const UScriptStruct* DataScriptStruct = DataInstanceToCopy->GetScriptStruct();
+	const UScriptStruct* TypeToMatch = DataInstanceToCopy->GetDataScriptStruct();
 
-	if (FMoverDataStructBase* ExistingMatchingData = FindDataByType(DataScriptStruct))
+	if (FMoverDataStructBase* ExistingMatchingData = FindDataByType(TypeToMatch))
 	{
-		DataScriptStruct->CopyScriptStruct(ExistingMatchingData, DataInstanceToCopy, 1);
-
+		// Note that we've matched based on the "data" type but we're copying the top-level type (a FMoverDataStructBase subtype)
+		const UScriptStruct* MoverDataTypeToCopy = DataInstanceToCopy->GetScriptStruct();
+		MoverDataTypeToCopy->CopyScriptStruct(ExistingMatchingData, DataInstanceToCopy, 1);
 	}
 	else
 	{
@@ -423,7 +447,7 @@ FMoverDataStructBase* FMoverDataCollection::FindDataByType(const UScriptStruct*
 {
 	for (const TSharedPtr<FMoverDataStructBase>& Data : DataArray)
 	{
-		UStruct* CandidateStruct = Data->GetScriptStruct();
+		const UStruct* CandidateStruct = Data->GetDataScriptStruct();
 		while (CandidateStruct)
 		{
 			if (DataStructType == CandidateStruct)
@@ -456,7 +480,7 @@ bool FMoverDataCollection::RemoveDataByType(const UScriptStruct* DataStructType)
 
 	for (int32 i=0; i < DataArray.Num() && IndexToRemove < 0; ++i)
 	{
-		UStruct* CandidateStruct = DataArray[i]->GetScriptStruct();
+		const UStruct* CandidateStruct = DataArray[i]->GetDataScriptStruct();
 		while (CandidateStruct)
 		{
 			if (DataStructType == CandidateStruct)
@@ -479,7 +503,7 @@ bool FMoverDataCollection::RemoveDataByType(const UScriptStruct* DataStructType)
 }
 
 /*static*/
-void FMoverDataCollection::NetSerializeDataArray(FArchive& Ar, TArray<TSharedPtr<FMoverDataStructBase>>& DataArray)
+void FMoverDataCollection::NetSerializeDataArray(FArchive& Ar, UPackageMap* Map, TArray<TSharedPtr<FMoverDataStructBase>>& DataArray)
 {
 	uint8 NumDataStructsToSerialize;
 	if (Ar.IsSaving())
@@ -541,8 +565,15 @@ void FMoverDataCollection::NetSerializeDataArray(FArchive& Ar, TArray<TSharedPtr
 					}
 				}
 
-				bool bIgnoredSuccess = false;
-				DataArray[i]->NetSerialize(Ar, nullptr, bIgnoredSuccess);
+				bool bArrayElementSuccess = false;
+				DataArray[i]->NetSerialize(Ar, Map, bArrayElementSuccess);
+
+				if (!bArrayElementSuccess)
+				{
+					UE_LOG(LogMover, Error, TEXT("FMoverDataCollection::NetSerialize: Failed to serialize ScriptStruct %s"), *ScriptStruct->GetName());
+					Ar.SetError();
+					break;
+				}
 			}
 			else
 			{
@@ -596,10 +627,25 @@ DEFINE_FUNCTION(UMoverDataCollectionLibrary::execK2_AddDataToCollection)
 	{
 		P_NATIVE_BEGIN;
 
-		if (ensure(SourceStructProp && SourceStructProp->Struct && SourceStructProp->Struct->IsChildOf(FMoverDataStructBase::StaticStruct()) && SourceDataAsRawPtr))
+		if (ensure(SourceStructProp->Struct))
 		{
-			FMoverDataStructBase* SourceDataAsBasePtr = reinterpret_cast<FMoverDataStructBase*>(SourceDataAsRawPtr);
-			TargetCollection.AddDataByCopy(SourceDataAsBasePtr);
+			// User-defined struct type support: we wrap an instance inside a FMoverUserDefinedDataStruct
+			if (SourceStructProp->Struct->IsA<UUserDefinedStruct>())
+			{
+				FMoverUserDefinedDataStruct UserDefinedDataWrapper;
+				UserDefinedDataWrapper.StructInstance.InitializeAs(SourceStructProp->Struct, (uint8*)SourceDataAsRawPtr);
+
+				TargetCollection.AddDataByCopy(&UserDefinedDataWrapper);
+			}
+			else if (SourceStructProp->Struct->IsChildOf(FMoverDataStructBase::StaticStruct()))
+			{
+				FMoverDataStructBase* SourceDataAsBasePtr = reinterpret_cast<FMoverDataStructBase*>(SourceDataAsRawPtr);
+				TargetCollection.AddDataByCopy(SourceDataAsBasePtr);
+			}
+			else
+			{
+				UE_LOG(LogMover, Warning, TEXT("AddDataToCollection: invalid struct type submitted: %s"), *SourceStructProp->Struct->GetName());
+			}
 		}
 
 		P_NATIVE_END;
@@ -639,11 +685,12 @@ DEFINE_FUNCTION(UMoverDataCollectionLibrary::execK2_GetDataFromCollection)
 
 		FBlueprintCoreDelegates::ThrowScriptException(P_THIS, Stack, ExceptionInfo);
 	}
-	else if (!TargetStructProp->Struct || !TargetStructProp->Struct->IsChildOf(FMoverDataStructBase::StaticStruct()))
+	else if (!TargetStructProp->Struct || 
+				(!TargetStructProp->Struct->IsChildOf(FMoverDataStructBase::StaticStruct()) && !TargetStructProp->Struct->IsA<UUserDefinedStruct>()))
 	{
 		FBlueprintExceptionInfo ExceptionInfo(
 			EBlueprintExceptionType::AbortExecution,
-			LOCTEXT("MoverDataCollection_GetDataFromCollection_BadType", "TargetAsRawBytes is not a valid type. Must be a child of FMoverDataStructBase.")
+			LOCTEXT("MoverDataCollection_GetDataFromCollection_BadType", "TargetAsRawBytes is not a valid type. Must be a child of FMoverDataStructBase or a User-Defined Struct type.")
 		);
 
 		FBlueprintCoreDelegates::ThrowScriptException(P_THIS, Stack, ExceptionInfo);
@@ -652,10 +699,23 @@ DEFINE_FUNCTION(UMoverDataCollectionLibrary::execK2_GetDataFromCollection)
 	{
 		P_NATIVE_BEGIN;
 
-		if (FMoverDataStructBase* FoundDataInstance = TargetCollection.FindDataByType(TargetStructProp->Struct))
+		if (TargetStructProp->Struct->IsA<UUserDefinedStruct>())
+		{
+			if (FMoverDataStructBase* FoundDataInstance = TargetCollection.FindDataByType(TargetStructProp->Struct))
+			{
+				// User-defined struct instances are wrapped in a FMoverUserDefinedDataStruct, so we need to extract the instance memory from inside it
+				FMoverUserDefinedDataStruct* FoundBPDataInstance = static_cast<FMoverUserDefinedDataStruct*>(FoundDataInstance);
+				TargetStructProp->Struct->CopyScriptStruct(TargetDataAsRawPtr, FoundBPDataInstance->StructInstance.GetMemory());
+				DidSucceed = true;
+			}
+		}
+		else
 		{
-			TargetStructProp->Struct->CopyScriptStruct(TargetDataAsRawPtr, FoundDataInstance);
-			DidSucceed = true;
+			if (FMoverDataStructBase* FoundDataInstance = TargetCollection.FindDataByType(TargetStructProp->Struct))
+			{
+				TargetStructProp->Struct->CopyScriptStruct(TargetDataAsRawPtr, FoundDataInstance);
+				DidSucceed = true;
+			}
 		}
 
 		P_NATIVE_END;
diff --git a/Engine/Plugins/Experimental/Mover/Source/Mover/Private/UserDefinedStructSupport.cpp b/Engine/Plugins/Experimental/Mover/Source/Mover/Private/UserDefinedStructSupport.cpp
new file mode 100644
--- /dev/null
+++ b/Engine/Plugins/Experimental/Mover/Source/Mover/Private/UserDefinedStructSupport.cpp
@@ -0,0 +1,95 @@
+// Copyright Epic Games, Inc. All Rights Reserved.
+
+
+#include "UserDefinedStructSupport.h"
+#include "StructUtils/UserDefinedStruct.h"
+#include "MoverLog.h"
+
+#include UE_INLINE_GENERATED_CPP_BY_NAME(UserDefinedStructSupport)
+
+#define LOCTEXT_NAMESPACE "MoverUDSInstances"
+
+
+// TODO: Consider different rules for interpolation/merging/reconciliation checks. 
+// This could be accomplished via cvars / Mover settings / per-type metadata , etc.
+
+bool FMoverUserDefinedDataStruct::ShouldReconcile(const FMoverDataStructBase& AuthorityState) const
+{
+	const FMoverUserDefinedDataStruct& TypedAuthority = static_cast<const FMoverUserDefinedDataStruct&>(AuthorityState);
+
+	check(TypedAuthority.StructInstance.GetScriptStruct() == this->StructInstance.GetScriptStruct());
+
+	return !StructInstance.Identical(&TypedAuthority.StructInstance, EPropertyPortFlags::PPF_DeepComparison);
+}
+
+void FMoverUserDefinedDataStruct::Interpolate(const FMoverDataStructBase& From, const FMoverDataStructBase& To, float LerpFactor)
+{
+	const FMoverUserDefinedDataStruct& PrimarySource = static_cast<const FMoverUserDefinedDataStruct&>((LerpFactor < 0.5f) ? From : To);
+
+	// copy all properties from the heaviest-weighted source rather than interpolate
+	StructInstance = PrimarySource.StructInstance;
+}
+
+void FMoverUserDefinedDataStruct::Merge(const FMoverDataStructBase& From)
+{
+	const FMoverUserDefinedDataStruct& TypedFrom = static_cast<const FMoverUserDefinedDataStruct&>(From);
+
+	check(TypedFrom.StructInstance.GetScriptStruct() == this->StructInstance.GetScriptStruct());
+
+	// Merging is typically only done for inputs. Let's make the assumption that boolean inputs should be OR'd so we never miss any digital inputs.
+
+	if (const UScriptStruct* UdsScriptStruct = TypedFrom.StructInstance.GetScriptStruct())
+	{
+		uint8* ThisInstanceMemory = StructInstance.GetMutableMemory();
+		const uint8* FromInstanceMemory = TypedFrom.StructInstance.GetMemory();
+
+		for (TFieldIterator<FBoolProperty> BoolProperty(UdsScriptStruct); BoolProperty; ++BoolProperty)
+		{
+			bool bMergedBool = BoolProperty->GetPropertyValue(ThisInstanceMemory);
+
+			if (!bMergedBool)
+			{
+				bMergedBool |= BoolProperty->GetPropertyValue(FromInstanceMemory);
+
+				if (bMergedBool)
+				{
+					BoolProperty->SetPropertyValue(ThisInstanceMemory, bMergedBool);
+				}
+			}
+		}
+	}
+}
+
+FMoverDataStructBase* FMoverUserDefinedDataStruct::Clone() const
+{
+	FMoverUserDefinedDataStruct* CopyPtr = new FMoverUserDefinedDataStruct(*this);
+	return CopyPtr;
+}
+
+bool FMoverUserDefinedDataStruct::NetSerialize(FArchive& Ar, UPackageMap* Map, bool& bOutSuccess)
+{
+	bool bSuperSuccess, bStructSuccess;
+
+	Super::NetSerialize(Ar, Map, bSuperSuccess);
+	StructInstance.NetSerialize(Ar, Map, bStructSuccess);
+
+	bOutSuccess = bSuperSuccess && bStructSuccess;
+
+	return true;
+}
+
+
+void FMoverUserDefinedDataStruct::ToString(FAnsiStringBuilderBase& Out) const
+{
+	Super::ToString(Out);
+
+	// TODO: add property-wise concatenated string output
+}
+
+const UScriptStruct* FMoverUserDefinedDataStruct::GetDataScriptStruct() const
+{
+	return StructInstance.GetScriptStruct();
+}
+
+
+#undef LOCTEXT_NAMESPACE
diff --git a/Engine/Plugins/Experimental/Mover/Source/Mover/Public/MoverSimulationTypes.h b/Engine/Plugins/Experimental/Mover/Source/Mover/Public/MoverSimulationTypes.h
--- a/Engine/Plugins/Experimental/Mover/Source/Mover/Public/MoverSimulationTypes.h
+++ b/Engine/Plugins/Experimental/Mover/Source/Mover/Public/MoverSimulationTypes.h
@@ -164,7 +164,7 @@ struct FMoverInputCmdContext
 	void NetSerialize(const FNetSerializeParams& P)
 	{
 		bool bIgnoredResult(false);
-		InputCollection.NetSerialize(P.Ar, nullptr, bIgnoredResult);
+		InputCollection.NetSerialize(P.Ar, P.Map, bIgnoredResult);
 	}
 
 	void ToString(FAnsiStringBuilderBase& Out) const
@@ -238,7 +238,7 @@ public:
 		MovementModifiers.NetSerialize(P.Ar);
 
 		bool bIgnoredResult(false);
-		SyncStateCollection.NetSerialize(P.Ar, nullptr, bIgnoredResult);
+		SyncStateCollection.NetSerialize(P.Ar, P.Map, bIgnoredResult);
 	}
 
 	void ToString(FAnsiStringBuilderBase& Out) const
@@ -332,7 +332,7 @@ public:
 	void NetSerialize(const FNetSerializeParams& P)
 	{
 		bool bIgnoredResult(false);
-		AuxStateCollection.NetSerialize(P.Ar, nullptr, bIgnoredResult);
+		AuxStateCollection.NetSerialize(P.Ar, P.Map, bIgnoredResult);
 	}
 
 	void ToString(FAnsiStringBuilderBase& Out) const
diff --git a/Engine/Plugins/Experimental/Mover/Source/Mover/Public/MoverTypes.h b/Engine/Plugins/Experimental/Mover/Source/Mover/Public/MoverTypes.h
--- a/Engine/Plugins/Experimental/Mover/Source/Mover/Public/MoverTypes.h
+++ b/Engine/Plugins/Experimental/Mover/Source/Mover/Public/MoverTypes.h
@@ -162,6 +162,10 @@ struct FMoverDataStructBase
 	 * Override for: Types that compose INPUT data for physics-based movement
 	 */
 	virtual void Decay(float DecayAmount) {};
+
+
+	/** Used to differentiate between this type, and a proxy data type. RARELY overridden by derived types. */
+	UE_API virtual const UScriptStruct* GetDataScriptStruct() const;
 };
 
 template<>
@@ -320,7 +324,7 @@ protected:
 	static UE_API TSharedPtr<FMoverDataStructBase> CreateDataByType(const UScriptStruct* DataStructType);
 	
 	/** Helper function for serializing array of data */
-	static UE_API void NetSerializeDataArray(FArchive& Ar, TArray<TSharedPtr<FMoverDataStructBase>>& DataArray);
+	static UE_API void NetSerializeDataArray(FArchive& Ar, UPackageMap* Map, TArray<TSharedPtr<FMoverDataStructBase>>& DataArray);
 
 	/** All data in this collection */
 	TArray< TSharedPtr<FMoverDataStructBase> > DataArray;
diff --git a/Engine/Plugins/Experimental/Mover/Source/Mover/Public/UserDefinedStructSupport.h b/Engine/Plugins/Experimental/Mover/Source/Mover/Public/UserDefinedStructSupport.h
new file mode 100644
--- /dev/null
+++ b/Engine/Plugins/Experimental/Mover/Source/Mover/Public/UserDefinedStructSupport.h
@@ -0,0 +1,49 @@
+// Copyright Epic Games, Inc. All Rights Reserved.
+
+#pragma once
+
+#include "MoverTypes.h"
+#include "StructUtils/InstancedStruct.h"
+#include "UserDefinedStructSupport.generated.h"
+
+class UUserDefinedStruct;
+
+
+#define UE_API MOVER_API
+
+/** Wrapper class that's used to add User-Defined Struct instances to Mover Data Collections (input or state).
+ * This allows devs to add custom data to inputs and/or state without requiring native code.
+ * Note that these are typically less efficient than natively-defined structs, and the logic of operations
+ * like interpolation, merging, and serialization may be simplistic for a project's needs.
+ * At present:
+ * - any differences between any struct contents will trigger reconciliation, even small floating point number differences
+ * - only boolean values can be merged
+ * - interpolation will take the entire struct instance from the highest weight frame
+ */
+USTRUCT()
+struct FMoverUserDefinedDataStruct : public FMoverDataStructBase
+{
+	GENERATED_BODY()
+
+public:
+	UPROPERTY()
+	FInstancedStruct StructInstance;
+
+	// Implementation of FMoverDataStructBase
+	virtual bool ShouldReconcile(const FMoverDataStructBase& AuthorityState) const override;
+	virtual void Interpolate(const FMoverDataStructBase& From, const FMoverDataStructBase& To, float LerpFactor) override;
+	virtual void Merge(const FMoverDataStructBase& From) override;
+	virtual FMoverDataStructBase* Clone() const override;
+	virtual bool NetSerialize(FArchive& Ar, UPackageMap* Map, bool& bOutSuccess) override;
+	virtual void ToString(FAnsiStringBuilderBase& Out) const override;
+
+	// Note: this is the FMoverUserDefinedDataStruct type, NOT the User-Defined Struct type
+	virtual UScriptStruct* GetScriptStruct() const override { return StaticStruct(); }
+	// This returns the User-Defined Struct type
+	virtual const UScriptStruct* GetDataScriptStruct() const override;
+
+
+};
+
+
+#undef UE_API
diff --git a/Engine/Plugins/Runtime/NetworkPrediction/Source/NetworkPrediction/Private/NetworkPredictionReplicationProxy.cpp b/Engine/Plugins/Runtime/NetworkPrediction/Source/NetworkPrediction/Private/NetworkPredictionReplicationProxy.cpp
--- a/Engine/Plugins/Runtime/NetworkPrediction/Source/NetworkPrediction/Private/NetworkPredictionReplicationProxy.cpp
+++ b/Engine/Plugins/Runtime/NetworkPrediction/Source/NetworkPrediction/Private/NetworkPredictionReplicationProxy.cpp
@@ -20,9 +20,12 @@ bool FReplicationProxy::NetSerialize(FArchive& Ar, class UPackageMap* Map, bool&
 {
 	if (npEnsureMsgf(NetSerializeFunc, TEXT("NetSerializeFunc not set for FReplicationProxy %d"), ReplicationTarget))
 	{
-		NetSerializeFunc(Ar);
+		FNetSerializeParams P(Ar, Map);
+		NetSerializeFunc(P);
+		bOutSuccess = !Ar.IsError();
 		return true;
 	}
+	bOutSuccess = false;
 	return true;
 }
 
diff --git a/Engine/Plugins/Runtime/NetworkPrediction/Source/NetworkPrediction/Public/NetworkPredictionReplicationProxy.h b/Engine/Plugins/Runtime/NetworkPrediction/Source/NetworkPrediction/Public/NetworkPredictionReplicationProxy.h
--- a/Engine/Plugins/Runtime/NetworkPrediction/Source/NetworkPrediction/Public/NetworkPredictionReplicationProxy.h
+++ b/Engine/Plugins/Runtime/NetworkPrediction/Source/NetworkPrediction/Public/NetworkPredictionReplicationProxy.h
@@ -27,8 +27,10 @@ inline FString LexToString(EReplicationProxyTarget A)
 // The parameters for NetSerialize that are passed around the system. Everything should use this, expecting to have to add more.
 struct FNetSerializeParams
 {
-	FNetSerializeParams(FArchive& InAr) : Ar(InAr) { }
+	FNetSerializeParams(FArchive& InAr, UPackageMap* InMap=nullptr) : Ar(InAr), Map(InMap) { }
+
 	FArchive& Ar;
+	UPackageMap* Map;
 };
 
 // Redirects NetSerialize to a dynamically set NetSerializeFunc.
