Mover: implementing rollback blackboard persistence policy to allow entries to auto-expire after 1 simulation frame.
Includes plumbing of sim time/frame during rollbacks and 'get' queries so we can reason about time-based expiration.

diff --git a/Engine/Plugins/Experimental/Mover/Source/Mover/Private/MoveLibrary/RollbackBlackboard.cpp b/Engine/Plugins/Experimental/Mover/Source/Mover/Private/MoveLibrary/RollbackBlackboard.cpp
--- a/Engine/Plugins/Experimental/Mover/Source/Mover/Private/MoveLibrary/RollbackBlackboard.cpp
+++ b/Engine/Plugins/Experimental/Mover/Source/Mover/Private/MoveLibrary/RollbackBlackboard.cpp
@@ -55,6 +55,41 @@ void URollbackBlackboard::BlackboardEntryBase::RollBack(uint32 NewPendingFrame)
 }
 
 
+
+bool URollbackBlackboard::BlackboardEntryBase::CanReadEntryAt(const EntryTimeStamp& ReaderTimeStamp, EEntryIndexType IndexType) const
+{
+	const uint32 TimestampIdx = IndexType == EEntryIndexType::External ? ExternalIdx : InternalIdx;
+
+	if (!Timestamps[TimestampIdx].IsValid())
+	{
+		return false;	// entry isn't initialized yet or never set
+	}
+
+	switch (Settings->PersistencePolicy)
+	{
+		case EBlackboardPersistencePolicy::NextFrameOnly:
+		{
+			// if value wasn't last set during the current or immediate prior sim frame, then it can't be read
+			if (!(ReaderTimeStamp.Frame == Timestamps[TimestampIdx].Frame ||
+				  ReaderTimeStamp.Frame == Timestamps[TimestampIdx].Frame+1))
+			{
+				return false;
+			}
+		}
+		break;
+
+		case EBlackboardPersistencePolicy::Forever:		// Always allow reading
+			break;
+
+		default:		// Allow reading of any not-yet-implemented policies
+			break;
+	}
+
+	return true;
+}
+
+
+
 void URollbackBlackboard::BeginSimulationFrame(const FMoverTimeStep& PendingTimeStep)
 {
 	check(!bIsSimulationInProgress && !bIsRollbackInProgress);
@@ -90,14 +125,19 @@ void URollbackBlackboard::BeginRollback()
 	bIsRollbackInProgress = true;
 }
 
-void URollbackBlackboard::RollBack(uint32 NewPendingFrame)
+void URollbackBlackboard::RollBack(const FMoverTimeStep& NewBaseTimeStep)
 {
 	check(!bIsSimulationInProgress && bIsRollbackInProgress);
 	
+	const EntryTimeStamp NewBaseTimeStamp((double)NewBaseTimeStep.BaseSimTimeMs, NewBaseTimeStep.ServerFrame);
+
 	for (const TPair<FName, TUniquePtr<BlackboardEntryBase>>& KVP : EntryMap)
 	{
-		KVP.Value->RollBack(NewPendingFrame);
+		KVP.Value->RollBack(NewBaseTimeStamp.Frame);
 	}
+
+	// As the rollback occurs, we need to pull back the timestamps to match
+	CurrentSimTimeStamp = InProgressSimTimeStamp = NewBaseTimeStamp;
 }
 
 void URollbackBlackboard::EndRollback()
diff --git a/Engine/Plugins/Experimental/Mover/Source/Mover/Private/MoverComponent.cpp b/Engine/Plugins/Experimental/Mover/Source/Mover/Private/MoverComponent.cpp
--- a/Engine/Plugins/Experimental/Mover/Source/Mover/Private/MoverComponent.cpp
+++ b/Engine/Plugins/Experimental/Mover/Source/Mover/Private/MoverComponent.cpp
@@ -1079,7 +1079,7 @@ void UMoverComponent::OnSimulationRollback(const FMoverSyncState* SyncState, con
 
 	RollbackBlackboard_InternalWrapper->BeginRollback();
 
-	RollbackBlackboard_InternalWrapper->RollBack(NewBaseTimeStep.ServerFrame);
+	RollbackBlackboard_InternalWrapper->RollBack(NewBaseTimeStep);
 
 	ModeFSM->OnSimulationRollback(SyncState, AuxState, NewBaseTimeStep);
 
diff --git a/Engine/Plugins/Experimental/Mover/Source/Mover/Public/MoveLibrary/RollbackBlackboard.h b/Engine/Plugins/Experimental/Mover/Source/Mover/Public/MoveLibrary/RollbackBlackboard.h
--- a/Engine/Plugins/Experimental/Mover/Source/Mover/Public/MoveLibrary/RollbackBlackboard.h
+++ b/Engine/Plugins/Experimental/Mover/Source/Mover/Public/MoveLibrary/RollbackBlackboard.h
@@ -47,14 +47,16 @@ enum class EBlackboardSizingPolicy : uint8
 
 	/** Buffer size matches the backend simulation's history size, determined by ticking rate and history settings.
 	 * A simulation running a fixed 30 fps with a 1 second history will need a buffer size of 30. For variable rate simulations, buffer size cam only be estimated. 
+	 * NOT YET IMPLEMENTED
 	 */
-	FixedBackendBufferSize,
+	FixedBackendBufferSize UMETA(Hidden),
 
 	/** Minimizes data size for blackboard entries where thread contention and rolling back shouldn't change the value.
 	 * Example: a statistic tracking the maximum time spent in a particular mode's simulation tick. 
 	 * Note this will still have multiple elements for asynchronous simulations, to prevent thread contention over the entry.
+	 * NOT YET IMPLEMENTED
 	 */
-	SingleEntry,
+	SingleEntry UMETA(Hidden),
 };
 
 /**
@@ -63,11 +65,14 @@ enum class EBlackboardSizingPolicy : uint8
 UENUM()
 enum class EBlackboardPersistencePolicy : uint8
 {
+	/** Any edits to the entry remain valid and readable indefinitely */
 	Forever,
 
+	/** Any edits remain valid until the end of the next simulation frame. If an edit occurs during frame N, the entry will be readable for the remainder of frame N and frame N+1. */
 	NextFrameOnly,
 
-	LimitedTime,
+	/** NOT YET IMPLEMENTED. Any edits remain valid until a specified amount of simulation time has passed. */
+	LimitedTime UMETA(Hidden),
 };
 
 
@@ -80,11 +85,11 @@ enum class EBlackboardRollbackPolicy : uint8
 	/** Entries occurring after the rollback time are effectively erased and no longer retrievable. New entries may be written during resimulation. */
 	InvalidatedOnRollback,
 
-	/** Entries are not invalidated during rollback, and may be read/written during resim. If left alone, they will persist. */
-	WritableDuringResimulation,
+	/** Entries are not invalidated during rollback, and may be read/written during resim. If left alone, they will persist. NOT YET IMPLEMENTED. */
+	WritableDuringResimulation UMETA(Hidden),
 
-	/** Entries are not invalidated during rollback, and may be read during resim. But writes cannot be performed during resim. */
-	LockedDuringResimulation,
+	/** Entries are not invalidated during rollback, and may be read during resim. But writes cannot be performed during resim. NOT YET IMPLEMENTED. */
+	LockedDuringResimulation UMETA(Hidden),
 };
 
 
@@ -114,10 +119,13 @@ private:
 
 	struct EntryTimeStamp
 	{
+		EntryTimeStamp() {}
+		EntryTimeStamp(double InTimeMs, uint32 InFrame) : TimeMs(InTimeMs), Frame(InFrame) {}
+
 		double TimeMs = -1.0;
 		uint32 Frame = 0;
 
-		bool IsValid()
+		bool IsValid() const
 		{
 			return TimeMs >= 0.0;
 		}
@@ -129,6 +137,12 @@ private:
 		}
 	};
 
+	enum EEntryIndexType : uint8
+	{
+		External,	// entry index for outside systems to read
+		Internal	// entry index for in-simulation system to read/write
+	};
+
 	// untyped base, so we can use entries of multiple types in a single-typed container
 	struct BlackboardEntryBase
 	{
@@ -142,7 +156,10 @@ private:
 		// Note that NewPendingFrame means the frame that will now be re-simulated. So any existing entries with a Frame >= NewPendingFrame can be invalidated.
 		void RollBack(uint32 NewPendingFrame);
 
+
 	protected:
+		bool CanReadEntryAt(const EntryTimeStamp& ReaderTimeStamp, EEntryIndexType IndexType) const;
+
 		uint32 ComputeBufferSize(EBlackboardSizingPolicy InSizingPolicy, uint32 FixedBufferSize = -1);
 
 		const EntrySettings* Settings = nullptr;
@@ -171,7 +188,7 @@ private:
 
 			Timestamps = TCircularBuffer<EntryTimeStamp>(BufferSize, EntryTimeStamp());
 
-			// JAH TODO: Consider making the initial object required, so we can also initilize where in the buffer we point to
+			// JAH TODO: Consider making the initial object required, so we can also initialize where in the buffer we point to
 			if (OptionalInitialObj)
 			{
 				EntryBuffer = TCircularBuffer<EntryT>(BufferSize, *OptionalInitialObj);
@@ -185,47 +202,47 @@ private:
 		}
 
 
-		bool TryGetEntryValue(EntryT& ValueOut)
+		bool TryGetEntryValue(const EntryTimeStamp& CurrentExternalTime, EntryT& ValueOut)
 		{
-			if (!Timestamps[ExternalIdx].IsValid())
+			if (CanReadEntryAt(CurrentExternalTime, EEntryIndexType::External))
 			{
-				return false;	// Probably not initialized yet
+				ValueOut = EntryBuffer[ExternalIdx];
+				return true;
 			}
 
-			ValueOut = EntryBuffer[ExternalIdx];
-			return true;
+			return false;
 		}
 
 		// TODO: Consider whether it's worthwhile to have this function... dangerous to let external folks hold onto references to internal storage, but may allow more efficient editing
-		EntryT* TryGetEntryRef()
+		EntryT* TryGetEntryRef(const EntryTimeStamp& CurrentExternalTime)
 		{
-			if (!Timestamps[ExternalIdx].IsValid())
-			{
-				return nullptr;	// Probably not initialized yet
+			if (CanReadEntryAt(CurrentExternalTime, EEntryIndexType::External))
+			{ 
+				return &EntryBuffer[ExternalIdx];
 			}
 
-			return &EntryBuffer[ExternalIdx];
+			return nullptr;
 		}
 
-		bool TryGetEntryValue_Internal(EntryT& ValueOut)
+		bool TryGetEntryValue_Internal(const EntryTimeStamp& InternalTime, EntryT& ValueOut)
 		{
-			if (!Timestamps[InternalIdx].IsValid())
+			if (CanReadEntryAt(InternalTime, EEntryIndexType::Internal))
 			{
-				return false;	// Probably not initialized yet
+				ValueOut = EntryBuffer[InternalIdx];
+				return true;
 			}
 
-			ValueOut = EntryBuffer[InternalIdx];
-			return true;
+			return false;
 		}
 
-		EntryT* TryGetEntryRef_Internal()
+		EntryT* TryGetEntryRef_Internal(const EntryTimeStamp& InternalTime)
 		{
-			if (!Timestamps[InternalIdx].IsValid())
+			if (CanReadEntryAt(InternalTime, EEntryIndexType::Internal))
 			{
-				return nullptr;	// Probably not initialized yet
+				return &EntryBuffer[InternalIdx];
 			}
 
-			return &EntryBuffer[InternalIdx];
+			return nullptr;
 		}
 
 		void SetEntryValue_Internal(const EntryT& Value, const EntryTimeStamp& TimeStamp)
@@ -260,13 +277,18 @@ public:
 		return true;
 	}
 
+	bool HasEntry(FName EntryName) const
+	{
+		return EntryMap.Contains(EntryName);
+	}
+
 	/** Store object by a named key, overwriting any existing object */
 	template<typename EntryT>
 	bool TrySet(FName ObjName, EntryT& Obj)
 	{
 		// TODO: Verify this isn't happening mid-sim on the sim thread (_Internal should be used instead)
 
-		if (BlackboardEntry<EntryT>* Entry = FindEntry(ObjName))
+		if (BlackboardEntry<EntryT>* Entry = static_cast<BlackboardEntry<EntryT>*>(FindEntry(ObjName)))
 		{
 			Entry->SetEntryValue(Obj);
 			return true;
@@ -280,9 +302,9 @@ public:
 	template<typename EntryT>
 	bool TryGet(FName ObjName, EntryT& OutFoundValue) const
 	{
-		if (BlackboardEntry<EntryT>* Entry = FindEntry(ObjName))
+		if (BlackboardEntry<EntryT>* Entry = static_cast<BlackboardEntry<EntryT>*>(FindEntry(ObjName)))
 		{
-			return Entry->TryGetEntryValue(OutFoundValue);
+			return Entry->TryGetEntryValue(CurrentSimTimeStamp, OutFoundValue);
 		}
 
 		return false;
@@ -294,9 +316,9 @@ public:
 	template<typename EntryT>
 	EntryT* TryGetRef(FName ObjName) const
 	{
-		if (BlackboardEntry<EntryT>* Entry = FindEntry(ObjName))
+		if (BlackboardEntry<EntryT>* Entry = static_cast<BlackboardEntry<EntryT>*>(FindEntry(ObjName)))
 		{
-			return Entry->TryGetEntryRef();
+			return Entry->TryGetEntryRef(CurrentSimTimeStamp);
 		}
 
 		return nullptr;
@@ -328,7 +350,7 @@ private:
 
 		if (BlackboardEntry<EntryT>* Entry = static_cast<BlackboardEntry<EntryT>*>(FindEntry(ObjName)))
 		{
-			return Entry->TryGetEntryValue_Internal(OutFoundValue);
+			return Entry->TryGetEntryValue_Internal(InProgressSimTimeStamp, OutFoundValue);
 		}
 
 		return false;
@@ -342,7 +364,7 @@ private:
 
 		if (BlackboardEntry<EntryT>* Entry = static_cast<BlackboardEntry<EntryT>*>(FindEntry(ObjName)))
 		{
-			return Entry->TryGetEntryRef_Internal();
+			return Entry->TryGetEntryRef_Internal(InProgressSimTimeStamp);
 		}
 
 		return nullptr;
@@ -366,7 +388,7 @@ private:
 	void EndSimulationFrame();
 
 	void BeginRollback();
-	void RollBack(uint32 NewPendingFrame);
+	void RollBack(const FMoverTimeStep& NewBaseTimeStep);	// NewBaseTimeStep represents the current time and the PENDING frame that's next to be resimulated
 	void EndRollback();
 
 
@@ -382,8 +404,8 @@ private:
 	uint32 InRollbackThreadId;
 
 
-	EntryTimeStamp CurrentSimTimeStamp;
-	EntryTimeStamp InProgressSimTimeStamp;	// only valid during a simulation step
+	EntryTimeStamp CurrentSimTimeStamp;	// this is the "committed" simulation time, and will lag behind InProgressSimTimeStamp while in the middle of a simulation frame
+	EntryTimeStamp InProgressSimTimeStamp;	// this is the sim time that's actively being used mid-simulation. It is only valid during a simulation step.
 
 	friend class URollbackBlackboard_InternalWrapper;	// expose non-public API to this wrapper class
 };
@@ -405,6 +427,17 @@ public:
 		Blackboard = &InBlackboard;
 	}
 
+	template<typename EntryT>
+	bool CreateEntry(FName EntryName, const URollbackBlackboard::EntrySettings& InSettings)
+	{
+		return Blackboard->CreateEntry<EntryT>(EntryName, InSettings);
+	}
+
+	bool HasEntry(FName EntryName) const
+	{
+		return Blackboard->HasEntry(EntryName);
+	}
+
 	template<typename EntryT>
 	bool TrySet(FName ObjName, EntryT& Obj)
 	{
@@ -439,9 +472,9 @@ public:
 		Blackboard->BeginRollback();
 	}
 
-	void RollBack(uint32 NewPendingFrame)
+	void RollBack(const FMoverTimeStep& NewBaseTimeStep)
 	{
-		Blackboard->RollBack(NewPendingFrame);
+		Blackboard->RollBack(NewBaseTimeStep);
 	}
 
 	void EndRollback()
