Mover: adding method of canceling movement features by matching gameplay tags

diff --git a/Engine/Plugins/Experimental/Mover/Source/Mover/Private/LayeredMove.cpp b/Engine/Plugins/Experimental/Mover/Source/Mover/Private/LayeredMove.cpp
--- a/Engine/Plugins/Experimental/Mover/Source/Mover/Private/LayeredMove.cpp
+++ b/Engine/Plugins/Experimental/Mover/Source/Mover/Private/LayeredMove.cpp
@@ -135,6 +135,14 @@ void FLayeredMoveGroup::QueueLayeredMove(TSharedPtr<FLayeredMoveBase> Move)
 	}
 }
 
+void FLayeredMoveGroup::CancelMovesByTag(FGameplayTag Tag, bool bRequireExactMatch)
+{
+
+	// Schedule a tag cancellation request, to be handled during simulation
+	TagCancellationRequests.Add(TPair<FGameplayTag, bool>(Tag, bRequireExactMatch));
+
+}
+
 TArray<TSharedPtr<FLayeredMoveBase>> FLayeredMoveGroup::GenerateActiveMoves(const FMoverTimeStep& TimeStep, const UMoverComponent* MoverComp, UMoverBlackboard* SimBlackboard)
 {
 	const double SimStartTimeMs		= TimeStep.BaseSimTimeMs;
@@ -225,6 +233,8 @@ FLayeredMoveGroup& FLayeredMoveGroup::operator=(const FLayeredMoveGroup& Other)
 	{
 		CopyLayeredMoveArray(ActiveLayeredMoves, Other.ActiveLayeredMoves);
 		CopyLayeredMoveArray(QueuedLayeredMoves, Other.QueuedLayeredMoves);
+
+		TagCancellationRequests = Other.TagCancellationRequests;
 	}
 
 	return *this;
@@ -374,6 +384,43 @@ void FLayeredMoveGroup::FlushMoveArrays(const UMoverComponent* MoverComp, UMover
 
 	bool bResidualVelocityOverridden = false;
 	bool bClampVelocityOverridden = false;
+	
+	// Process any cancellations
+	{
+		for (TPair<FGameplayTag, bool> CancelRequest : TagCancellationRequests)
+		{
+			const FGameplayTag TagToMatch = CancelRequest.Key;
+			const bool bRequireExactMatch = CancelRequest.Value;
+
+			QueuedLayeredMoves.RemoveAll([TagToMatch, bRequireExactMatch](const TSharedPtr<FLayeredMoveBase>& Move)
+				{
+					return (Move.IsValid() && Move->HasGameplayTag(TagToMatch, bRequireExactMatch));
+				});
+
+			ActiveLayeredMoves.RemoveAll([MoverComp, SimBlackboard, CurrentSimTimeMs, &bResidualVelocityOverridden, &bClampVelocityOverridden, bIsAsync, TagToMatch, bRequireExactMatch, this] (const TSharedPtr<FLayeredMoveBase>& Move)
+				{
+					if (Move.IsValid() && Move->HasGameplayTag(TagToMatch, bRequireExactMatch))
+					{
+						GatherResidualVelocitySettings(Move, bResidualVelocityOverridden, bClampVelocityOverridden);
+						if (bIsAsync)
+						{
+							Move->EndMove_Async(SimBlackboard, CurrentSimTimeMs);
+						}
+						else
+						{
+							Move->EndMove(MoverComp, SimBlackboard, CurrentSimTimeMs);
+						}
+						return true;
+					}
+
+					return false;
+				});
+		}
+
+		TagCancellationRequests.Empty();
+	}
+
+	
 	// Remove any finished moves
 	ActiveLayeredMoves.RemoveAll([MoverComp, SimBlackboard, CurrentSimTimeMs, &bResidualVelocityOverridden, &bClampVelocityOverridden, bIsAsync, this]
 		(const TSharedPtr<FLayeredMoveBase>& Move)
@@ -601,5 +648,6 @@ void FLayeredMoveGroup::Reset()
 	ResetResidualVelocity();
 	QueuedLayeredMoves.Empty();
 	ActiveLayeredMoves.Empty();
+	TagCancellationRequests.Empty();
 }
 
diff --git a/Engine/Plugins/Experimental/Mover/Source/Mover/Private/LayeredMoveGroup.cpp b/Engine/Plugins/Experimental/Mover/Source/Mover/Private/LayeredMoveGroup.cpp
--- a/Engine/Plugins/Experimental/Mover/Source/Mover/Private/LayeredMoveGroup.cpp
+++ b/Engine/Plugins/Experimental/Mover/Source/Mover/Private/LayeredMoveGroup.cpp
@@ -28,6 +28,8 @@ FLayeredMoveInstanceGroup& FLayeredMoveInstanceGroup::operator=(const FLayeredMo
 		};
 		DeepCopyMovesFunc(Other.ActiveMoves, ActiveMoves);
 		DeepCopyMovesFunc(Other.QueuedMoves, QueuedMoves);
+
+		TagCancellationRequests = Other.TagCancellationRequests;
 	}
 
 	return *this;
@@ -180,6 +182,7 @@ void FLayeredMoveInstanceGroup::Reset()
 	ResetResidualVelocity();
 	QueuedMoves.Empty();
 	ActiveMoves.Empty();
+	TagCancellationRequests.Empty();
 }
 
 void FLayeredMoveInstanceGroup::PopulateMissingActiveMoveLogic(const TArray<TObjectPtr<ULayeredMoveLogic>>& RegisteredMoves)
@@ -258,13 +261,51 @@ const FLayeredMoveInstance* FLayeredMoveInstanceGroup::PrivateFindQueuedMove(con
 	return nullptr;
 }
 
+void FLayeredMoveInstanceGroup::CancelMovesByTag(FGameplayTag Tag, bool bRequireExactMatch)
+{
+
+	// Schedule a tag cancellation request, to be handled during simulation
+	TagCancellationRequests.Add(TPair<FGameplayTag, bool>(Tag, bRequireExactMatch));
+
+}
+
+
 void FLayeredMoveInstanceGroup::FlushMoveArrays(const FMoverTimeStep& TimeStep, UMoverBlackboard* SimBlackboard)
 {
+	bool bResidualVelocityOverridden = false;
+	bool bClampVelocityOverridden = false;
+
+	// Process any cancellations
+	{
+		for (TPair<FGameplayTag, bool> CancelRequest : TagCancellationRequests)
+		{
+			const FGameplayTag TagToMatch = CancelRequest.Key;
+			const bool bRequireExactMatch = CancelRequest.Value;
+
+			QueuedMoves.RemoveAll([TagToMatch, bRequireExactMatch](const TSharedPtr<FLayeredMoveInstance>& Move)
+				{
+					return (!Move.IsValid() || Move->HasGameplayTag(TagToMatch, bRequireExactMatch));
+				});
+
+			// Process completion of any active moves that being canceled
+			ActiveMoves.RemoveAll([&TimeStep, &SimBlackboard, &bResidualVelocityOverridden, &bClampVelocityOverridden, TagToMatch, bRequireExactMatch, this] (const TSharedPtr<FLayeredMoveInstance>& Move)
+				{
+					if (Move.IsValid() && Move->HasGameplayTag(TagToMatch, bRequireExactMatch))
+					{
+						ProcessFinishedMove(*Move, bResidualVelocityOverridden, bClampVelocityOverridden);
+						Move->EndMove(TimeStep, SimBlackboard);
+						return true;
+					}
+					return false;
+				});
+		}
+
+		TagCancellationRequests.Empty();
+	}
+
 	{
-		bool bResidualVelocityOverridden = false;
-		bool bClampVelocityOverridden = false;
 	
-		// Process completion of any moves that are finished
+		// Process completion of any active moves that are finished
 		ActiveMoves.RemoveAll([&TimeStep, &SimBlackboard, &bResidualVelocityOverridden, &bClampVelocityOverridden, this]
 			(const TSharedPtr<FLayeredMoveInstance>& Move)
 			{
@@ -278,6 +319,7 @@ void FLayeredMoveInstanceGroup::FlushMoveArrays(const FMoverTimeStep& TimeStep,
 			});	
 	}
 
+	// Begin any queued moves
 	for (TSharedPtr<FLayeredMoveInstance>& QueuedMove : QueuedMoves)
 	{
 		if (QueuedMove.IsValid())
diff --git a/Engine/Plugins/Experimental/Mover/Source/Mover/Private/MovementModeStateMachine.cpp b/Engine/Plugins/Experimental/Mover/Source/Mover/Private/MovementModeStateMachine.cpp
--- a/Engine/Plugins/Experimental/Mover/Source/Mover/Private/MovementModeStateMachine.cpp
+++ b/Engine/Plugins/Experimental/Mover/Source/Mover/Private/MovementModeStateMachine.cpp
@@ -204,6 +204,9 @@ void UMovementModeStateMachine::OnSimulationTick(USceneComponent* UpdatedCompone
 
 		SubTimeStep.StepMs = TimeStep.StepMs - TotalUsedMs;		// TODO: convert this to an overridable function that can support MaxStepTime, MaxIterations, etc.
 
+		// Process any cancellation requests first, so we can catch any queued features before they're activated
+		FlushTagCancellationsToSyncState(WorkingSubstepStartData.SyncState);
+
 		// Transfer any queued moves into the starting state. They'll be started during the move generation.
 		FlushQueuedMovesToGroup(WorkingSubstepStartData.SyncState.LayeredMoves);
 		OutputState.SyncState.LayeredMoves = WorkingSubstepStartData.SyncState.LayeredMoves;
@@ -624,6 +627,43 @@ const FMovementModifierBase* UMovementModeStateMachine::FindQueuedModifierByType
 	return nullptr;
 }
 
+void UMovementModeStateMachine::CancelFeaturesWithTag(FGameplayTag TagToCancel, bool bRequireExactMatch)
+{
+	// Cancel all matching queued movement features
+	{
+		FRWScopeLock QueueLock(ModifiersQueueLock, SLT_Write);
+		QueuedMovementModifiers.RemoveAll([TagToCancel, bRequireExactMatch] (const TSharedPtr<FMovementModifierBase>& Modifier)
+			{
+				return (Modifier.IsValid() && Modifier->HasGameplayTag(TagToCancel, bRequireExactMatch));
+			});
+	}
+
+	{
+		FRWScopeLock QueueLock(LayeredMoveQueueLock, SLT_Write);
+		QueuedLayeredMoves.RemoveAll([TagToCancel, bRequireExactMatch](const TSharedPtr<FLayeredMoveBase>& LayeredMove)
+			{
+				return (LayeredMove.IsValid() && LayeredMove->HasGameplayTag(TagToCancel, bRequireExactMatch));
+			});
+	}
+
+	{
+		FRWScopeLock QueueLock(LayeredMoveQueueLock, SLT_Write);
+		QueuedLayeredMoveInstances.RemoveAll([TagToCancel, bRequireExactMatch](const TSharedPtr<FLayeredMoveInstance>& LayeredMoveInstance)
+			{
+				return (LayeredMoveInstance.IsValid() && LayeredMoveInstance->HasGameplayTag(TagToCancel, bRequireExactMatch));
+			});
+	}
+
+	// TODO: also support cancellation of queued instant effects if they end up supporting gameplay tags
+
+	// Request cancelation of any matching ACTIVE movement features during the next simulation tick
+	{
+		FRWScopeLock QueueLock(ModifierCancelQueueLock, SLT_Write);
+		TagCancellationRequests.Add(TPair<FGameplayTag, bool>(TagToCancel, bRequireExactMatch));
+	}
+
+}
+
 
 void UMovementModeStateMachine::ConstructDefaultModes()
 {
@@ -732,6 +772,20 @@ void UMovementModeStateMachine::FlushModifierCancellationsToGroup(FMovementModif
 	ModifiersToCancel.Empty();
 }
 
+void UMovementModeStateMachine::FlushTagCancellationsToSyncState(FMoverSyncState& SyncState)
+{
+	FRWScopeLock QueueLock(TagCancellationRequestsLock, SLT_Write);
+
+	for (TPair<FGameplayTag, bool> TagCancelRequest : TagCancellationRequests)
+	{
+		SyncState.MovementModifiers.CancelModifiersByTag(TagCancelRequest.Key, TagCancelRequest.Value);
+		SyncState.LayeredMoves.CancelMovesByTag(TagCancelRequest.Key, TagCancelRequest.Value);
+		SyncState.LayeredMoveInstances.CancelMovesByTag(TagCancelRequest.Key, TagCancelRequest.Value);
+	}
+
+	TagCancellationRequests.Empty();
+}
+
 void UMovementModeStateMachine::RollbackModifiers(const FMoverSyncState* InvalidSyncState, const FMoverSyncState* SyncState, const FMoverAuxStateContext* InvalidAuxState, const FMoverAuxStateContext* AuxState)
 {
 	{
diff --git a/Engine/Plugins/Experimental/Mover/Source/Mover/Private/MovementModifier.cpp b/Engine/Plugins/Experimental/Mover/Source/Mover/Private/MovementModifier.cpp
--- a/Engine/Plugins/Experimental/Mover/Source/Mover/Private/MovementModifier.cpp
+++ b/Engine/Plugins/Experimental/Mover/Source/Mover/Private/MovementModifier.cpp
@@ -135,15 +135,30 @@ void FMovementModifierGroup::CancelModifierFromHandle(const FMovementModifierHan
 		}
 	}
 
-	for (TSharedPtr<FMovementModifierBase> QueuedModifier : QueuedModifiers)
+	QueuedModifiers.RemoveAll([HandleToCancel](const TSharedPtr<FMovementModifierBase>& Modifier)
+		{
+			return (!Modifier.IsValid() || HandleToCancel == Modifier->GetHandle());
+		});
+
+}
+
+void FMovementModifierGroup::CancelModifiersByTag(FGameplayTag Tag, bool bRequiresExactMatch)
+{
+	for (TSharedPtr<FMovementModifierBase> ActiveModifier : ActiveModifiers)
 	{
-		if (HandleToCancel == QueuedModifier->GetHandle())
+		if (ActiveModifier.IsValid() && ActiveModifier->HasGameplayTag(Tag, bRequiresExactMatch))
 		{
-			QueuedModifier->DurationMs = 0;
+			ActiveModifier->DurationMs = 0;
 		}
 	}
+
+	QueuedModifiers.RemoveAll([Tag, bRequiresExactMatch](const TSharedPtr<FMovementModifierBase>& Modifier)
+		{
+			return (!Modifier.IsValid() || Modifier->HasGameplayTag(Tag, bRequiresExactMatch));
+		});
 }
 
+
 TArray<TSharedPtr<FMovementModifierBase>> FMovementModifierGroup::GenerateActiveModifiers(UMoverComponent* MoverComp, const FMoverTimeStep& TimeStep, const FMoverSyncState& SyncState, const FMoverAuxStateContext& AuxState)
 {
 	FlushModifierArrays(MoverComp, TimeStep, SyncState, AuxState);
diff --git a/Engine/Plugins/Experimental/Mover/Source/Mover/Private/MoverComponent.cpp b/Engine/Plugins/Experimental/Mover/Source/Mover/Private/MoverComponent.cpp
--- a/Engine/Plugins/Experimental/Mover/Source/Mover/Private/MoverComponent.cpp
+++ b/Engine/Plugins/Experimental/Mover/Source/Mover/Private/MoverComponent.cpp
@@ -1654,6 +1654,13 @@ void UMoverComponent::CancelModifierFromHandle(FMovementModifierHandle ModifierH
 	ModeFSM->CancelModifierFromHandle(ModifierHandle);
 }
 
+
+void UMoverComponent::CancelFeaturesWithTag(FGameplayTag TagToCancel, bool bRequireExactMatch)
+{
+	ModeFSM->CancelFeaturesWithTag(TagToCancel, bRequireExactMatch);
+}
+
+
 void UMoverComponent::K2_QueueInstantMovementEffect(const int32& EffectAsRawData)
 {
 	// This will never be called, the exec version below will be hit instead
diff --git a/Engine/Plugins/Experimental/Mover/Source/Mover/Public/LayeredMove.h b/Engine/Plugins/Experimental/Mover/Source/Mover/Public/LayeredMove.h
--- a/Engine/Plugins/Experimental/Mover/Source/Mover/Public/LayeredMove.h
+++ b/Engine/Plugins/Experimental/Mover/Source/Mover/Public/LayeredMove.h
@@ -181,6 +181,9 @@ struct FLayeredMoveGroup
 	
 	UE_API void QueueLayeredMove(TSharedPtr<FLayeredMoveBase> Move);
 
+	/** Schedule matching layered moves to be cancelled ASAP */
+	void CancelMovesByTag(FGameplayTag Tag, bool bRequireExactMatch=false);
+
 	bool HasAnyMoves() const { return (!ActiveLayeredMoves.IsEmpty() || !QueuedLayeredMoves.IsEmpty()); }
 
 	/** @return True if there is at least one layered move of MoveType that's either active or queued */
@@ -256,6 +259,8 @@ protected:
 	/** Moves that are queued to become active next sim frame */
 	TArray< TSharedPtr<FLayeredMoveBase> > QueuedLayeredMoves;
 
+	/** Used during simulation to cancel any moves that match a tag */
+	TArray<TPair<FGameplayTag, bool>> TagCancellationRequests;
 public:
 
 	/**
diff --git a/Engine/Plugins/Experimental/Mover/Source/Mover/Public/LayeredMoveBase.h b/Engine/Plugins/Experimental/Mover/Source/Mover/Public/LayeredMoveBase.h
--- a/Engine/Plugins/Experimental/Mover/Source/Mover/Public/LayeredMoveBase.h
+++ b/Engine/Plugins/Experimental/Mover/Source/Mover/Public/LayeredMoveBase.h
@@ -256,7 +256,8 @@ struct FLayeredMoveInstance
 	UE_API EMoveMixMode GetMixMode() const;
 	uint8 GetPriority() const { return MoveLogic->GetPriority(); }
 	double GetStartingTimeMs() const { return InstanceMoveData->GetStartSimTimeMs(); }
-	
+	bool HasGameplayTag(FGameplayTag TagToFind, bool bExactMatch) const { return InstanceMoveData->HasGameplayTag(TagToFind, bExactMatch); }
+
 	UE_API FLayeredMoveInstance Clone() const;
 	UScriptStruct* GetDataStructType() const { return InstanceMoveData->GetScriptStruct(); }
 	void AddReferencedObjects(FReferenceCollector& Collector) const { InstanceMoveData->AddReferencedObjects(Collector); } 
diff --git a/Engine/Plugins/Experimental/Mover/Source/Mover/Public/LayeredMoveGroup.h b/Engine/Plugins/Experimental/Mover/Source/Mover/Public/LayeredMoveGroup.h
--- a/Engine/Plugins/Experimental/Mover/Source/Mover/Public/LayeredMoveGroup.h
+++ b/Engine/Plugins/Experimental/Mover/Source/Mover/Public/LayeredMoveGroup.h
@@ -2,6 +2,7 @@
 
 #pragma once
 
+#include "GameplayTagContainer.h"
 #include "MoveLibrary/MovementUtilsTypes.h"
 #include "LayeredMoveGroup.generated.h"
 
@@ -91,6 +92,9 @@ struct FLayeredMoveInstanceGroup
 		return PrivateFindQueuedMove(MoveDataType);
 	}
 
+	/** Cancel any active or queued moves with a matching tag */
+	void CancelMovesByTag(FGameplayTag Tag, bool bRequireExactMatch); 
+
 	// Clears out any finished or invalid active moves and adds any queued moves to the active moves
 	UE_API void FlushMoveArrays(const FMoverTimeStep& TimeStep, UMoverBlackboard* SimBlackboard);
 	
@@ -131,6 +135,9 @@ private:
 	 */
 	UPROPERTY(NotReplicated, VisibleAnywhere, BlueprintReadOnly, Category = Mover, meta = (AllowPrivateAccess))
 	FVector ResidualVelocity;
+
+	/** Used during simulation to cancel any moves that match a tag */
+	TArray<TPair<FGameplayTag, bool>> TagCancellationRequests;
 };
 
 template<>
diff --git a/Engine/Plugins/Experimental/Mover/Source/Mover/Public/MovementModeStateMachine.h b/Engine/Plugins/Experimental/Mover/Source/Mover/Public/MovementModeStateMachine.h
--- a/Engine/Plugins/Experimental/Mover/Source/Mover/Public/MovementModeStateMachine.h
+++ b/Engine/Plugins/Experimental/Mover/Source/Mover/Public/MovementModeStateMachine.h
@@ -63,6 +63,8 @@ public:
  	UE_API const FMovementModifierBase* FindQueuedModifier(FMovementModifierHandle ModifierHandle) const;
  	UE_API const FMovementModifierBase* FindQueuedModifierByType(const UScriptStruct* ModifierType) const;
 
+	UE_API void CancelFeaturesWithTag(FGameplayTag TagToCancel, bool bRequireExactMatch=false);
+
 	// This function is meant to be used only in async mode on the physics thread, not on the game thread
 	void QueueInstantMovementEffect_Internal(const FScheduledInstantMovementEffect& ScheduledInstantMovementEffect);
 protected:
@@ -102,6 +104,9 @@ protected:
  	/** Modifiers that are to be canceled at the start of the next sim subtick.  Access covered by lock. */
  	TArray<FMovementModifierHandle> ModifiersToCancel;
  	
+	/** Tags that are used to cancel any matching movement features (modifiers, layered moves, etc). Access covered by lock. */
+	TArray<TPair<FGameplayTag, bool>> TagCancellationRequests;
+
 	// Internal-use-only tick data structs, for efficiency since they typically have the same contents from frame to frame
 	FMoverTickStartData WorkingSubstepStartData;
 	FSimulationTickParams WorkingSimTickParams;
@@ -112,6 +117,7 @@ private:
 	mutable FRWLock InstantEffectsQueueLock;
 	mutable FRWLock ModifiersQueueLock;
 	mutable FRWLock ModifierCancelQueueLock;
+	mutable FRWLock TagCancellationRequestsLock;
 
 	UE_API void ConstructDefaultModes();
 	UE_API void AdvanceToNextMode();
@@ -120,6 +126,7 @@ private:
  	UE_API void ActivateQueuedMoves(FLayeredMoveInstanceGroup& Group);
  	UE_API void FlushQueuedModifiersToGroup(FMovementModifierGroup& ModifierGroup);
  	UE_API void FlushModifierCancellationsToGroup(FMovementModifierGroup& ActiveModifierGroup);
+	UE_API void FlushTagCancellationsToSyncState(FMoverSyncState& SyncState);
  	UE_API void RollbackModifiers(const FMoverSyncState* InvalidSyncState, const FMoverSyncState* SyncState, const FMoverAuxStateContext* InvalidAuxState, const FMoverAuxStateContext* AuxState);
 	UE_API bool HasAnyInstantEffectsQueued() const;
  	UE_API bool ApplyInstantEffects(FApplyMovementEffectParams& ApplyEffectParams, FMoverSyncState& OutputState);
diff --git a/Engine/Plugins/Experimental/Mover/Source/Mover/Public/MovementModifier.h b/Engine/Plugins/Experimental/Mover/Source/Mover/Public/MovementModifier.h
--- a/Engine/Plugins/Experimental/Mover/Source/Mover/Public/MovementModifier.h
+++ b/Engine/Plugins/Experimental/Mover/Source/Mover/Public/MovementModifier.h
@@ -232,6 +232,7 @@ struct FMovementModifierGroup
 	UE_API void QueueMovementModifier(TSharedPtr<FMovementModifierBase> Modifier);
 
 	UE_API void CancelModifierFromHandle(const FMovementModifierHandle& HandleToCancel);
+	UE_API void CancelModifiersByTag(FGameplayTag Tag, bool bRequiresExactMatch=false);
 	
 	// Generates active modifier list (by calling FlushModifierArrays) and returns the an array of all currently active modifiers
 	UE_API TArray<TSharedPtr<FMovementModifierBase>> GenerateActiveModifiers(UMoverComponent* MoverComp, const FMoverTimeStep& TimeStep, const FMoverSyncState& SyncState, const FMoverAuxStateContext& AuxState);
diff --git a/Engine/Plugins/Experimental/Mover/Source/Mover/Public/MoverComponent.h b/Engine/Plugins/Experimental/Mover/Source/Mover/Public/MoverComponent.h
--- a/Engine/Plugins/Experimental/Mover/Source/Mover/Public/MoverComponent.h
+++ b/Engine/Plugins/Experimental/Mover/Source/Mover/Public/MoverComponent.h
@@ -304,7 +304,13 @@ public:
 	 */
 	UFUNCTION(BlueprintCallable, Category = Mover)
 	MOVER_API void CancelModifierFromHandle(FMovementModifierHandle ModifierHandle);
-	
+
+	/**
+	 * Cancel any active or queued movement features (layered moves, modifiers, etc.) that have a matching gameplay tag. Does not affect the active movement mode.
+	 */
+	UFUNCTION(BlueprintCallable, Category = Mover)
+	MOVER_API void CancelFeaturesWithTag(FGameplayTag TagToCancel, bool bRequireExactMatch=false);
+
 	/**
 	 * Queue an Instant Movement Effect to start at the end of this frame or start of the next subtick - whichever happens first. This will clone whatever move you pass in, so you'll need to fully set it up before queuing.
 	 * @param InstantMovementEffect			The effect to queue, which must be a FInstantMovementEffect sub-type. 
