Cleanup of possible null dereference

#rnx

diff --git a/Engine/Plugins/Experimental/Mover/Source/Mover/Private/LayeredMoveGroup.cpp b/Engine/Plugins/Experimental/Mover/Source/Mover/Private/LayeredMoveGroup.cpp
--- a/Engine/Plugins/Experimental/Mover/Source/Mover/Private/LayeredMoveGroup.cpp
+++ b/Engine/Plugins/Experimental/Mover/Source/Mover/Private/LayeredMoveGroup.cpp
@@ -309,7 +309,7 @@ void FLayeredMoveInstanceGroup::FlushMoveArrays(const FMoverTimeStep& TimeStep,
 		ActiveMoves.RemoveAll([&TimeStep, &SimBlackboard, &bResidualVelocityOverridden, &bClampVelocityOverridden, this]
 			(const TSharedPtr<FLayeredMoveInstance>& Move)
 			{
-				if (Move->IsFinished(TimeStep, SimBlackboard))
+				if (Move.IsValid() && Move->IsFinished(TimeStep, SimBlackboard))
 				{
 					ProcessFinishedMove(*Move, bResidualVelocityOverridden, bClampVelocityOverridden);
 					Move->EndMove(TimeStep, SimBlackboard);
