Updated Smooth Walking Mode to better account for collisions and induced velocities

diff --git a/Engine/Plugins/Experimental/Mover/Source/Mover/Private/DefaultMovementSet/Modes/SmoothWalkingMode.cpp b/Engine/Plugins/Experimental/Mover/Source/Mover/Private/DefaultMovementSet/Modes/SmoothWalkingMode.cpp
--- a/Engine/Plugins/Experimental/Mover/Source/Mover/Private/DefaultMovementSet/Modes/SmoothWalkingMode.cpp
+++ b/Engine/Plugins/Experimental/Mover/Source/Mover/Private/DefaultMovementSet/Modes/SmoothWalkingMode.cpp
@@ -44,10 +44,27 @@ void USmoothWalkingMode::GenerateWalkMove_Implementation(FMoverTickStartData& St
 	// carried between movement modes in FMoverDefaultSyncState.
 	if (bSmoothWalkingStateAdded)
 	{
-		SpringState.IntermediateFacing = InOutFacing;
+		SpringState.SpringVelocity = InOutVelocity;
+		SpringState.SpringAcceleration = FVector::ZeroVector;
 		SpringState.IntermediateVelocity = InOutVelocity;
+		SpringState.IntermediateFacing = InOutFacing;
+		SpringState.IntermediateAngularVelocity = FVector::ZeroVector;
+		SpringState.CurrentAngularVelocity = FVector::ZeroVector;
 	}
 
+	// We can compute the degree to which the internal velocity matches the actual velocity by projecting the spring velocity onto the actual velocity 
+	// which we moved last frame. This gives a number between 0 and 1 which says how different our velocity was from what we expected.
+	const float VelocityMatch = FMath::Clamp(SpringState.SpringVelocity.Dot(InOutVelocity) / 
+		FMath::Max(InOutVelocity.Length() * SpringState.SpringVelocity.Length(), UE_SMALL_NUMBER), 0.0f, 1.0f);
+
+	// If our velocity was very different from what we expected then we can effectively "reset" the intermediate velocity in a smooth way towards it. 
+	// This removes any velocity we may have built-up in the intermediate spring that is different from our current velocity.
+	FMath::ExponentialSmoothingApprox(SpringState.IntermediateVelocity, InOutVelocity, DeltaSeconds, 
+		(OutsideInfluenceSmoothingTime + UE_KINDA_SMALL_NUMBER) / (1.0f - VelocityMatch));
+
+	// Update spring velocity based on real velocity
+	SpringState.SpringVelocity = InOutVelocity;
+
 	// Rotate the intermediate velocity towards the target direction using the TurningStrength
 	if (TurningStrength > 0.0f)
 	{
@@ -63,18 +80,18 @@ void USmoothWalkingMode::GenerateWalkMove_Implementation(FMoverTickStartData& St
 
 	// Check if we are acceleration or decelerating and work out how much lateral vs directional acceleration to apply. Note that even when the
 	// DirectionalAccelerationFactor is high Deceleration is always applied laterally. This is similar to how the default walking mode behaves.
-	const bool bIsAccelerating = (1.01f * DesiredVelocity.SquaredLength()) > InOutVelocity.SquaredLength();
+	const bool bIsAccelerating = (1.01f * DesiredVelocity.SquaredLength()) > SpringState.SpringVelocity.SquaredLength();
 	const float LateralAccelerationMagnitude =  bIsAccelerating ? (1.0f - DirectionalAccelerationFactor) * Acceleration : Deceleration;
 	const float DirectionalAccelerationMagnitude = bIsAccelerating ? DirectionalAccelerationFactor * Acceleration : 0.0f;
 
-	// Record the previous velocity
-	const FVector PreviousVelocity = SpringState.IntermediateVelocity;
+	// Record the previous velocity length
+	const float PreviousVelocityLength = SpringState.IntermediateVelocity.Length();
 	
 	// Compute the desired difference in velocity
 	const FVector VelocityDifference = DesiredVelocity - SpringState.IntermediateVelocity;
 
 	// Compute the lateral acceleration moving directly toward the desired velocity
-	const FVector LateralAccelerationVector = VelocityDifference.GetSafeNormal() * FMath::Min(LateralAccelerationMagnitude, VelocityDifference.Length() / DeltaSeconds);
+	const FVector LateralAccelerationVector = VelocityDifference.GetSafeNormal() * FMath::Min(LateralAccelerationMagnitude, VelocityDifference.Length() / FMath::Max(DeltaSeconds, UE_SMALL_NUMBER));
 
 	// Compute the directional acceleration moving in the desired direction. This emulates how acceleration is applied in the default movement mode.
 	const FVector DirectionalAccelerationVector = DesiredVelocity.GetSafeNormal() * DirectionalAccelerationMagnitude;
@@ -90,7 +107,7 @@ void USmoothWalkingMode::GenerateWalkMove_Implementation(FMoverTickStartData& St
 	FVector NextVelocity = VelocityDifference.Length() > DesiredAcceleration.Length() * DeltaSeconds ?
 		SpringState.IntermediateVelocity + DesiredAcceleration * DeltaSeconds : DesiredVelocity;
 
-	NextVelocity = NextVelocity.GetClampedToMaxSize(FMath::Max(PreviousVelocity.Length(), DesiredVelocity.Length()));
+	NextVelocity = NextVelocity.GetClampedToMaxSize(FMath::Max(PreviousVelocityLength, DesiredVelocity.Length()));
 
 	// Compute the smoothing time based on acceleration or deceleration
 	const float VelocitySmoothingTime = bIsAccelerating ? AccelerationSmoothingTime : DecelerationSmoothingTime;
@@ -100,22 +117,22 @@ void USmoothWalkingMode::GenerateWalkMove_Implementation(FMoverTickStartData& St
 
 	// Get the lag associated with the current velocity smoothing time and perform the same process to estimate the velocity we need to track to
 	// avoid lagging behind
-	const float LagSeconds =  DeltaSeconds + VelocitySmoothingCompensation * SpringMath::HalfLifeToLag(SpringMath::SpringSmoothingTimeToHalfLife(VelocitySmoothingTime));
+	const float LagSeconds = DeltaSeconds + VelocitySmoothingCompensation * SpringMath::HalfLifeToLag(SpringMath::SpringSmoothingTimeToHalfLife(VelocitySmoothingTime));
 
 	FVector TrackVelocity =
 		VelocityDifference.Length() > DesiredAcceleration.Length() * LagSeconds ?
 		SpringState.IntermediateVelocity + DesiredAcceleration * LagSeconds : DesiredVelocity;
 
-	TrackVelocity = TrackVelocity.GetClampedToMaxSize(FMath::Max(PreviousVelocity.Length(), DesiredVelocity.Length()));
+	TrackVelocity = TrackVelocity.GetClampedToMaxSize(FMath::Max(PreviousVelocityLength, DesiredVelocity.Length()));
 
 	// Apply the smoothing to the track velocity - effectively tracking the intermediate velocity at the appropriate time in the future
-	SpringMath::SimpleSpringDamper(InOutVelocity, SpringState.SpringAcceleration, TrackVelocity, VelocitySmoothingTime, DeltaSeconds);
+	SpringMath::SimpleSpringDamper(SpringState.SpringVelocity, SpringState.SpringAcceleration, TrackVelocity, VelocitySmoothingTime, DeltaSeconds);
 
 	// Snap the velocity to the desired velocity based on the dead-zone
-	if ((DesiredVelocity - InOutVelocity).SquaredLength() < FMath::Square(VelocityDeadzoneThreshold))
+	if ((DesiredVelocity - SpringState.SpringVelocity).SquaredLength() < FMath::Square(VelocityDeadzoneThreshold))
 	{
 		// We reached our target
-		InOutVelocity = DesiredVelocity;
+		SpringState.SpringVelocity = DesiredVelocity;
 	
 		// If we've reached our target then also snap the acceleration to zero once it is close enough
 		if (SpringState.SpringAcceleration.SquaredLength() < FMath::Square(AccelerationDeadzoneThreshold))
@@ -124,6 +141,9 @@ void USmoothWalkingMode::GenerateWalkMove_Implementation(FMoverTickStartData& St
 		}
 	}
 
+	// Update the target velocity
+	InOutVelocity = SpringState.SpringVelocity;
+
 	// Update the intermediate velocity
 	SpringState.IntermediateVelocity = NextVelocity;
 
diff --git a/Engine/Plugins/Experimental/Mover/Source/Mover/Private/DefaultMovementSet/Modes/SmoothWalkingState.cpp b/Engine/Plugins/Experimental/Mover/Source/Mover/Private/DefaultMovementSet/Modes/SmoothWalkingState.cpp
--- a/Engine/Plugins/Experimental/Mover/Source/Mover/Private/DefaultMovementSet/Modes/SmoothWalkingState.cpp
+++ b/Engine/Plugins/Experimental/Mover/Source/Mover/Private/DefaultMovementSet/Modes/SmoothWalkingState.cpp
@@ -19,6 +19,7 @@ bool FSmoothWalkingState::NetSerialize(FArchive& Ar, UPackageMap* Map, bool& bOu
 	bool bSuccess = Super::NetSerialize(Ar, Map, bOutSuccess);
 
 	// Could be quantized to save bandwidth
+	Ar << SpringVelocity;
 	Ar << SpringAcceleration;
 	Ar << IntermediateVelocity;
 	Ar << IntermediateFacing;
@@ -32,9 +33,10 @@ void FSmoothWalkingState::ToString(FAnsiStringBuilderBase& Out) const
 {
 	Super::ToString(Out);
 
-	Out.Appendf("SpringAcceleration=%s IntVel=%s IntFac=%s IntAng=%s AngVel=%s\n",
+	Out.Appendf("SpringVelocity=%s SpringAcceleration=%s IntVel=%s IntFac=%s IntAng=%s AngVel=%s\n",
+		*SpringVelocity.ToCompactString(),
 		*SpringAcceleration.ToCompactString(),
-		*IntermediateVelocity.ToCompactString(), 
+		*IntermediateVelocity.ToCompactString(),
 		*IntermediateFacing.ToString(), 
 		*IntermediateAngularVelocity.ToString(),  
 		*CurrentAngularVelocity.ToCompactString());
@@ -44,7 +46,8 @@ bool FSmoothWalkingState::ShouldReconcile(const FMoverDataStructBase& AuthorityS
 {
 	const FSmoothWalkingState* AuthoritySpringState = static_cast<const FSmoothWalkingState*>(&AuthorityState);
 
-	return (!(SpringAcceleration - AuthoritySpringState->SpringAcceleration).IsNearlyZero() ||
+	return (!(SpringVelocity - AuthoritySpringState->SpringVelocity).IsNearlyZero() ||
+			!(SpringAcceleration - AuthoritySpringState->SpringAcceleration).IsNearlyZero() ||
 			!(IntermediateVelocity - AuthoritySpringState->IntermediateVelocity).IsNearlyZero() ||
 			 (IntermediateFacing.AngularDistance(AuthoritySpringState->IntermediateFacing) > UE_SMALL_NUMBER ||
 			!(IntermediateAngularVelocity - AuthoritySpringState->IntermediateAngularVelocity).IsNearlyZero() ||
@@ -57,6 +60,7 @@ void FSmoothWalkingState::Interpolate(const FMoverDataStructBase& From, const FM
 	const FSmoothWalkingState* FromState = static_cast<const FSmoothWalkingState*>(&From);
 	const FSmoothWalkingState* ToState   = static_cast<const FSmoothWalkingState*>(&To);
 
+	SpringVelocity = FMath::Lerp(FromState->SpringVelocity, ToState->SpringVelocity, Pct);
 	SpringAcceleration = FMath::Lerp(FromState->SpringAcceleration, ToState->SpringAcceleration, Pct);
 	IntermediateVelocity = FMath::Lerp(FromState->IntermediateVelocity, ToState->IntermediateVelocity, Pct);
 	IntermediateFacing = FQuat::Slerp(FromState->IntermediateFacing, ToState->IntermediateFacing, Pct);
diff --git a/Engine/Plugins/Experimental/Mover/Source/Mover/Private/DefaultMovementSet/Modes/SmoothWalkingState.h b/Engine/Plugins/Experimental/Mover/Source/Mover/Private/DefaultMovementSet/Modes/SmoothWalkingState.h
--- a/Engine/Plugins/Experimental/Mover/Source/Mover/Private/DefaultMovementSet/Modes/SmoothWalkingState.h
+++ b/Engine/Plugins/Experimental/Mover/Source/Mover/Private/DefaultMovementSet/Modes/SmoothWalkingState.h
@@ -22,6 +22,10 @@ struct FSmoothWalkingState : public FMoverDataStructBase
 	virtual bool ShouldReconcile(const FMoverDataStructBase& AuthorityState) const override;
 	virtual void Interpolate(const FMoverDataStructBase& From, const FMoverDataStructBase& To, float Pct) override;
 
+	// Velocity of internal velocity spring
+	UPROPERTY(BlueprintReadOnly, Category = "Mover|Experimental")
+	FVector SpringVelocity = FVector::ZeroVector;
+
 	// Acceleration of internal velocity spring
 	UPROPERTY(BlueprintReadOnly, Category = "Mover|Experimental")
 	FVector SpringAcceleration = FVector::ZeroVector;
diff --git a/Engine/Plugins/Experimental/Mover/Source/Mover/Public/DefaultMovementSet/Modes/SimpleWalkingMode.h b/Engine/Plugins/Experimental/Mover/Source/Mover/Public/DefaultMovementSet/Modes/SimpleWalkingMode.h
--- a/Engine/Plugins/Experimental/Mover/Source/Mover/Public/DefaultMovementSet/Modes/SimpleWalkingMode.h
+++ b/Engine/Plugins/Experimental/Mover/Source/Mover/Public/DefaultMovementSet/Modes/SimpleWalkingMode.h
@@ -26,7 +26,7 @@ public:
 	                                     const FQuat& DesiredFacing,UPARAM(ref)  FVector& InOutVelocity, UPARAM(ref) FQuat& InOutFacing);
 
 	// If this value is greater or equal to 0, this will override the max speed read from the common legacy shared walk settings
-	UPROPERTY(EditAnywhere, BlueprintReadWrite, Category = "Mover", meta = (ForceUnits = "cm/s"))
+	UPROPERTY(EditAnywhere, BlueprintReadWrite, Category = "Mover|Walking Settings", meta = (ForceUnits = "cm/s"))
 	float MaxSpeedOverride = -1.0f;
 };
 
diff --git a/Engine/Plugins/Experimental/Mover/Source/Mover/Public/DefaultMovementSet/Modes/SmoothWalkingMode.h b/Engine/Plugins/Experimental/Mover/Source/Mover/Public/DefaultMovementSet/Modes/SmoothWalkingMode.h
--- a/Engine/Plugins/Experimental/Mover/Source/Mover/Public/DefaultMovementSet/Modes/SmoothWalkingMode.h
+++ b/Engine/Plugins/Experimental/Mover/Source/Mover/Public/DefaultMovementSet/Modes/SmoothWalkingMode.h
@@ -27,11 +27,11 @@ protected: // Velocity Controls
 
 	// Base acceleration to apply when the desired velocity magnitude is greater than the current velocity.
 	UPROPERTY(EditAnywhere, BlueprintReadWrite, Category = "Mover|Smooth Walking Settings", meta = (ClampMin = "0", UIMin = "0", ForceUnits = "cm/s^2"))
-	float Acceleration = 800.0f;
+	float Acceleration = 1500.0f;
 
 	// Base deceleration to apply when the desired velocity magnitude is less than the current velocity
 	UPROPERTY(EditAnywhere, BlueprintReadWrite, Category = "Mover|Smooth Walking Settings", meta = (ClampMin = "0", UIMin = "0", ForceUnits = "cm/s^2"))
-	float Deceleration = 1000.0f;
+	float Deceleration = 1500.0f;
 
 	/**
 	 * Value between 0 and 1 that controls how the acceleration is applied. When set to 1 this will effectively add the acceleration in the direction 
@@ -40,7 +40,7 @@ protected: // Velocity Controls
 	 * velocity faster and more regular, always using a fixed acceleration, but can cause the overall velocity magnitude to dip during turning.
 	 */
 	UPROPERTY(EditAnywhere, BlueprintReadWrite, Category = "Mover|Advanced Smooth Walking Settings", meta = (ClampMin = "0", UIMin = "0", ClampMax = "1", UIMax = "1"))
-	float DirectionalAccelerationFactor = 0.0f;
+	float DirectionalAccelerationFactor = 1.0f;
 
 	/**
 	 * Applies an additional force that rotates the velocity to make it face the desired direction of travel. This allows quick turns which don't 
@@ -48,15 +48,15 @@ protected: // Velocity Controls
 	 * quite a strong turning force and 100 will produce near instant turns (before smoothing).
 	 */
 	UPROPERTY(EditAnywhere, BlueprintReadWrite, Category = "Mover|Smooth Walking Settings", meta = (ClampMin = "0", UIMin = "0"))
-	float TurningStrength = 3.0f;
+	float TurningStrength = 10.0f;
 
 	// Controls how much smoothing is applied to the velocity changes of the movement when accelerating. Set to zero to apply no smoothing.
 	UPROPERTY(EditAnywhere, BlueprintReadWrite, Category = "Mover|Smooth Walking Settings", meta = (ClampMin = "0", UIMin = "0", ForceUnits = "s"))
-	float AccelerationSmoothingTime = 0.2f;
+	float AccelerationSmoothingTime = 0.1f;
 
 	// Controls how much smoothing is applied to the velocity changes of the movement when decelerating. Set to zero to apply no smoothing.
 	UPROPERTY(EditAnywhere, BlueprintReadWrite, Category = "Mover|Smooth Walking Settings", meta = (ClampMin = "0", UIMin = "0", ForceUnits = "s"))
-	float DecelerationSmoothingTime = 0.05f;
+	float DecelerationSmoothingTime = 0.1f;
 
 	/**
 	 * When the velocity is smoothed it will naturally lag behind the unsmoothed target. This factor controls how much this is compensated for by
@@ -65,7 +65,7 @@ protected: // Velocity Controls
 	 * the initial change in velocity less smooth and reduces the effective "lead-in".
 	 */
 	UPROPERTY(EditAnywhere, BlueprintReadWrite, Category = "Mover|Advanced Smooth Walking Settings", meta = (ClampMin = "0", UIMin = "0", ClampMax = "1", UIMax = "1"))
-	float AccelerationSmoothingCompensation = 1.0f;
+	float AccelerationSmoothingCompensation = 0.0f;
 
 	// This parameter acts the same as AccelerationSmoothingCompensation but is applied during deceleration instead
 	UPROPERTY(EditAnywhere, BlueprintReadWrite, Category = "Mover|Advanced Smooth Walking Settings", meta = (ClampMin = "0", UIMin = "0", ClampMax = "1", UIMax = "1"))
@@ -79,12 +79,18 @@ protected: // Velocity Controls
 	UPROPERTY(EditAnywhere, BlueprintReadWrite, Category = "Mover|Advanced Smooth Walking Settings", meta = (ClampMin = "0", UIMin = "0", ForceUnits = "cm/s^2"))
 	float AccelerationDeadzoneThreshold = 0.001f;
 
+	// Controls how quickly the built-up internal velocity will be modified when the character's movement is influenced by the outside factors such 
+	// as collisions. Setting this to a small value will mean that collisions and other outside influences on the character's velocity "reset" 
+	// the built-up internal velocity quickly. Larger values will mean that the character loses less momentum and takes less time to re-accumulate 
+	// internal velocity in situations such as glancing collisions and other short external impulses such as pushes.
+	UPROPERTY(EditAnywhere, BlueprintReadWrite, Category = "Mover|Advanced Smooth Walking Settings", meta = (ClampMin = "0", UIMin = "0", ForceUnits = "s"))
+	float OutsideInfluenceSmoothingTime = 0.05f;
 
 protected: // Facing Controls
 
 	// Controls how much smoothing is applied to the facing direction. Set to zero to apply no smoothing.
 	UPROPERTY(EditAnywhere, BlueprintReadWrite, Category = "Mover|Smooth Walking Settings", meta = (ClampMin = "0", UIMin = "0", ForceUnits = "s"))
-	float FacingSmoothingTime = 0.3f;
+	float FacingSmoothingTime = 0.25f;
 
 	// Smooths facing using a double spring rather than a single spring. This produces a more S-shaped profile with a shorter "lead-out".
 	UPROPERTY(EditAnywhere, BlueprintReadWrite, Category = "Mover|Advanced Smooth Walking Settings")
