Mover: add tag to identify mesh attribute root motion layered move

#rnx

diff --git a/Engine/Plugins/Experimental/Mover/Source/Mover/Private/DefaultMovementSet/LayeredMoves/RootMotionAttributeLayeredMove.cpp b/Engine/Plugins/Experimental/Mover/Source/Mover/Private/DefaultMovementSet/LayeredMoves/RootMotionAttributeLayeredMove.cpp
--- a/Engine/Plugins/Experimental/Mover/Source/Mover/Private/DefaultMovementSet/LayeredMoves/RootMotionAttributeLayeredMove.cpp
+++ b/Engine/Plugins/Experimental/Mover/Source/Mover/Private/DefaultMovementSet/LayeredMoves/RootMotionAttributeLayeredMove.cpp
@@ -13,6 +13,8 @@
 
 #include UE_INLINE_GENERATED_CPP_BY_NAME(RootMotionAttributeLayeredMove)
 
+UE_DEFINE_GAMEPLAY_TAG_COMMENT(Mover_AnimRootMotion_MeshAttribute, "Mover.AnimRootMotion.MeshAttribute", "Signifies an association with root motion that comes via a skeletal mesh attribute");
+
 
 #if !UE_BUILD_SHIPPING
 FAutoConsoleVariable CVarLogRootMotionAttrSteps(
@@ -184,6 +186,14 @@ bool FLayeredMove_RootMotionAttribute::GenerateMove(const FMoverTickStartData& S
 	return false;
 }
 
+bool FLayeredMove_RootMotionAttribute::HasGameplayTag(FGameplayTag TagToFind, bool bExactMatch) const
+{
+	static const FGameplayTag MeshAttributeTag = Mover_AnimRootMotion_MeshAttribute.GetTag();
+	const bool bFoundMatch = bExactMatch ? MeshAttributeTag.MatchesTagExact(TagToFind) : MeshAttributeTag.MatchesTag(TagToFind);
+
+	return bFoundMatch || Super::HasGameplayTag(TagToFind, bExactMatch);
+}
+
 FLayeredMoveBase* FLayeredMove_RootMotionAttribute::Clone() const
 {
 	FLayeredMove_RootMotionAttribute* CopyPtr = new FLayeredMove_RootMotionAttribute(*this);
diff --git a/Engine/Plugins/Experimental/Mover/Source/Mover/Public/DefaultMovementSet/LayeredMoves/RootMotionAttributeLayeredMove.h b/Engine/Plugins/Experimental/Mover/Source/Mover/Public/DefaultMovementSet/LayeredMoves/RootMotionAttributeLayeredMove.h
--- a/Engine/Plugins/Experimental/Mover/Source/Mover/Public/DefaultMovementSet/LayeredMoves/RootMotionAttributeLayeredMove.h
+++ b/Engine/Plugins/Experimental/Mover/Source/Mover/Public/DefaultMovementSet/LayeredMoves/RootMotionAttributeLayeredMove.h
@@ -3,11 +3,14 @@
 #pragma once
 
 #include "LayeredMove.h"
+#include "NativeGameplayTags.h"
 #include "RootMotionModifier.h"
 #include "RootMotionAttributeLayeredMove.generated.h"
 
 #define UE_API MOVER_API
 
+MOVER_API UE_DECLARE_GAMEPLAY_TAG_EXTERN(Mover_AnimRootMotion_MeshAttribute);
+
 /** 
  * Root Motion Attribute Move: handles root motion from a mesh's custom attribute, ignoring scaling.
  * Currently only supports Independent ticking mode, and allows controlled movement while jumping/falling or when a SkipAnimRootMotion tag is active.
@@ -33,6 +36,8 @@ protected:
 	// Generate a movement 
 	UE_API virtual bool GenerateMove(const FMoverTickStartData& StartState, const FMoverTimeStep& TimeStep, const UMoverComponent* MoverComp, UMoverBlackboard* SimBlackboard, FProposedMove& OutProposedMove) override;
 
+	UE_API virtual bool HasGameplayTag(FGameplayTag TagToFind, bool bExactMatch) const override;
+
 	UE_API virtual FLayeredMoveBase* Clone() const override;
 
 	UE_API virtual void NetSerialize(FArchive& Ar) override;
