Physics Mover - Fix for pushing a low, dynamic object

#rnx

diff --git a/Engine/Plugins/Experimental/Mover/Source/Mover/Private/PhysicsMover/Modes/PhysicsDrivenWalkingMode.cpp b/Engine/Plugins/Experimental/Mover/Source/Mover/Private/PhysicsMover/Modes/PhysicsDrivenWalkingMode.cpp
--- a/Engine/Plugins/Experimental/Mover/Source/Mover/Private/PhysicsMover/Modes/PhysicsDrivenWalkingMode.cpp
+++ b/Engine/Plugins/Experimental/Mover/Source/Mover/Private/PhysicsMover/Modes/PhysicsDrivenWalkingMode.cpp
@@ -163,10 +163,16 @@ void UPhysicsDrivenWalkingMode::FloorCheck(const FMoverDefaultSyncState& SyncSta
 	const float StepBlockedHeight = TargetHeight - PawnHalfHeight + PawnRadius;
 	const float StepHeight = TargetHeight - OutFloorResult.FloorDist;
 
-	if (StepHeight > StepBlockedHeight)
+	bool bIsDynamicSurface = false;
+	if (const Chaos::FPBDRigidParticleHandle* GroundParticle = UPhysicsMovementUtils::GetRigidParticleHandleFromHitResult(OutFloorResult.HitResult))
+	{
+		bIsDynamicSurface = GroundParticle->IsDynamic();
+	}
+
+	if ((StepHeight > StepBlockedHeight) || bIsDynamicSurface)
 	{
 		// Collision should prevent movement. Just try to find ground at start of movement
-		const float ShrinkMultiplier = 0.75f;
+		const float ShrinkMultiplier = 0.5f;
 		UPhysicsMovementUtils::FloorSweep_Internal(SyncState.GetLocation_WorldSpace(), FVector::ZeroVector, UpdatedPrimitive, UpDir,
 			ShrinkMultiplier * QueryRadius, FloorSweepDistance, CommonLegacySettings->MaxWalkSlopeCosine, TargetHeight, OutFloorResult, OutWaterResult);
 
@@ -358,7 +364,11 @@ void UPhysicsDrivenWalkingMode::SimulationTick_Implementation(const FSimulationT
 		const bool bIsGroundMoving = ProjectedGroundVelocity.SizeSquared() > UE_KINDA_SMALL_NUMBER;
 
 		TargetVelocity = StartingSyncState->GetVelocity_WorldSpace();
-		FVector TargetPosition = StartingSyncState->GetLocation_WorldSpace() + UpDir * (TargetHeight - FloorResult.FloorDist);
+		FVector TargetPosition = StartingSyncState->GetLocation_WorldSpace();
+		if (!bIsGroundMoving)
+		{
+			TargetPosition += UpDir * (TargetHeight - FloorResult.FloorDist);
+		}
 
 		constexpr float ParallelCosThreshold = 0.999f;
 		const bool bVerticalVelocity = FVector::Parallel(TargetVelocity.GetSafeNormal(), UpDir, ParallelCosThreshold);
