Mover: Making sure we use BaseVisualComponentTransform to update PrimaryVisualComponent transform in FinalizeFrame so changes to BaseVisualComponentTransform apply to the visual component regardless of whether we're smoothing state or not.

diff --git a/Engine/Plugins/Experimental/Mover/Source/Mover/Private/MoverComponent.cpp b/Engine/Plugins/Experimental/Mover/Source/Mover/Private/MoverComponent.cpp
--- a/Engine/Plugins/Experimental/Mover/Source/Mover/Private/MoverComponent.cpp
+++ b/Engine/Plugins/Experimental/Mover/Source/Mover/Private/MoverComponent.cpp
@@ -377,6 +377,14 @@ void UMoverComponent::FinalizeFrame(const FMoverSyncState* SyncState, const FMov
 		UpdateCachedFrameState(SyncState, AuxState);
 	}
 
+	if (PrimaryVisualComponent)
+	{
+		if (!PrimaryVisualComponent->GetRelativeTransform().Equals(BaseVisualComponentTransform))
+		{
+			PrimaryVisualComponent->SetRelativeTransform(BaseVisualComponentTransform);
+		}
+	}
+	
 	if (OnPostFinalize.IsBound())
 	{
 		OnPostFinalize.Broadcast(MoverSyncStateDoubleBuffer.GetReadable(), CachedLastAuxState);
@@ -408,13 +416,6 @@ void UMoverComponent::FinalizeSmoothingFrame(const FMoverSyncState* SyncState, c
 				PrimaryVisualComponent->SetWorldTransform(BaseVisualComponentTransform * ActorTransform);	// smoothed location with base offset applied
 			}
 		}
-		else
-		{
-			if (!PrimaryVisualComponent->GetRelativeTransform().Equals(BaseVisualComponentTransform))
-			{
-				PrimaryVisualComponent->SetRelativeTransform(BaseVisualComponentTransform);
-			}
-		}
 	}
 }
 
