[Mover] Crouching now changes the physics capsule instead of capsule component

diff --git a/Engine/Plugins/Experimental/Mover/Source/Mover/Private/PhysicsMover/MovementModifiers/PhysicsStanceModifier.cpp b/Engine/Plugins/Experimental/Mover/Source/Mover/Private/PhysicsMover/MovementModifiers/PhysicsStanceModifier.cpp
--- a/Engine/Plugins/Experimental/Mover/Source/Mover/Private/PhysicsMover/MovementModifiers/PhysicsStanceModifier.cpp
+++ b/Engine/Plugins/Experimental/Mover/Source/Mover/Private/PhysicsMover/MovementModifiers/PhysicsStanceModifier.cpp
@@ -3,17 +3,40 @@
 #include "PhysicsMover/MovementModifiers/PhysicsStanceModifier.h"
 
 #include "Chaos/ParticleHandle.h"
+#include "Chaos/Capsule.h"
 #include "Components/CapsuleComponent.h"
 #include "DefaultMovementSet/CharacterMoverComponent.h"
 #include "DefaultMovementSet/Settings/StanceSettings.h"
 #include "MoveLibrary/MovementUtils.h"
+#include "PhysicsEngine/PhysicsObjectExternalInterface.h"
 #include "PhysicsMover/Modes/PhysicsDrivenWalkingMode.h"
 #include "PhysicsMover/Modes/PhysicsDrivenFallingMode.h"
 #include "PhysicsMover/PhysicsMovementUtils.h"
+#include "PBDRigidsSolver.h"
+#include "Physics/Experimental/PhysScene_Chaos.h"
 
 #include UE_INLINE_GENERATED_CPP_BY_NAME(PhysicsStanceModifier)
 
-void FPhysicsStanceModifier::ApplyModifierMove(UMoverComponent* MoverComp, const FMoverSyncState& InSyncState)
+
+const Chaos::FCapsule FPhysicsStanceModifier::GenerateNewCapsule(float NewHalfHeight, float NewRadius, float NewGroundClearance, float TargetHeight)
+{
+	const float NewHeight = FMath::Max(2.0f * (NewHalfHeight - NewRadius), 0.0f);
+	const Chaos::FVec3 Axis = Chaos::FVec3::UpVector;
+	const Chaos::FVec3 NewX1 = (NewGroundClearance + NewRadius - TargetHeight) * Axis;
+	const Chaos::FVec3 NewX2 = NewX1 + Axis * NewHeight;
+	return Chaos::FCapsule(NewX1, NewX2, NewRadius);
+}
+
+void FPhysicsStanceModifier::UpdateCapsule(Chaos::FPBDRigidParticleHandle* ParticleHandle, float NewHalfHeight, float NewRadius, float TargetHeight, float GroundClearance)
+{
+	check(ParticleHandle);
+
+	const Chaos::FCapsule NewCapsule = GenerateNewCapsule(NewHalfHeight, NewRadius, GroundClearance, TargetHeight);
+	ParticleHandle->SetGeometry(Chaos::FImplicitObjectPtr(new Chaos::FCapsule(NewCapsule)));
+	ParticleHandle->SetCenterOfMass(NewCapsule.GetCenterf());
+}
+
+void FPhysicsStanceModifier::ApplyModifierMove(UMoverComponent* MoverComp, const FMoverSyncState& InSyncState, float OriginalHalfHeight, float NewHalfHeight, float NewRadius)
 {
 	if (const UBaseMovementMode* MovementMode = MoverComp->FindMovementModeByName(InSyncState.MovementMode))
 	{
@@ -25,68 +48,26 @@ void FPhysicsStanceModifier::ApplyModifierMove(UMoverComponent* MoverComp, const
 				{
 					if (UPrimitiveComponent* PrimitiveComp = Cast<UPrimitiveComponent>(MoverComp->GetUpdatedComponent()))
 					{
-						FMoverDefaultSyncState& SyncState = MoverState->SyncStateCollection.FindOrAddMutableDataByType<FMoverDefaultSyncState>();
-
-						FVector TargetPosition = SyncState.GetLocation_WorldSpace();
-						const float TargetHeight = PhysicsMode->GetTargetHeight();
-
-						const FVector UpDir = FVector::UpVector;
-						FFloorCheckResult FloorResult;
-						float FloorDist = UE_BIG_NUMBER;
-						UMoverBlackboard* SimBlackboard = MoverComp->GetSimBlackboard_Mutable();
-
-						if (SimBlackboard->TryGet(CommonBlackboard::LastFloorResult, FloorResult) && FloorResult.bBlockingHit)
-						{
-							// Last floor distance is ok in this case
-							FloorDist = FloorResult.FloorDist;
-						}
-						else
+						// Update the character particle shape
+						if (Chaos::FPBDRigidParticleHandle* ParticleHandle = UPhysicsMovementUtils::GetRigidParticleHandleFromComponent(PrimitiveComp))
 						{
-							// Perform a raycast if no floor found (e.g., on rollbacks)
-							FCollisionQueryParams QueryParams(SCENE_QUERY_STAT(PhysicsFloorTest), false, PrimitiveComp->GetOwner());
-							QueryParams.bTraceIntoSubComponents = false;
-							const ECollisionChannel CollisionChannel = PrimitiveComp->GetCollisionObjectType();
-							FCollisionResponseParams ResponseParams(ECR_Overlap);
-							ResponseParams.CollisionResponse.SetResponse(ECC_WorldStatic, ECR_Block);
-							ResponseParams.CollisionResponse.SetResponse(ECC_WorldDynamic, ECR_Block);
-							ResponseParams.CollisionResponse.SetResponse(ECC_Vehicle, ECR_Block);
-							ResponseParams.CollisionResponse.SetResponse(ECC_Destructible, ECR_Block);
-							ResponseParams.CollisionResponse.SetResponse(ECC_PhysicsBody, ECR_Block);
-
-							PrimitiveComp->InitSweepCollisionParams(QueryParams, ResponseParams);
-
-							const float AdjustedQueryDistance = 2.f * TargetHeight;
-
-							FVector Start = TargetPosition;
-							FVector End = Start - AdjustedQueryDistance * UpDir;
-							FHitResult OutHit;
-							TArray<FHitResult> Hits;
-							if (const UWorld* World = PrimitiveComp->GetWorld())
+							// Invalidate particle before changing collision geometry (matches Chaos stance flow)
+							if (const UWorld* World = MoverComp->GetWorld())
 							{
-								if (Chaos::Private::FGenericPhysicsInterface_Internal::RaycastMulti(World, Hits, Start, End, CollisionChannel, QueryParams, ResponseParams))
+								if (FPhysScene* Scene = World->GetPhysicsScene())
 								{
-									OutHit = Hits.Last();
+									if (Chaos::FPhysicsSolver* Solver = Scene->GetSolver())
+									{
+										if (Solver->GetEvolution())
+										{
+											Solver->GetEvolution()->InvalidateParticle(ParticleHandle);
+										}
+									}
 								}
 							}
 
-							FloorDist = UpDir.Dot(TargetPosition - OutHit.ImpactPoint);
-						}
-
-						if (FloorDist <= (2.f * TargetHeight))
-						{
-							TargetPosition += UpDir * (TargetHeight - FloorDist);
-						}
-						SyncState.SetTransforms_WorldSpace(
-							TargetPosition,
-							SyncState.GetOrientation_WorldSpace(),
-							SyncState.GetVelocity_WorldSpace(),
-							nullptr);
-
-						// Force update the character particle transform
-						if (Chaos::FPBDRigidParticleHandle* ParticleHandle = UPhysicsMovementUtils::GetRigidParticleHandleFromComponent(PrimitiveComp))
-						{
-							ParticleHandle->SetX(TargetPosition);
-							ParticleHandle->SetV(SyncState.GetVelocity_WorldSpace());
+							const float TargetHeight = PhysicsMode->GetTargetHeight();
+							UpdateCapsule(ParticleHandle, NewHalfHeight, NewRadius, TargetHeight, TargetHeight - OriginalHalfHeight);
 						}
 					}
 				}
@@ -105,13 +86,9 @@ void FPhysicsStanceModifier::OnStart(UMoverComponent* MoverComp, const FMoverTim
 			{
 				const float OriginalHalfHeight = OriginalCapsule->GetScaledCapsuleHalfHeight();
 				const float CrouchedHalfHeight = StanceSettings->CrouchHalfHeight;
-				const float CrouchedEyeHeight = StanceSettings->CrouchedEyeHeight;
-				
-				// Adjust the TargetHeight of all the relevant MovementModes
-				UpdateTargetHeight<UPhysicsDrivenWalkingMode>(MoverComp, OriginalHalfHeight, CrouchedHalfHeight, true);
-				UpdateTargetHeight<UPhysicsDrivenFallingMode>(MoverComp, OriginalHalfHeight, CrouchedHalfHeight, true);
+				const float OriginalRadius = OriginalCapsule->GetScaledCapsuleRadius();
 
-				ApplyModifierMove(MoverComp, SyncState);
+				ApplyModifierMove(MoverComp, SyncState, OriginalHalfHeight, CrouchedHalfHeight, OriginalRadius);
 
 				ApplyMovementSettings(MoverComp);
 			}
@@ -134,13 +111,9 @@ void FPhysicsStanceModifier::OnEnd(UMoverComponent* MoverComp, const FMoverTimeS
 			if (const UCapsuleComponent* OriginalCapsule = UMovementUtils::GetOriginalComponentType<UCapsuleComponent>(MoverComp->GetOwner()))
 			{
 				const float OriginalHalfHeight = OriginalCapsule->GetScaledCapsuleHalfHeight();
-				const float CrouchedHalfHeight = StanceSettings->CrouchHalfHeight;
+				const float OriginalRadius = OriginalCapsule->GetScaledCapsuleRadius();
 
-				// Adjust the TargetHeight of all the relevant MovementModes
-				UpdateTargetHeight<UPhysicsDrivenWalkingMode>(MoverComp, OriginalHalfHeight, CrouchedHalfHeight, false);
-				UpdateTargetHeight<UPhysicsDrivenFallingMode>(MoverComp, OriginalHalfHeight, CrouchedHalfHeight, false);
-
-				ApplyModifierMove(MoverComp, SyncState);
+				ApplyModifierMove(MoverComp, SyncState, OriginalHalfHeight, OriginalHalfHeight, OriginalRadius);
 
 				RevertMovementSettings(MoverComp);
 			}
@@ -222,36 +195,51 @@ void FPhysicsStanceModifier::OnPostSimulationTick(const FStanceModifier* Modifie
 	{
 		if (UCapsuleComponent* CapsuleComponent = Cast<UCapsuleComponent>(UpdatedPrimitive))
 		{
-			const UCapsuleComponent* OriginalCapsule = UMovementUtils::GetOriginalComponentType<UCapsuleComponent>(MoverComp->GetOwner());
-
-			if (OriginalCapsule)
+			if (const UCapsuleComponent* OriginalCapsule = UMovementUtils::GetOriginalComponentType<UCapsuleComponent>(MoverComp->GetOwner()))
 			{
 				if (Modifier && bIsCrouching && !bPostProcessed) // Crouching
 				{
-					const float OldHalfHeight = OriginalCapsule->GetScaledCapsuleHalfHeight();
-					const float NewHalfHeight = StanceSettings->CrouchHalfHeight;
-					const float HalfHeightDifference = NewHalfHeight - OldHalfHeight;
-					const float ColliderOffset = -MoverComp->GetPrimaryVisualComponent()->GetRelativeLocation().Z - OldHalfHeight;
-	
-					CapsuleComponent->SetCapsuleSize(CapsuleComponent->GetUnscaledCapsuleRadius(), NewHalfHeight);
-					MoverComp->GetPrimaryVisualComponent()->SetRelativeLocation(-(NewHalfHeight + ColliderOffset) * FVector::UpVector);
-	
 					bStanceChanged = true;
 					bPostProcessed = true;
 				}
 				else if (!Modifier && !bIsCrouching && bPostProcessed) // Uncrouching
 				{
-					const float OldHalfHeight = StanceSettings->CrouchHalfHeight;
-					const float NewHalfHeight = OriginalCapsule->GetScaledCapsuleHalfHeight();
-					const float HalfHeightDifference = NewHalfHeight - OldHalfHeight;
-					const float ColliderOffset = -MoverComp->GetPrimaryVisualComponent()->GetRelativeLocation().Z - OldHalfHeight;
-	
-					CapsuleComponent->SetCapsuleSize(CapsuleComponent->GetUnscaledCapsuleRadius(), NewHalfHeight);
-					MoverComp->GetPrimaryVisualComponent()->SetRelativeLocation(-(NewHalfHeight + ColliderOffset) * FVector::UpVector);
-	
 					bStanceChanged = true;
 					bPostProcessed = false;
 				}
+
+				if (bStanceChanged)
+				{
+					float NewRadius = OriginalCapsule->GetScaledCapsuleRadius();
+					float OriginalHalfHeight = OriginalCapsule->GetScaledCapsuleHalfHeight();
+					float NewHalfHeight = bIsCrouching ? StanceSettings->CrouchHalfHeight : OriginalCapsule->GetScaledCapsuleHalfHeight();
+					float TargetHeight = NewHalfHeight;
+
+					if (const UBaseMovementMode* MovementMode = MoverComp->FindMovementModeByName(MoverComp->GetSyncState().MovementMode))
+					{
+						if (const IPhysicsCharacterMovementModeInterface* PhysicsMode = CastChecked<IPhysicsCharacterMovementModeInterface>(MovementMode))
+						{
+							TargetHeight = PhysicsMode->GetTargetHeight();
+						}
+					}
+
+					const Chaos::FCapsule GeneratedCapsule = GenerateNewCapsule(NewHalfHeight, NewRadius, TargetHeight - OriginalHalfHeight, TargetHeight);
+					Chaos::FImplicitObjectPtr NewCapsule(new Chaos::FCapsule(GeneratedCapsule));
+
+					if (NewCapsule.IsValid())
+					{
+						if (IPhysicsComponent* PhysicsComponent = Cast<IPhysicsComponent>(UpdatedPrimitive))
+						{
+							if (Chaos::FPhysicsObject* PhysicsObject = PhysicsComponent->GetPhysicsObjectByName(NAME_None))
+							{
+								if (Chaos::FPBDRigidParticle* Particle = FPhysicsObjectExternalInterface::LockRead(PhysicsObject)->GetRigidParticle(PhysicsObject))
+								{
+									Particle->SetGeometry(NewCapsule);
+								}
+							}
+						}
+					}
+				}
 			}
 		}
 	}
diff --git a/Engine/Plugins/Experimental/Mover/Source/Mover/Public/PhysicsMover/MovementModifiers/PhysicsStanceModifier.h b/Engine/Plugins/Experimental/Mover/Source/Mover/Public/PhysicsMover/MovementModifiers/PhysicsStanceModifier.h
--- a/Engine/Plugins/Experimental/Mover/Source/Mover/Public/PhysicsMover/MovementModifiers/PhysicsStanceModifier.h
+++ b/Engine/Plugins/Experimental/Mover/Source/Mover/Public/PhysicsMover/MovementModifiers/PhysicsStanceModifier.h
@@ -39,23 +39,12 @@ struct FPhysicsStanceModifier : public FStanceModifier
 	UE_API virtual void OnEnd(UMoverComponent* MoverComp, const FMoverTimeStep& TimeStep, const FMoverSyncState& SyncState, const FMoverAuxStateContext& AuxState) override;
 
 protected:
+	void ApplyModifierMove(UMoverComponent* MoverComp, const FMoverSyncState& InSyncState, float OriginalHalfHeight, float NewHalfHeight, float NewRadius);
+	
+	void UpdateCapsule(Chaos::FPBDRigidParticleHandle* ParticleHandle, float NewHalfHeight, float NewRadius, float TargetHeight, float GroundClearance);
 
-	template <class T>
-	void UpdateTargetHeight(UMoverComponent* MoverComp, float OriginalHalfHeight, float CrouchedHalfHeight, bool bCrouching)
-	{
-		if (T* MovementMode = MoverComp->FindMode_Mutable<T>())
-		{
-			if (MovementMode->template Implements<UPhysicsCharacterMovementModeInterface>())
-			{
-				const float OriginalTargetHeight = GetDefault<T>(MovementMode->GetClass())->GetTargetHeight();
-				const float Offset = OriginalTargetHeight - OriginalHalfHeight;
-				
-				MovementMode->SetTargetHeightOverride(Offset + (bCrouching ? CrouchedHalfHeight : OriginalHalfHeight));
-			}
-		}
-	}
-
-	void ApplyModifierMove(UMoverComponent* MoverComp, const FMoverSyncState& InSyncState);
+private:
+	static const Chaos::FCapsule GenerateNewCapsule(float NewHalfHeight, float NewRadius, float NewGroundClearance, float TargetHeight);
 };
 
 #undef UE_API
