Fix for smooth walking mode overshooting target when using directional acceleration

diff --git a/Engine/Plugins/Experimental/Mover/Source/Mover/Private/DefaultMovementSet/Modes/SmoothWalkingMode.cpp b/Engine/Plugins/Experimental/Mover/Source/Mover/Private/DefaultMovementSet/Modes/SmoothWalkingMode.cpp
--- a/Engine/Plugins/Experimental/Mover/Source/Mover/Private/DefaultMovementSet/Modes/SmoothWalkingMode.cpp
+++ b/Engine/Plugins/Experimental/Mover/Source/Mover/Private/DefaultMovementSet/Modes/SmoothWalkingMode.cpp
@@ -104,7 +104,7 @@ void USmoothWalkingMode::GenerateWalkMove_Implementation(FMoverTickStartData& St
 	// 
 	// We then clamp this velocity to make it no larger than the previous intermediate velocity or the total desired velocity this stops the 
 	// DirectionalAcceleration adding velocity to the system and infinitely speeding up the character
-	FVector NextVelocity = VelocityDifference.Length() > DesiredAcceleration.Length() * DeltaSeconds ?
+	FVector NextVelocity = VelocityDifference.Dot(DesiredAcceleration * DeltaSeconds) < VelocityDifference.SquaredLength() ?
 		SpringState.IntermediateVelocity + DesiredAcceleration * DeltaSeconds : DesiredVelocity;
 
 	NextVelocity = NextVelocity.GetClampedToMaxSize(FMath::Max(PreviousVelocityLength, DesiredVelocity.Length()));
@@ -119,8 +119,7 @@ void USmoothWalkingMode::GenerateWalkMove_Implementation(FMoverTickStartData& St
 	// avoid lagging behind
 	const float LagSeconds = DeltaSeconds + (VelocitySmoothingCompensation * VelocitySmoothingTime);
 
-	FVector TrackVelocity =
-		VelocityDifference.Length() > DesiredAcceleration.Length() * LagSeconds ?
+	FVector TrackVelocity = VelocityDifference.Dot(DesiredAcceleration * LagSeconds) < VelocityDifference.SquaredLength() ?
 		SpringState.IntermediateVelocity + DesiredAcceleration * LagSeconds : DesiredVelocity;
 
 	TrackVelocity = TrackVelocity.GetClampedToMaxSize(FMath::Max(PreviousVelocityLength, DesiredVelocity.Length()));
