[Mover] PhysicsStanceMod non-unity tests fix

diff --git a/Engine/Plugins/Experimental/Mover/Source/Mover/Private/PhysicsMover/MovementModifiers/PhysicsStanceModifier.cpp b/Engine/Plugins/Experimental/Mover/Source/Mover/Private/PhysicsMover/MovementModifiers/PhysicsStanceModifier.cpp
--- a/Engine/Plugins/Experimental/Mover/Source/Mover/Private/PhysicsMover/MovementModifiers/PhysicsStanceModifier.cpp
+++ b/Engine/Plugins/Experimental/Mover/Source/Mover/Private/PhysicsMover/MovementModifiers/PhysicsStanceModifier.cpp
@@ -7,6 +7,7 @@
 #include "Components/CapsuleComponent.h"
 #include "DefaultMovementSet/CharacterMoverComponent.h"
 #include "DefaultMovementSet/Settings/StanceSettings.h"
+#include "Engine/World.h"
 #include "MoveLibrary/MovementUtils.h"
 #include "PhysicsEngine/PhysicsObjectExternalInterface.h"
 #include "PhysicsMover/Modes/PhysicsDrivenWalkingMode.h"
