Fix `UMovementModeStateMachine` under transactions by using transactionally-safe RWLocks.

diff --git a/Engine/Plugins/Experimental/Mover/Source/Mover/Private/MovementModeStateMachine.cpp b/Engine/Plugins/Experimental/Mover/Source/Mover/Private/MovementModeStateMachine.cpp
--- a/Engine/Plugins/Experimental/Mover/Source/Mover/Private/MovementModeStateMachine.cpp
+++ b/Engine/Plugins/Experimental/Mover/Source/Mover/Private/MovementModeStateMachine.cpp
@@ -484,22 +484,22 @@ void UMovementModeStateMachine::OnSimulationRollback(const FMoverSyncState* Sync
 	}
 
 	{
-		FRWScopeLock QueueLock(LayeredMoveQueueLock, SLT_Write);
+		UE::TRWScopeLock QueueLock(LayeredMoveQueueLock, SLT_Write);
 		QueuedLayeredMoves.Empty();
 	}
 
 	{
-		FRWScopeLock QueueLock(LayeredMoveQueueLock, SLT_Write);
+		UE::TRWScopeLock QueueLock(LayeredMoveQueueLock, SLT_Write);
 		QueuedLayeredMoveInstances.Empty();
 	}
 	
 	{
-		FRWScopeLock QueueLock(InstantEffectsQueueLock, SLT_Write);
+		UE::TRWScopeLock QueueLock(InstantEffectsQueueLock, SLT_Write);
 		QueuedInstantEffects.Empty();
 	}
 
 	{
-		FRWScopeLock QueueLock(ModifiersQueueLock, SLT_Write);
+		UE::TRWScopeLock QueueLock(ModifiersQueueLock, SLT_Write);
 		QueuedMovementModifiers.Empty();
 	}
 }
@@ -527,7 +527,7 @@ const UBaseMovementMode* UMovementModeStateMachine::FindMovementMode(FName ModeN
 
 void UMovementModeStateMachine::QueueLayeredMove(TSharedPtr<FLayeredMoveBase> Move)
 {
-	FRWScopeLock QueueLock(LayeredMoveQueueLock, SLT_Write);
+	UE::TRWScopeLock QueueLock(LayeredMoveQueueLock, SLT_Write);
 	QueuedLayeredMoves.Add(Move);
 }
 
@@ -544,13 +544,13 @@ void UMovementModeStateMachine::QueueActiveLayeredMove(const TSharedPtr<FLayered
 
 void UMovementModeStateMachine::QueueInstantMovementEffect(const FScheduledInstantMovementEffect& ScheduledInstantMovementEffect)
 {
-	FRWScopeLock QueueLock(InstantEffectsQueueLock, SLT_Write);
+	UE::TRWScopeLock QueueLock(InstantEffectsQueueLock, SLT_Write);
 	QueuedInstantEffects.Add(ScheduledInstantMovementEffect);
 }
 
 void UMovementModeStateMachine::QueueInstantMovementEffects(const TArray<FScheduledInstantMovementEffect>& ScheduledInstantMovementEffects)
 {
-	FRWScopeLock QueueLock(InstantEffectsQueueLock, SLT_Write);
+	UE::TRWScopeLock QueueLock(InstantEffectsQueueLock, SLT_Write);
 	for (const FScheduledInstantMovementEffect& ScheduledInstantMovementEffect : ScheduledInstantMovementEffects)
 	{
 		QueuedInstantEffects.Add(ScheduledInstantMovementEffect);
@@ -561,7 +561,7 @@ FMovementModifierHandle UMovementModeStateMachine::QueueMovementModifier(TShared
 {
 	if (ensure(Modifier.IsValid()))
 	{
-		FRWScopeLock QueueLock(ModifiersQueueLock, SLT_Write);
+		UE::TRWScopeLock QueueLock(ModifiersQueueLock, SLT_Write);
 
 		QueuedMovementModifiers.Add(Modifier);
 		Modifier->GenerateHandle();
@@ -575,7 +575,7 @@ FMovementModifierHandle UMovementModeStateMachine::QueueMovementModifier(TShared
 void UMovementModeStateMachine::CancelModifierFromHandle(FMovementModifierHandle ModifierHandle)
 {
 	{
-		FRWScopeLock QueueLock(ModifiersQueueLock, SLT_Write);
+		UE::TRWScopeLock QueueLock(ModifiersQueueLock, SLT_Write);
 		QueuedMovementModifiers.RemoveAll([ModifierHandle, this]
 		(const TSharedPtr<FMovementModifierBase>& Modifier)
 		{
@@ -596,7 +596,7 @@ void UMovementModeStateMachine::CancelModifierFromHandle(FMovementModifierHandle
 	}
 	
 	{
-		FRWScopeLock QueueLock(ModifierCancelQueueLock, SLT_Write);
+		UE::TRWScopeLock QueueLock(ModifierCancelQueueLock, SLT_Write);
 		ModifiersToCancel.Add(ModifierHandle);
 	}
 }
@@ -631,7 +631,7 @@ void UMovementModeStateMachine::CancelFeaturesWithTag(FGameplayTag TagToCancel,
 {
 	// Cancel all matching queued movement features
 	{
-		FRWScopeLock QueueLock(ModifiersQueueLock, SLT_Write);
+		UE::TRWScopeLock QueueLock(ModifiersQueueLock, SLT_Write);
 		QueuedMovementModifiers.RemoveAll([TagToCancel, bRequireExactMatch] (const TSharedPtr<FMovementModifierBase>& Modifier)
 			{
 				return (Modifier.IsValid() && Modifier->HasGameplayTag(TagToCancel, bRequireExactMatch));
@@ -639,7 +639,7 @@ void UMovementModeStateMachine::CancelFeaturesWithTag(FGameplayTag TagToCancel,
 	}
 
 	{
-		FRWScopeLock QueueLock(LayeredMoveQueueLock, SLT_Write);
+		UE::TRWScopeLock QueueLock(LayeredMoveQueueLock, SLT_Write);
 		QueuedLayeredMoves.RemoveAll([TagToCancel, bRequireExactMatch](const TSharedPtr<FLayeredMoveBase>& LayeredMove)
 			{
 				return (LayeredMove.IsValid() && LayeredMove->HasGameplayTag(TagToCancel, bRequireExactMatch));
@@ -647,7 +647,7 @@ void UMovementModeStateMachine::CancelFeaturesWithTag(FGameplayTag TagToCancel,
 	}
 
 	{
-		FRWScopeLock QueueLock(LayeredMoveQueueLock, SLT_Write);
+		UE::TRWScopeLock QueueLock(LayeredMoveQueueLock, SLT_Write);
 		QueuedLayeredMoveInstances.RemoveAll([TagToCancel, bRequireExactMatch](const TSharedPtr<FLayeredMoveInstance>& LayeredMoveInstance)
 			{
 				return (LayeredMoveInstance.IsValid() && LayeredMoveInstance->HasGameplayTag(TagToCancel, bRequireExactMatch));
@@ -658,7 +658,7 @@ void UMovementModeStateMachine::CancelFeaturesWithTag(FGameplayTag TagToCancel,
 
 	// Request cancelation of any matching ACTIVE movement features during the next simulation tick
 	{
-		FRWScopeLock QueueLock(ModifierCancelQueueLock, SLT_Write);
+		UE::TRWScopeLock QueueLock(ModifierCancelQueueLock, SLT_Write);
 		TagCancellationRequests.Add(TPair<FGameplayTag, bool>(TagToCancel, bRequireExactMatch));
 	}
 
@@ -717,7 +717,7 @@ void UMovementModeStateMachine::AdvanceToNextMode()
 
 void UMovementModeStateMachine::FlushQueuedMovesToGroup(FLayeredMoveGroup& Group)
 {
-	FRWScopeLock QueueLock(LayeredMoveQueueLock, SLT_Write);
+	UE::TRWScopeLock QueueLock(LayeredMoveQueueLock, SLT_Write);
 
 	if (!QueuedLayeredMoves.IsEmpty())
 	{
@@ -732,7 +732,7 @@ void UMovementModeStateMachine::FlushQueuedMovesToGroup(FLayeredMoveGroup& Group
 
 void UMovementModeStateMachine::ActivateQueuedMoves(FLayeredMoveInstanceGroup& Group)
 {
-	FRWScopeLock QueueLock(LayeredMoveQueueLock, SLT_Write);
+	UE::TRWScopeLock QueueLock(LayeredMoveQueueLock, SLT_Write);
 
 	if (!QueuedLayeredMoveInstances.IsEmpty())
 	{
@@ -747,7 +747,7 @@ void UMovementModeStateMachine::ActivateQueuedMoves(FLayeredMoveInstanceGroup& G
 
 void UMovementModeStateMachine::FlushQueuedModifiersToGroup(FMovementModifierGroup& ModifierGroup)
 {
-	FRWScopeLock QueueLock(ModifiersQueueLock, SLT_Write);
+	UE::TRWScopeLock QueueLock(ModifiersQueueLock, SLT_Write);
 
 	if (!QueuedMovementModifiers.IsEmpty())
 	{
@@ -762,7 +762,7 @@ void UMovementModeStateMachine::FlushQueuedModifiersToGroup(FMovementModifierGro
 
 void UMovementModeStateMachine::FlushModifierCancellationsToGroup(FMovementModifierGroup& ActiveModifierGroup)
 {
-	FRWScopeLock QueueLock(ModifierCancelQueueLock, SLT_Write);
+	UE::TRWScopeLock QueueLock(ModifierCancelQueueLock, SLT_Write);
 	
 	for (FMovementModifierHandle HandleToCancel : ModifiersToCancel)
 	{
@@ -774,7 +774,7 @@ void UMovementModeStateMachine::FlushModifierCancellationsToGroup(FMovementModif
 
 void UMovementModeStateMachine::FlushTagCancellationsToSyncState(FMoverSyncState& SyncState)
 {
-	FRWScopeLock QueueLock(TagCancellationRequestsLock, SLT_Write);
+	UE::TRWScopeLock QueueLock(TagCancellationRequestsLock, SLT_Write);
 
 	for (TPair<FGameplayTag, bool> TagCancelRequest : TagCancellationRequests)
 	{
@@ -789,7 +789,7 @@ void UMovementModeStateMachine::FlushTagCancellationsToSyncState(FMoverSyncState
 void UMovementModeStateMachine::RollbackModifiers(const FMoverSyncState* InvalidSyncState, const FMoverSyncState* SyncState, const FMoverAuxStateContext* InvalidAuxState, const FMoverAuxStateContext* AuxState)
 {
 	{
-		FRWScopeLock QueueLock(ModifiersQueueLock, SLT_Write);
+		UE::TRWScopeLock QueueLock(ModifiersQueueLock, SLT_Write);
 		QueuedMovementModifiers.Empty();
 	}
 	
@@ -841,7 +841,7 @@ void UMovementModeStateMachine::RollbackModifiers(const FMoverSyncState* Invalid
 
 bool UMovementModeStateMachine::HasAnyInstantEffectsQueued() const 
 {
-	FRWScopeLock QueueLock(InstantEffectsQueueLock, SLT_ReadOnly);
+	UE::TRWScopeLock QueueLock(InstantEffectsQueueLock, SLT_ReadOnly);
 
 	return !QueuedInstantEffects.IsEmpty();
 }
@@ -850,7 +850,7 @@ bool UMovementModeStateMachine::ApplyInstantEffects(FApplyMovementEffectParams&
 {
 	bool bInstantMovementEffectApplied = false;
 
-	FRWScopeLock QueueLock(InstantEffectsQueueLock, SLT_Write);
+	UE::TRWScopeLock QueueLock(InstantEffectsQueueLock, SLT_Write);
 
 	for (int EffectIndex = 0; EffectIndex < QueuedInstantEffects.Num(); )
 	{
diff --git a/Engine/Plugins/Experimental/Mover/Source/Mover/Public/MovementModeStateMachine.h b/Engine/Plugins/Experimental/Mover/Source/Mover/Public/MovementModeStateMachine.h
--- a/Engine/Plugins/Experimental/Mover/Source/Mover/Public/MovementModeStateMachine.h
+++ b/Engine/Plugins/Experimental/Mover/Source/Mover/Public/MovementModeStateMachine.h
@@ -3,6 +3,7 @@
 #pragma once
 
 #include "CoreMinimal.h"
+#include "Misc/TransactionallySafeRWLock.h"
 #include "Templates/SubclassOf.h"
 #include "MovementMode.h"
 #include "InstantMovementEffect.h"
@@ -113,11 +114,11 @@ protected:
 
 private:
 	// Locks for thread safety on queueing mechanisms
-	mutable FRWLock LayeredMoveQueueLock;
-	mutable FRWLock InstantEffectsQueueLock;
-	mutable FRWLock ModifiersQueueLock;
-	mutable FRWLock ModifierCancelQueueLock;
-	mutable FRWLock TagCancellationRequestsLock;
+	mutable FTransactionallySafeRWLock LayeredMoveQueueLock;
+	mutable FTransactionallySafeRWLock InstantEffectsQueueLock;
+	mutable FTransactionallySafeRWLock ModifiersQueueLock;
+	mutable FTransactionallySafeRWLock ModifierCancelQueueLock;
+	mutable FTransactionallySafeRWLock TagCancellationRequestsLock;
 
 	UE_API void ConstructDefaultModes();
 	UE_API void AdvanceToNextMode();
