Mover: adding simple montage (fixed playback) replication support for the character-focused MoverComponent, and adding this to the anim root motion layered move.

#FailedValidations Preflight

diff --git a/Engine/Plugins/Experimental/Mover/Source/Mover/Private/DefaultMovementSet/CharacterMoverComponent.cpp b/Engine/Plugins/Experimental/Mover/Source/Mover/Private/DefaultMovementSet/CharacterMoverComponent.cpp
--- a/Engine/Plugins/Experimental/Mover/Source/Mover/Private/DefaultMovementSet/CharacterMoverComponent.cpp
+++ b/Engine/Plugins/Experimental/Mover/Source/Mover/Private/DefaultMovementSet/CharacterMoverComponent.cpp
@@ -2,6 +2,9 @@
 
 #include "DefaultMovementSet/CharacterMoverComponent.h"
 
+#include "Animation/AnimInstance.h"
+#include "Animation/AnimMontage.h"
+#include "Components/SkeletalMeshComponent.h"
 #include "DefaultMovementSet/InstantMovementEffects/BasicInstantMovementEffects.h"
 #include "DefaultMovementSet/Modes/FallingMode.h"
 #include "DefaultMovementSet/Modes/FlyingMode.h"
@@ -10,6 +13,7 @@
 #include "MoveLibrary/FloorQueryUtils.h"
 #include "MoveLibrary/MovementUtils.h"
 
+
 #include UE_INLINE_GENERATED_CPP_BY_NAME(CharacterMoverComponent)
 
 UCharacterMoverComponent::UCharacterMoverComponent()
@@ -27,6 +31,8 @@ void UCharacterMoverComponent::BeginPlay()
 	Super::BeginPlay();
 
 	OnHandlerSettingChanged();
+
+	OnPostFinalize.AddDynamic(this, &UCharacterMoverComponent::OnMoverPostFinalize);
 }
 
 bool UCharacterMoverComponent::GetHandleJump() const
@@ -198,6 +204,11 @@ void UCharacterMoverComponent::OnMoverPreSimulationTick(const FMoverTimeStep& Ti
 	}
 }
 
+void UCharacterMoverComponent::OnMoverPostFinalize(const FMoverSyncState& SyncState, const FMoverAuxStateContext& AuxState)
+{
+	UpdateSyncedMontageState(GetLastTimeStep(), SyncState, AuxState);
+}
+
 void UCharacterMoverComponent::OnHandlerSettingChanged()
 {
 	const bool bIsHandlingAnySettings = bHandleJump || bHandleStanceChanges;
@@ -211,3 +222,69 @@ void UCharacterMoverComponent::OnHandlerSettingChanged()
 		OnPreSimulationTick.RemoveDynamic(this, &UCharacterMoverComponent::OnMoverPreSimulationTick);
 	}
 }
+
+void UCharacterMoverComponent::UpdateSyncedMontageState(const FMoverTimeStep& TimeStep, const FMoverSyncState& SyncState, const FMoverAuxStateContext& AuxState)
+{
+	if (GetOwnerRole() == ROLE_SimulatedProxy)
+	{
+		const FLayeredMove_MontageStateProvider* MontageStateProvider = static_cast<const FLayeredMove_MontageStateProvider*>(SyncState.LayeredMoves.FindActiveMove(FLayeredMove_MontageStateProvider::StaticStruct()));
+
+		bool bShouldStopSyncedMontage = false;
+		bool bShouldStartNewMontage = false;
+		FMoverAnimMontageState NewMontageState;
+
+		if (SyncedMontageState.Montage)
+		{
+			if (MontageStateProvider)
+			{
+				NewMontageState = MontageStateProvider->GetMontageState();
+
+				if (NewMontageState.Montage != SyncedMontageState.Montage)
+				{
+					bShouldStartNewMontage = true;
+					bShouldStopSyncedMontage = true;
+				}
+			}
+			else
+			{
+				bShouldStopSyncedMontage = true;
+			}
+		}
+		else // We aren't actively syncing a montage state yet
+		{
+			if (MontageStateProvider)
+			{
+				// We have just received a montage state to sync against
+				NewMontageState = MontageStateProvider->GetMontageState();
+				bShouldStartNewMontage = true;
+			}
+		}
+
+		if (bShouldStopSyncedMontage || bShouldStartNewMontage)
+		{
+			const USkeletalMeshComponent* MeshComp = Cast<USkeletalMeshComponent>(GetPrimaryVisualComponent());
+			UAnimInstance* MeshAnimInstance = MeshComp ? MeshComp->GetAnimInstance() : nullptr;
+
+			if (bShouldStopSyncedMontage)
+			{
+				if (MeshAnimInstance)
+				{
+					MeshAnimInstance->Montage_Stop(SyncedMontageState.Montage->GetDefaultBlendOutTime(), SyncedMontageState.Montage);
+				}
+
+				SyncedMontageState.Reset();
+			}
+
+			if (bShouldStartNewMontage && NewMontageState.Montage && MeshAnimInstance)
+			{
+				const float StartPosition = NewMontageState.StartingMontagePosition + (FMath::Max(0.0f, (TimeStep.BaseSimTimeMs - MontageStateProvider->StartSimTimeMs) * 0.001f) * NewMontageState.PlayRate);
+				const float PlaySeconds = MeshAnimInstance->Montage_Play(NewMontageState.Montage, NewMontageState.PlayRate, EMontagePlayReturnType::MontageLength, StartPosition);
+
+				if (PlaySeconds > 0.0f)
+				{
+					SyncedMontageState = NewMontageState;	// only consider us sync'd if the montage actually started
+				}
+			}
+		}
+	}
+}
\ No newline at end of file
diff --git a/Engine/Plugins/Experimental/Mover/Source/Mover/Private/DefaultMovementSet/LayeredMoves/AnimRootMotionLayeredMove.cpp b/Engine/Plugins/Experimental/Mover/Source/Mover/Private/DefaultMovementSet/LayeredMoves/AnimRootMotionLayeredMove.cpp
--- a/Engine/Plugins/Experimental/Mover/Source/Mover/Private/DefaultMovementSet/LayeredMoves/AnimRootMotionLayeredMove.cpp
+++ b/Engine/Plugins/Experimental/Mover/Source/Mover/Private/DefaultMovementSet/LayeredMoves/AnimRootMotionLayeredMove.cpp
@@ -20,14 +20,13 @@ FAutoConsoleVariable CVarLogAnimRootMotionSteps(
 	ECVF_Cheat);
 #endif	// !UE_BUILD_SHIPPING
 
+
+
 FLayeredMove_AnimRootMotion::FLayeredMove_AnimRootMotion()
 {
 	DurationMs = 0.f;
 	MixMode = EMoveMixMode::OverrideAll;
 
-	Montage = nullptr;
-	StartingMontagePosition = 0.f;
-	PlayRate = 1.f;
 }
 
 bool FLayeredMove_AnimRootMotion::GenerateMove(const FMoverTickStartData& SimState, const FMoverTimeStep& TimeStep, const UMoverComponent* MoverComp, UMoverBlackboard* SimBlackboard, FProposedMove& OutProposedMove)
@@ -41,7 +40,7 @@ bool FLayeredMove_AnimRootMotion::GenerateMove(const FMoverTickStartData& SimSta
 		{
 			if (const UAnimInstance* MeshAnimInstance = MeshComp->GetAnimInstance())
 			{
-				bIsMontageStillPlaying = MeshAnimInstance->Montage_IsPlaying(Montage);
+				bIsMontageStillPlaying = MontageState.Montage && MeshAnimInstance->Montage_IsPlaying(MontageState.Montage);
 			}
 		}
 
@@ -56,6 +55,11 @@ bool FLayeredMove_AnimRootMotion::GenerateMove(const FMoverTickStartData& SimSta
 
 	const FMoverDefaultSyncState* SyncState = SimState.SyncState.SyncStateCollection.FindDataByType<FMoverDefaultSyncState>();
 
+	if (SyncState == nullptr)
+	{
+		return false;
+	}
+
 	// First pass simply samples based on the duration. For long animations, this has the potential to diverge.
 	// Future improvements could include:
 	//     - speeding up or slowing down slightly to match the associated montage instance
@@ -63,19 +67,19 @@ bool FLayeredMove_AnimRootMotion::GenerateMove(const FMoverTickStartData& SimSta
 	
 	// Note that Montage 'position' equates to seconds when PlayRate is 1
 	const double SecondsSinceMontageStarted = (TimeStep.BaseSimTimeMs - StartSimTimeMs) / 1000.0;
-	const double ScaledSecondsSinceMontageStarted = SecondsSinceMontageStarted * PlayRate;
+	const double ScaledSecondsSinceMontageStarted = SecondsSinceMontageStarted * MontageState.PlayRate;
 
-	const float ExtractionStartPosition = StartingMontagePosition + ScaledSecondsSinceMontageStarted;
-	const float ExtractionEndPosition   = ExtractionStartPosition + (DeltaSeconds * PlayRate);
+	const float ExtractionStartPosition = MontageState.StartingMontagePosition + ScaledSecondsSinceMontageStarted;
+	const float ExtractionEndPosition   = ExtractionStartPosition + (DeltaSeconds * MontageState.PlayRate);
 
 	// Read the local transform directly from the montage
-	const FTransform LocalRootMotion = UMotionWarpingUtilities::ExtractRootMotionFromAnimation(Montage, ExtractionStartPosition, ExtractionEndPosition);
+	const FTransform LocalRootMotion = MontageState.Montage ? UMotionWarpingUtilities::ExtractRootMotionFromAnimation(MontageState.Montage, ExtractionStartPosition, ExtractionEndPosition) : FTransform::Identity;
 
 	FMotionWarpingUpdateContext WarpingContext;
-	WarpingContext.Animation = Montage;
+	WarpingContext.Animation = MontageState.Montage;
 	WarpingContext.CurrentPosition = ExtractionEndPosition;
 	WarpingContext.PreviousPosition = ExtractionStartPosition;
-	WarpingContext.PlayRate = PlayRate;
+	WarpingContext.PlayRate = MontageState.PlayRate;
 	WarpingContext.Weight = 1.f;
 
 	// Note that we're forcing the use of the sync state's actor transform data. This is necessary when the movement simulation 
@@ -87,8 +91,11 @@ bool FLayeredMove_AnimRootMotion::GenerateMove(const FMoverTickStartData& SimSta
 	OutProposedMove.MixMode = MixMode;
 
 	// Convert the transform into linear and angular velocities
-	OutProposedMove.LinearVelocity    = WorldSpaceRootMotion.GetTranslation() / DeltaSeconds;
-	OutProposedMove.AngularVelocity   = WorldSpaceRootMotion.GetRotation().Rotator() * (1.f / DeltaSeconds);
+	if (DeltaSeconds > UE_KINDA_SMALL_NUMBER)
+	{
+		OutProposedMove.LinearVelocity    = WorldSpaceRootMotion.GetTranslation() / DeltaSeconds;
+		OutProposedMove.AngularVelocity   = WorldSpaceRootMotion.GetRotation().Rotator() * (1.f / DeltaSeconds);
+	}
 
 #if !UE_BUILD_SHIPPING
 	UE_CLOG(CVarLogAnimRootMotionSteps->GetBool(), LogMover, Log, TEXT("AnimRootMotion. SimF %i (dt %.3f) Range [%.3f, %.3f] => LocalT: %s (WST: %s)  Vel: %.3f"),
@@ -109,9 +116,7 @@ void FLayeredMove_AnimRootMotion::NetSerialize(FArchive& Ar)
 {
 	Super::NetSerialize(Ar);
 
-	Ar << Montage;
-	Ar << StartingMontagePosition;
-	Ar << PlayRate;
+	MontageState.NetSerialize(Ar);
 }
 
 UScriptStruct* FLayeredMove_AnimRootMotion::GetScriptStruct() const
@@ -128,3 +133,8 @@ void FLayeredMove_AnimRootMotion::AddReferencedObjects(class FReferenceCollector
 {
 	Super::AddReferencedObjects(Collector);
 }
+
+FMoverAnimMontageState FLayeredMove_AnimRootMotion::GetMontageState() const
+{
+	return MontageState;
+}
\ No newline at end of file
diff --git a/Engine/Plugins/Experimental/Mover/Source/Mover/Private/DefaultMovementSet/LayeredMoves/MontageStateProvider.cpp b/Engine/Plugins/Experimental/Mover/Source/Mover/Private/DefaultMovementSet/LayeredMoves/MontageStateProvider.cpp
new file mode 100644
--- /dev/null
+++ b/Engine/Plugins/Experimental/Mover/Source/Mover/Private/DefaultMovementSet/LayeredMoves/MontageStateProvider.cpp
@@ -0,0 +1,46 @@
+// Copyright Epic Games, Inc. All Rights Reserved.
+
+#include "DefaultMovementSet/LayeredMoves/MontageStateProvider.h"
+
+#include UE_INLINE_GENERATED_CPP_BY_NAME(MontageStateProvider)
+
+
+void FMoverAnimMontageState::Reset()
+{
+	Montage = nullptr;
+}
+
+void FMoverAnimMontageState::NetSerialize(FArchive& Ar)
+{
+	Ar << Montage;
+
+	uint8 bHasNonDefaultStartingPosition = 0;
+	uint8 bHasNonDefaultPlayRate = 0;
+
+	if (Ar.IsSaving())
+	{
+		bHasNonDefaultStartingPosition = (StartingMontagePosition != 0.0f);
+		bHasNonDefaultPlayRate = (PlayRate != 1.0f);
+	}
+
+	Ar.SerializeBits(&bHasNonDefaultStartingPosition, 1);
+	Ar.SerializeBits(&bHasNonDefaultPlayRate, 1);
+
+	if (bHasNonDefaultStartingPosition)
+	{
+		Ar << StartingMontagePosition;
+	}
+	else
+	{
+		StartingMontagePosition = 0.0f;
+	}
+
+	if (bHasNonDefaultPlayRate)
+	{
+		Ar << PlayRate;
+	}
+	else
+	{
+		PlayRate = 1.0f;
+	}
+}
diff --git a/Engine/Plugins/Experimental/Mover/Source/Mover/Private/MoveLibrary/PlayMoverMontageCallbackProxy.cpp b/Engine/Plugins/Experimental/Mover/Source/Mover/Private/MoveLibrary/PlayMoverMontageCallbackProxy.cpp
--- a/Engine/Plugins/Experimental/Mover/Source/Mover/Private/MoveLibrary/PlayMoverMontageCallbackProxy.cpp
+++ b/Engine/Plugins/Experimental/Mover/Source/Mover/Private/MoveLibrary/PlayMoverMontageCallbackProxy.cpp
@@ -57,9 +57,9 @@ bool UPlayMoverMontageCallbackProxy::PlayMoverMontage(
 
 				// Queue a layered move to perform the same anim root motion over the same time span
 				TSharedPtr<FLayeredMove_AnimRootMotion> AnimRootMotionMove = MakeShared<FLayeredMove_AnimRootMotion>();
-				AnimRootMotionMove->Montage = MontageToPlay;
-				AnimRootMotionMove->PlayRate = PlayRate;
-				AnimRootMotionMove->StartingMontagePosition = StartingMontagePosition;
+				AnimRootMotionMove->MontageState.Montage = MontageToPlay;
+				AnimRootMotionMove->MontageState.PlayRate = PlayRate;
+				AnimRootMotionMove->MontageState.StartingMontagePosition = StartingMontagePosition;
 				
 				float RemainingUnscaledMontageSeconds(0.f);
 
diff --git a/Engine/Plugins/Experimental/Mover/Source/Mover/Public/DefaultMovementSet/CharacterMoverComponent.h b/Engine/Plugins/Experimental/Mover/Source/Mover/Public/DefaultMovementSet/CharacterMoverComponent.h
--- a/Engine/Plugins/Experimental/Mover/Source/Mover/Public/DefaultMovementSet/CharacterMoverComponent.h
+++ b/Engine/Plugins/Experimental/Mover/Source/Mover/Public/DefaultMovementSet/CharacterMoverComponent.h
@@ -3,6 +3,7 @@
 #pragma once
 
 #include "MoverComponent.h"
+#include "DefaultMovementSet/LayeredMoves/MontageStateProvider.h"
 #include "MovementModifiers/StanceModifier.h"
 #include "CharacterMoverComponent.generated.h"
 
@@ -16,6 +17,12 @@
 DECLARE_DYNAMIC_MULTICAST_DELEGATE_TwoParams(FMover_OnStanceChanged, EStanceMode, OldStance, EStanceMode, NewStance);
 
 
+
+/** 
+* Character Mover Component: this is a specialization of the core Mover Component that is set up with a 
+* classic character in mind. Defaults and extended functionality, such as jumping and simple montage replication, 
+* are intended to support features similar to UE's ACharacter actor type.
+*/
 UCLASS(MinimalAPI, BlueprintType, Blueprintable, meta = (BlueprintSpawnableComponent))
 class UCharacterMoverComponent : public UMoverComponent
 {
@@ -98,9 +105,14 @@ public:
 protected:
 	UFUNCTION()
 	UE_API virtual void OnMoverPreSimulationTick(const FMoverTimeStep& TimeStep, const FMoverInputCmdContext& InputCmd);
-	
+
+	UFUNCTION()
+	UE_API virtual void OnMoverPostFinalize(const FMoverSyncState& SyncState, const FMoverAuxStateContext& AuxState);
+
 	UE_API virtual void OnHandlerSettingChanged();
 
+	UE_API virtual void UpdateSyncedMontageState(const FMoverTimeStep& TimeStep, const FMoverSyncState& SyncState, const FMoverAuxStateContext& AuxState);
+
 	// ID used to keep track of the modifier responsible for crouching
 	FMovementModifierHandle StanceModifierHandle;
 
@@ -115,6 +127,10 @@ protected:
 	// Whether this component should directly handle stance changes, including crouching input
 	UPROPERTY(EditAnywhere, BlueprintGetter = GetHandleStanceChanges, BlueprintSetter = SetHandleStanceChanges, Category = "Mover|Character")
 	uint8 bHandleStanceChanges : 1 = 1;
+
+	/** Current state of replicated montage playback from an active movement mechanism (layered move, etc.) */
+	UPROPERTY(Transient)
+	FMoverAnimMontageState SyncedMontageState;
 };
 
 #undef UE_API
diff --git a/Engine/Plugins/Experimental/Mover/Source/Mover/Public/DefaultMovementSet/LayeredMoves/AnimRootMotionLayeredMove.h b/Engine/Plugins/Experimental/Mover/Source/Mover/Public/DefaultMovementSet/LayeredMoves/AnimRootMotionLayeredMove.h
--- a/Engine/Plugins/Experimental/Mover/Source/Mover/Public/DefaultMovementSet/LayeredMoves/AnimRootMotionLayeredMove.h
+++ b/Engine/Plugins/Experimental/Mover/Source/Mover/Public/DefaultMovementSet/LayeredMoves/AnimRootMotionLayeredMove.h
@@ -4,6 +4,7 @@
 
 
 #include "LayeredMove.h"
+#include "DefaultMovementSet/LayeredMoves/MontageStateProvider.h"
 #include "AnimRootMotionLayeredMove.generated.h"
 
 #define UE_API MOVER_API
@@ -16,23 +17,15 @@ class UAnimMontage;
  * is interrupted on the mesh.
  */
 USTRUCT(BlueprintType)
-struct FLayeredMove_AnimRootMotion : public FLayeredMoveBase
+struct FLayeredMove_AnimRootMotion : public FLayeredMove_MontageStateProvider
 {
 	GENERATED_BODY()
 
 	UE_API FLayeredMove_AnimRootMotion();
 	virtual ~FLayeredMove_AnimRootMotion() {}
 
-	UPROPERTY(EditAnywhere, BlueprintReadWrite, Category = Mover)
-	TObjectPtr<UAnimMontage> Montage;
-
-	// Montage position when started (in unscaled seconds). 
-	UPROPERTY(EditAnywhere, BlueprintReadWrite, Category = Mover)
-	float StartingMontagePosition;
-
-	// Rate at which this montage is intended to play
-	UPROPERTY(EditAnywhere, BlueprintReadWrite, Category = Mover)
-	float PlayRate;
+	UPROPERTY(BlueprintReadOnly, Category = Mover)
+	FMoverAnimMontageState MontageState;
 
 	// Generate a movement 
 	UE_API virtual bool GenerateMove(const FMoverTickStartData& StartState, const FMoverTimeStep& TimeStep, const UMoverComponent* MoverComp, UMoverBlackboard* SimBlackboard, FProposedMove& OutProposedMove) override;
@@ -46,6 +39,9 @@ struct FLayeredMove_AnimRootMotion : public FLayeredMoveBase
 	UE_API virtual FString ToSimpleString() const override;
 
 	UE_API virtual void AddReferencedObjects(class FReferenceCollector& Collector) override;
+
+	// FLayeredMove_MontageStateProvider
+	UE_API virtual FMoverAnimMontageState GetMontageState() const override;
 };
 
 template<>
diff --git a/Engine/Plugins/Experimental/Mover/Source/Mover/Public/DefaultMovementSet/LayeredMoves/MontageStateProvider.h b/Engine/Plugins/Experimental/Mover/Source/Mover/Public/DefaultMovementSet/LayeredMoves/MontageStateProvider.h
new file mode 100644
--- /dev/null
+++ b/Engine/Plugins/Experimental/Mover/Source/Mover/Public/DefaultMovementSet/LayeredMoves/MontageStateProvider.h
@@ -0,0 +1,51 @@
+// Copyright Epic Games, Inc. All Rights Reserved.
+
+#pragma once
+
+
+#include "LayeredMove.h"
+#include "MontageStateProvider.generated.h"
+
+#define UE_API MOVER_API
+
+class UAnimMontage;
+
+
+/** Data about montages that is replicated to simulated clients */
+USTRUCT(BlueprintType)
+struct FMoverAnimMontageState
+{
+	GENERATED_BODY()
+
+	UPROPERTY(BlueprintReadOnly, Category = Mover)
+	TObjectPtr<UAnimMontage> Montage;
+
+	// Montage position when started (in unscaled seconds). 
+	UPROPERTY(BlueprintReadOnly, Category = Mover)
+	float StartingMontagePosition = 0.0f;
+
+	// Rate at which this montage is intended to play
+	UPROPERTY(BlueprintReadOnly, Category = Mover)
+	float PlayRate = 1.0f;
+
+
+	void Reset();
+	void NetSerialize(FArchive& Ar);
+};
+
+
+
+/** Note this will become obsolete once layered move logic is represented by a uobject, allowing use of interface classes.  */
+USTRUCT(BlueprintType)
+struct FLayeredMove_MontageStateProvider : public FLayeredMoveBase
+{
+	GENERATED_BODY()
+
+	virtual FMoverAnimMontageState GetMontageState() const
+	{
+		checkNoEntry();
+		return FMoverAnimMontageState();
+	}
+};
+
+#undef UE_API
