Mover: fix rollback blackboard assert that could happen when using the Network Physics liaison

diff --git a/Engine/Plugins/Experimental/Mover/Source/Mover/Private/Backends/MoverNetworkPhysicsLiaisonBase.cpp b/Engine/Plugins/Experimental/Mover/Source/Mover/Private/Backends/MoverNetworkPhysicsLiaisonBase.cpp
--- a/Engine/Plugins/Experimental/Mover/Source/Mover/Private/Backends/MoverNetworkPhysicsLiaisonBase.cpp
+++ b/Engine/Plugins/Experimental/Mover/Source/Mover/Private/Backends/MoverNetworkPhysicsLiaisonBase.cpp
@@ -2,6 +2,7 @@
 
 #include "Backends/MoverNetworkPhysicsLiaisonBase.h"
 
+#include "MoveLibrary/RollbackBlackboard.h"
 #include "MovementModeStateMachine.h"
 #include "PBDRigidsSolver.h"
 #include "Chaos/PhysicsObjectInterface.h"
@@ -823,12 +824,16 @@ void UMoverNetworkPhysicsLiaisonComponentBase::PerformPreSimulate_Internal(const
 		MoverComp.ProcessFirstSimTickAfterRollback(TimeStep);
 	}
 
+
 	//@todo DanH: Invoking the FSM OnSimulationTick can sometimes trigger AdvanceToNextMode to happen on the PT
 	// Tick the actual simulation. This is where the proposed moves are queried and executed, affecting change to the moving actor's gameplay state and captured in the output sim state
 	FMoverTickStartData TickStartData(Input.InputCmd, Input.SyncState, FMoverAuxStateContext());
 	FMoverTickEndData TickEndData;
 	TickEndData.SyncState = TickStartData.SyncState;
 
+	MoverComp.RollbackBlackboard_InternalWrapper->BeginSimulationFrame(TimeStep);
+
+
 	MoverComp.ModeFSM->OnSimulationTick(MoverComp.UpdatedComponent, MoverComp.UpdatedCompAsPrimitive, MoverComp.GetSimBlackboard_Mutable(), TickStartData, TimeStep, TickEndData);
 
 	// Set the output sync state and fill in the movement mode
@@ -853,6 +858,9 @@ void UMoverNetworkPhysicsLiaisonComponentBase::PerformPreSimulate_Internal(const
 
 		//@todo DanH: Does the base need to concern itself with setting W on the particle?
 	}
+
+	MoverComp.RollbackBlackboard_InternalWrapper->EndSimulationFrame();
+
 }
 
 Chaos::FPhysicsObject* UMoverNetworkPhysicsLiaisonComponentBase::GetControlledPhysicsObject() const
