Mover/MoverTests: allowing montage playback data to be BP-writable so BP devs can kick off an ARM layered move.
Fixing anim root motion tests to correlate with the API change.

#rnx
#virtualized

diff --git a/Engine/Plugins/Experimental/Mover/Source/Mover/Public/DefaultMovementSet/LayeredMoves/AnimRootMotionLayeredMove.h b/Engine/Plugins/Experimental/Mover/Source/Mover/Public/DefaultMovementSet/LayeredMoves/AnimRootMotionLayeredMove.h
--- a/Engine/Plugins/Experimental/Mover/Source/Mover/Public/DefaultMovementSet/LayeredMoves/AnimRootMotionLayeredMove.h
+++ b/Engine/Plugins/Experimental/Mover/Source/Mover/Public/DefaultMovementSet/LayeredMoves/AnimRootMotionLayeredMove.h
@@ -24,7 +24,7 @@ struct FLayeredMove_AnimRootMotion : public FLayeredMove_MontageStateProvider
 	UE_API FLayeredMove_AnimRootMotion();
 	virtual ~FLayeredMove_AnimRootMotion() {}
 
-	UPROPERTY(BlueprintReadOnly, Category = Mover)
+	UPROPERTY(BlueprintReadWrite, Category = Mover)
 	FMoverAnimMontageState MontageState;
 
 	// Generate a movement 
diff --git a/Engine/Plugins/Experimental/Mover/Source/Mover/Public/DefaultMovementSet/LayeredMoves/MontageStateProvider.h b/Engine/Plugins/Experimental/Mover/Source/Mover/Public/DefaultMovementSet/LayeredMoves/MontageStateProvider.h
--- a/Engine/Plugins/Experimental/Mover/Source/Mover/Public/DefaultMovementSet/LayeredMoves/MontageStateProvider.h
+++ b/Engine/Plugins/Experimental/Mover/Source/Mover/Public/DefaultMovementSet/LayeredMoves/MontageStateProvider.h
@@ -17,15 +17,15 @@ struct FMoverAnimMontageState
 {
 	GENERATED_BODY()
 
-	UPROPERTY(BlueprintReadOnly, Category = Mover)
+	UPROPERTY(BlueprintReadWrite, Category = Mover)
 	TObjectPtr<UAnimMontage> Montage;
 
 	// Montage position when started (in unscaled seconds). 
-	UPROPERTY(BlueprintReadOnly, Category = Mover)
+	UPROPERTY(BlueprintReadWrite, Category = Mover)
 	float StartingMontagePosition = 0.0f;
 
 	// Rate at which this montage is intended to play
-	UPROPERTY(BlueprintReadOnly, Category = Mover)
+	UPROPERTY(BlueprintReadWrite, Category = Mover)
 	float PlayRate = 1.0f;
 
 
