[Mover] Consistent angular velocity when changing movement modes
Adjusted mover API to use a FVector for angular velocity instead of a FRotator. This changes some public API related to Mover.
Callers must now provide an angular velocity to FMoverDefaultSyncState::SetTransforms_WorldSpace in order to ensure that angular velocity can persist when changing mode
Fixup usages of mover
Fix a bug in mover angular velocity clamping where the limit was not respected when rotating off axis
Remove gravity orientation from angular velocity calculation

#virtualized

diff --git a/Engine/Plugins/Experimental/ChaosMover/Source/ChaosMover/Private/Backends/ChaosMoverBackend.cpp b/Engine/Plugins/Experimental/ChaosMover/Source/ChaosMover/Private/Backends/ChaosMoverBackend.cpp
--- a/Engine/Plugins/Experimental/ChaosMover/Source/ChaosMover/Private/Backends/ChaosMoverBackend.cpp
+++ b/Engine/Plugins/Experimental/ChaosMover/Source/ChaosMover/Private/Backends/ChaosMoverBackend.cpp
@@ -620,7 +620,12 @@ void UChaosMoverBackendComponent::FinalizeFrame(double ResultsTimeInMs)
 	FMoverDefaultSyncState& TransformSyncState = InterpolatedOutput.SyncState.SyncStateCollection.FindOrAddMutableDataByType<FMoverDefaultSyncState>();
 	if (Chaos::FPBDRigidParticle* Particle = GetControlledParticle())
 	{
-		TransformSyncState.SetTransforms_WorldSpace(Particle->GetX(), FRotator(Particle->GetR()), Particle->GetV(), TransformSyncState.GetMovementBase(), TransformSyncState.GetMovementBaseBoneName());
+		TransformSyncState.SetTransforms_WorldSpace(Particle->GetX(),
+			FRotator(Particle->GetR()),
+			Particle->GetV(),
+			FMath::RadiansToDegrees(Particle->GetW()),
+			TransformSyncState.GetMovementBase(),
+			TransformSyncState.GetMovementBaseBoneName());
 
 		// Make sure the move direction intent is in base space (the base quat is identity if there's no base, effectively making this a no-op)
 		TransformSyncState.MoveDirectionIntent = TransformSyncState.GetCapturedMovementBaseQuat().UnrotateVector(TransformSyncState.MoveDirectionIntent);
diff --git a/Engine/Plugins/Experimental/ChaosMover/Source/ChaosMover/Private/ChaosMoverSimulation.cpp b/Engine/Plugins/Experimental/ChaosMover/Source/ChaosMover/Private/ChaosMoverSimulation.cpp
--- a/Engine/Plugins/Experimental/ChaosMover/Source/ChaosMover/Private/ChaosMoverSimulation.cpp
+++ b/Engine/Plugins/Experimental/ChaosMover/Source/ChaosMover/Private/ChaosMoverSimulation.cpp
@@ -359,7 +359,8 @@ void UChaosMoverSimulation::PreSimulationTickCharacter(const IChaosCharacterMove
 				PreSimDefaultSyncState.SetTransforms_WorldSpace(
 					PreSimDefaultSyncState.GetLocation_WorldSpace(),
 					PreSimDefaultSyncState.GetOrientation_WorldSpace(),
-					PreSimDefaultSyncState.GetVelocity_WorldSpace() - LocalGroundVelocity);
+					PreSimDefaultSyncState.GetVelocity_WorldSpace() - LocalGroundVelocity,
+					PreSimDefaultSyncState.GetAngularVelocityDegrees_WorldSpace());
 			}
 		}
 	}
@@ -394,7 +395,7 @@ void UChaosMoverSimulation::OnPreSimulationTick(const FMoverTimeStep& TimeStep,
 	PreSimDefaultSyncState.SetMovementBase(nullptr);
 	if (const Chaos::FPBDRigidParticleHandle* ParticleHandle = GetControlledParticle())
 	{
-		PreSimDefaultSyncState.SetTransforms_WorldSpace(ParticleHandle->GetX(), FRotator(ParticleHandle->GetR()), ParticleHandle->GetV());
+		PreSimDefaultSyncState.SetTransforms_WorldSpace(ParticleHandle->GetX(), FRotator(ParticleHandle->GetR()), ParticleHandle->GetV(), FMath::RadiansToDegrees(ParticleHandle->GetW()));
 	}
 
 	if (TStrongObjectPtr<UBaseMovementMode> CurrentModePtr = StateMachine.GetCurrentMode().Pin())
@@ -745,7 +746,7 @@ void UChaosMoverSimulation::Teleport(const FTransform& TargetTransform, const FM
 	// Invalidate the movement base for now, until we reestablish the character ground constraint
 	DefaultSyncState.SetMovementBase(nullptr);
 	// Make sure the sync state matches the particle transforms
-	DefaultSyncState.SetTransforms_WorldSpace(ParticleHandle->GetX(), FRotator(ParticleHandle->GetR()), ParticleHandle->GetV());
+	DefaultSyncState.SetTransforms_WorldSpace(ParticleHandle->GetX(), FRotator(ParticleHandle->GetR()), ParticleHandle->GetV(), FMath::RadiansToDegrees(ParticleHandle->GetW()));
 	// Rotate the move direction intent
 	DefaultSyncState.MoveDirectionIntent = TeleportDeltaTransform.TransformVector(DefaultSyncState.MoveDirectionIntent);
 }
diff --git a/Engine/Plugins/Experimental/ChaosMover/Source/ChaosMover/Private/Character/Effects/ChaosCharacterApplyVelocityEffect.cpp b/Engine/Plugins/Experimental/ChaosMover/Source/ChaosMover/Private/Character/Effects/ChaosCharacterApplyVelocityEffect.cpp
--- a/Engine/Plugins/Experimental/ChaosMover/Source/ChaosMover/Private/Character/Effects/ChaosCharacterApplyVelocityEffect.cpp
+++ b/Engine/Plugins/Experimental/ChaosMover/Source/ChaosMover/Private/Character/Effects/ChaosCharacterApplyVelocityEffect.cpp
@@ -61,7 +61,7 @@ bool FChaosCharacterApplyVelocityEffect::ApplyMovementEffect_Async(FApplyMovemen
 		break;
 	}
 
-	OutputSyncState.SetTransforms_WorldSpace(Position,Orientation,Velocity);
+	OutputSyncState.SetTransforms_WorldSpace(Position,Orientation,Velocity,FVector::ZeroVector);
 
 	return true;
 }
diff --git a/Engine/Plugins/Experimental/ChaosMover/Source/ChaosMover/Private/Character/Modes/ChaosFallingMode.cpp b/Engine/Plugins/Experimental/ChaosMover/Source/ChaosMover/Private/Character/Modes/ChaosFallingMode.cpp
--- a/Engine/Plugins/Experimental/ChaosMover/Source/ChaosMover/Private/Character/Modes/ChaosFallingMode.cpp
+++ b/Engine/Plugins/Experimental/ChaosMover/Source/ChaosMover/Private/Character/Modes/ChaosFallingMode.cpp
@@ -193,7 +193,7 @@ void UChaosFallingMode::SimulationTick_Implementation(const FSimulationTickParam
 
 	const float DeltaSeconds = Params.TimeStep.StepMs * 0.001f;
 
-	const FRotator TargetOrient = UMovementUtils::ApplyAngularVelocity(StartingSyncState->GetOrientation_WorldSpace(), ProposedMove.AngularVelocity, DeltaSeconds);
+	const FRotator TargetOrient = UMovementUtils::ApplyAngularVelocityToRotator(StartingSyncState->GetOrientation_WorldSpace(), ProposedMove.AngularVelocityDegrees, DeltaSeconds);
 
 	// The physics simulation applies Z-only gravity acceleration via physics volumes, so we need to account for it here 
 	FVector TargetVel = ProposedMove.LinearVelocity - DefaultSimInputs->PhysicsObjectGravity * FVector::UpVector * DeltaSeconds;
@@ -205,5 +205,6 @@ void UChaosFallingMode::SimulationTick_Implementation(const FSimulationTickParam
 	OutputSyncState.SetTransforms_WorldSpace(
 		TargetPos,
 		TargetOrient,
-		TargetVel);
+		TargetVel,
+		ProposedMove.AngularVelocityDegrees);
 }
\ No newline at end of file
diff --git a/Engine/Plugins/Experimental/ChaosMover/Source/ChaosMover/Private/Character/Modes/ChaosFlyingMode.cpp b/Engine/Plugins/Experimental/ChaosMover/Source/ChaosMover/Private/Character/Modes/ChaosFlyingMode.cpp
--- a/Engine/Plugins/Experimental/ChaosMover/Source/ChaosMover/Private/Character/Modes/ChaosFlyingMode.cpp
+++ b/Engine/Plugins/Experimental/ChaosMover/Source/ChaosMover/Private/Character/Modes/ChaosFlyingMode.cpp
@@ -120,7 +120,7 @@ void UChaosFlyingMode::SimulationTick_Implementation(const FSimulationTickParams
 
 	const float DeltaSeconds = Params.TimeStep.StepMs * 0.001f;
 
-	const FRotator TargetOrient = UMovementUtils::ApplyAngularVelocity(StartingSyncState->GetOrientation_WorldSpace(), ProposedMove.AngularVelocity, DeltaSeconds);
+	const FRotator TargetOrient = UMovementUtils::ApplyAngularVelocityToRotator(StartingSyncState->GetOrientation_WorldSpace(), ProposedMove.AngularVelocityDegrees, DeltaSeconds);
 
 	// The physics simulation applies Z-only gravity acceleration via physics volumes, so we need to account for it here 
 	FVector TargetVel = ProposedMove.LinearVelocity - DefaultSimInputs->PhysicsObjectGravity * FVector::UpVector * DeltaSeconds;
@@ -132,5 +132,6 @@ void UChaosFlyingMode::SimulationTick_Implementation(const FSimulationTickParams
 	OutputSyncState.SetTransforms_WorldSpace(
 		TargetPos,
 		TargetOrient,
-		TargetVel);
+		TargetVel,
+		ProposedMove.AngularVelocityDegrees);
 }
\ No newline at end of file
diff --git a/Engine/Plugins/Experimental/ChaosMover/Source/ChaosMover/Private/Character/Modes/ChaosSwimmingMode.cpp b/Engine/Plugins/Experimental/ChaosMover/Source/ChaosMover/Private/Character/Modes/ChaosSwimmingMode.cpp
--- a/Engine/Plugins/Experimental/ChaosMover/Source/ChaosMover/Private/Character/Modes/ChaosSwimmingMode.cpp
+++ b/Engine/Plugins/Experimental/ChaosMover/Source/ChaosMover/Private/Character/Modes/ChaosSwimmingMode.cpp
@@ -238,10 +238,7 @@ void UChaosSwimmingMode::SimulationTick_Implementation(const FSimulationTickPara
 	const float DeltaSeconds = Params.TimeStep.StepMs * 0.001f;
 
 	FRotator TargetOrient = StartingSyncState->GetOrientation_WorldSpace();
-	if (!UMovementUtils::IsAngularVelocityZero(ProposedMove.AngularVelocity))
-	{
-		TargetOrient += (ProposedMove.AngularVelocity * DeltaSeconds);
-	}
+	TargetOrient = UMovementUtils::ApplyAngularVelocityToRotator(TargetOrient, ProposedMove.AngularVelocityDegrees, DeltaSeconds);
 
 	FVector TargetVel = ProposedMove.LinearVelocity;
 	TargetVel -= DefaultSimInputs->PhysicsObjectGravity * DefaultSimInputs->UpDir * DeltaSeconds;;
@@ -252,7 +249,8 @@ void UChaosSwimmingMode::SimulationTick_Implementation(const FSimulationTickPara
 	OutputSyncState.SetTransforms_WorldSpace(
 		TargetPos,
 		TargetOrient,
-		TargetVel);
+		TargetVel,
+		ProposedMove.AngularVelocityDegrees);
 
 	OutputState.MovementEndState.RemainingMs = 0.0f;
 	OutputState.MovementEndState.NextModeName = Params.StartState.SyncState.MovementMode;
diff --git a/Engine/Plugins/Experimental/ChaosMover/Source/ChaosMover/Private/Character/Modes/ChaosWalkingMode.cpp b/Engine/Plugins/Experimental/ChaosMover/Source/ChaosMover/Private/Character/Modes/ChaosWalkingMode.cpp
--- a/Engine/Plugins/Experimental/ChaosMover/Source/ChaosMover/Private/Character/Modes/ChaosWalkingMode.cpp
+++ b/Engine/Plugins/Experimental/ChaosMover/Source/ChaosMover/Private/Character/Modes/ChaosWalkingMode.cpp
@@ -264,9 +264,12 @@ void UChaosWalkingMode::SimulationTick_Implementation(const FSimulationTickParam
 
 		// Target orientation
 		// This is always applied regardless of whether the character is supported
-		const FRotator TargetOrientation = UMovementUtils::ApplyAngularVelocity(StartingSyncState->GetOrientation_WorldSpace(), ProposedMove.AngularVelocity, DeltaSeconds);
+		const FRotator TargetOrientation = UMovementUtils::ApplyAngularVelocityToRotator(StartingSyncState->GetOrientation_WorldSpace(), ProposedMove.AngularVelocityDegrees, DeltaSeconds);
 
-		OutputSyncState.SetTransforms_WorldSpace(TargetPosition, TargetOrientation, TargetVelocity);
+		OutputSyncState.SetTransforms_WorldSpace(TargetPosition,
+			TargetOrientation,
+			TargetVelocity,
+			ProposedMove.AngularVelocityDegrees);
 	}
 	
 	OutputState.MovementEndState.RemainingMs = 0.0f;
diff --git a/Engine/Plugins/Experimental/ChaosMover/Source/ChaosMover/Private/PathedMovement/ChaosPathedMovementMode.cpp b/Engine/Plugins/Experimental/ChaosMover/Source/ChaosMover/Private/PathedMovement/ChaosPathedMovementMode.cpp
--- a/Engine/Plugins/Experimental/ChaosMover/Source/ChaosMover/Private/PathedMovement/ChaosPathedMovementMode.cpp
+++ b/Engine/Plugins/Experimental/ChaosMover/Source/ChaosMover/Private/PathedMovement/ChaosPathedMovementMode.cpp
@@ -286,6 +286,7 @@ void UChaosPathedMovementMode::SimulationTick_Implementation(const FSimulationTi
 
 	OutputDefaultSyncState.SetTransforms_WorldSpace(WorldTranslation, WorldRotation,
 											  FVector::ZeroVector, // Fix this eventually to return the velocity tangential to the path, including due to a change in path basis since last update
+											  FVector::ZeroVector,
 											  nullptr ); // no movement base
 
 #if WITH_CHAOS_VISUAL_DEBUGGER
diff --git a/Engine/Plugins/Experimental/Mover/Source/Mover/Private/Backends/MoverNetworkPhysicsLiaison.cpp b/Engine/Plugins/Experimental/Mover/Source/Mover/Private/Backends/MoverNetworkPhysicsLiaison.cpp
--- a/Engine/Plugins/Experimental/Mover/Source/Mover/Private/Backends/MoverNetworkPhysicsLiaison.cpp
+++ b/Engine/Plugins/Experimental/Mover/Source/Mover/Private/Backends/MoverNetworkPhysicsLiaison.cpp
@@ -360,7 +360,7 @@ void UMoverNetworkPhysicsLiaisonComponent::PerformPreSimulate_Internal(const FPh
 	}
 
 	FMoverDefaultSyncState& InputSyncState = Input.SyncState.SyncStateCollection.FindOrAddMutableDataByType<FMoverDefaultSyncState>();
-	InputSyncState.SetTransforms_WorldSpace(CharacterParticle->GetX(), FRotator(CharacterParticle->GetR()), CharacterParticle->GetV() - LocalGroundVelocity + AIMoveVelocity);
+	InputSyncState.SetTransforms_WorldSpace(CharacterParticle->GetX(), FRotator(CharacterParticle->GetR()), CharacterParticle->GetV() - LocalGroundVelocity + AIMoveVelocity, FMath::RadiansToDegrees(CharacterParticle->GetW()));
 
 	FCharacterDefaultInputs& CharacterDefaultInputs = Input.InputCmd.InputCollection.FindOrAddMutableDataByType<FCharacterDefaultInputs>();
 	if (!CharacterDefaultInputs.SuggestedMovementMode.IsNone())
diff --git a/Engine/Plugins/Experimental/Mover/Source/Mover/Private/Backends/MoverNetworkPhysicsLiaisonBase.cpp b/Engine/Plugins/Experimental/Mover/Source/Mover/Private/Backends/MoverNetworkPhysicsLiaisonBase.cpp
--- a/Engine/Plugins/Experimental/Mover/Source/Mover/Private/Backends/MoverNetworkPhysicsLiaisonBase.cpp
+++ b/Engine/Plugins/Experimental/Mover/Source/Mover/Private/Backends/MoverNetworkPhysicsLiaisonBase.cpp
@@ -634,7 +634,7 @@ void UMoverNetworkPhysicsLiaisonComponentBase::PostPhysicsUpdate_External()
 	FMoverDefaultSyncState& TransformSyncState = InterpolatedSyncState.SyncStateCollection.FindOrAddMutableDataByType<FMoverDefaultSyncState>();
 	if (Chaos::FPBDRigidParticle* Particle = GetControlledParticle_External())
 	{
-		TransformSyncState.SetTransforms_WorldSpace(Particle->GetX(), FRotator(Particle->GetR()), Particle->GetV(), TransformSyncState.GetMovementBase(), TransformSyncState.GetMovementBaseBoneName());
+		TransformSyncState.SetTransforms_WorldSpace(Particle->GetX(), FRotator(Particle->GetR()), Particle->GetV(), FMath::RadiansToDegrees(Particle->GetW()), TransformSyncState.GetMovementBase(), TransformSyncState.GetMovementBaseBoneName());
 		
 		// Make sure the move direction intent is in base space (the base quat is identity if there's no base, effectively making this a no-op)
 		TransformSyncState.MoveDirectionIntent = TransformSyncState.GetCapturedMovementBaseQuat().UnrotateVector(TransformSyncState.MoveDirectionIntent);
@@ -716,7 +716,7 @@ void UMoverNetworkPhysicsLiaisonComponentBase::ProcessInputs_Internal(int32 Phys
 					if (Chaos::FPBDRigidParticleHandle* ParticleHandle = GetControlledParticle_Internal())
 					{
 						FMoverDefaultSyncState& NetSyncState = LatestSyncState.SyncStateCollection.FindOrAddMutableDataByType<FMoverDefaultSyncState>();
-						NetSyncState.SetTransforms_WorldSpace(ParticleHandle->GetX(), FRotator(ParticleHandle->GetR()), ParticleHandle->GetV());
+						NetSyncState.SetTransforms_WorldSpace(ParticleHandle->GetX(), FRotator(ParticleHandle->GetR()), ParticleHandle->GetV(), FMath::RadiansToDegrees(ParticleHandle->GetW()));
 					}
 
 					MoverComp.OnSimulationPreRollback(&Input.SyncState, &LatestSyncState, &UnusedInvalidAuxState, &UnusedAuxState, CurrentBaseTimeStep);
@@ -750,7 +750,7 @@ void UMoverNetworkPhysicsLiaisonComponentBase::PerformProcessInputs_Internal(int
 	if (Chaos::FPBDRigidParticleHandle* ParticleHandle = GetControlledParticle_Internal())
 	{
 		FMoverDefaultSyncState& InputSyncState = Input.SyncState.SyncStateCollection.FindOrAddMutableDataByType<FMoverDefaultSyncState>();
-		InputSyncState.SetTransforms_WorldSpace(ParticleHandle->GetX(), FRotator(ParticleHandle->GetR()), ParticleHandle->GetV());
+		InputSyncState.SetTransforms_WorldSpace(ParticleHandle->GetX(), FRotator(ParticleHandle->GetR()), ParticleHandle->GetV(), FMath::RadiansToDegrees(ParticleHandle->GetW()));
 	}
 
 	if (GetMoverComponent().OnPreMovement.IsBound())
diff --git a/Engine/Plugins/Experimental/Mover/Source/Mover/Private/Backends/MoverNetworkPredictionLiaison.cpp b/Engine/Plugins/Experimental/Mover/Source/Mover/Private/Backends/MoverNetworkPredictionLiaison.cpp
--- a/Engine/Plugins/Experimental/Mover/Source/Mover/Private/Backends/MoverNetworkPredictionLiaison.cpp
+++ b/Engine/Plugins/Experimental/Mover/Source/Mover/Private/Backends/MoverNetworkPredictionLiaison.cpp
@@ -262,6 +262,7 @@ void UMoverNetworkPredictionLiaisonComponent::BeginPlay()
 			{
 				StartingSyncState->SetTransforms_WorldSpace(UpdatedComponentTransform.GetLocation(),
 													 UpdatedComponentTransform.GetRotation().Rotator(),
+													 FVector::ZeroVector,
 													 FVector::ZeroVector);	// no initial velocity
 			}
 		}
diff --git a/Engine/Plugins/Experimental/Mover/Source/Mover/Private/DefaultMovementSet/InstantMovementEffects/BasicInstantMovementEffects.cpp b/Engine/Plugins/Experimental/Mover/Source/Mover/Private/DefaultMovementSet/InstantMovementEffects/BasicInstantMovementEffects.cpp
--- a/Engine/Plugins/Experimental/Mover/Source/Mover/Private/DefaultMovementSet/InstantMovementEffects/BasicInstantMovementEffects.cpp
+++ b/Engine/Plugins/Experimental/Mover/Source/Mover/Private/DefaultMovementSet/InstantMovementEffects/BasicInstantMovementEffects.cpp
@@ -62,6 +62,7 @@ bool FTeleportEffect::ApplyMovementEffect(FApplyMovementEffectParams& ApplyEffec
 		OutputSyncState.SetTransforms_WorldSpace(UpdatedLocation,
 													UpdatedComponent->GetComponentRotation(),
 													OutputSyncState.GetVelocity_WorldSpace(),
+													OutputSyncState.GetAngularVelocityDegrees_WorldSpace(),
 													nullptr ); // no movement base
 		
 		// TODO: instead of invalidating it, consider checking for a floor. Possibly a dynamic base?
@@ -164,6 +165,7 @@ bool FAsyncTeleportEffect::ApplyMovementEffect(FApplyMovementEffectParams& Apply
 			OutputSyncState.SetTransforms_WorldSpace(TeleportLocation,
 				TeleportRotation,
 				OutputSyncState.GetVelocity_WorldSpace(),
+				OutputSyncState.GetAngularVelocityDegrees_WorldSpace(),
 				nullptr); // no movement base
 
 			// TODO: instead of invalidating it, consider checking for a floor. Possibly a dynamic base?
@@ -241,6 +243,7 @@ bool FJumpImpulseEffect::ApplyMovementEffect(FApplyMovementEffectParams& ApplyEf
 		OutputSyncState.SetTransforms_WorldSpace( ApplyEffectParams.UpdatedComponent->GetComponentLocation(),
 												  ApplyEffectParams.UpdatedComponent->GetComponentRotation(),
 												  FinalVelocity,
+												  FVector::ZeroVector,
 												  MovementBaseInfo.MovementBase.Get(),
 												  MovementBaseInfo.BoneName);
 		
@@ -315,6 +318,7 @@ bool FApplyVelocityEffect::ApplyMovementEffect(FApplyMovementEffectParams& Apply
 	OutputSyncState.SetTransforms_WorldSpace( ApplyEffectParams.UpdatedComponent->GetComponentLocation(),
 											  ApplyEffectParams.UpdatedComponent->GetComponentRotation(),
 											  Velocity,
+											  FVector::ZeroVector,
 											  MovementBaseInfo.MovementBase.Get(),
 											  MovementBaseInfo.BoneName);
 
diff --git a/Engine/Plugins/Experimental/Mover/Source/Mover/Private/DefaultMovementSet/LayeredMoves/AnimRootMotionLayeredMove.cpp b/Engine/Plugins/Experimental/Mover/Source/Mover/Private/DefaultMovementSet/LayeredMoves/AnimRootMotionLayeredMove.cpp
--- a/Engine/Plugins/Experimental/Mover/Source/Mover/Private/DefaultMovementSet/LayeredMoves/AnimRootMotionLayeredMove.cpp
+++ b/Engine/Plugins/Experimental/Mover/Source/Mover/Private/DefaultMovementSet/LayeredMoves/AnimRootMotionLayeredMove.cpp
@@ -94,7 +94,7 @@ bool FLayeredMove_AnimRootMotion::GenerateMove(const FMoverTickStartData& SimSta
 	if (DeltaSeconds > UE_KINDA_SMALL_NUMBER)
 	{
 		OutProposedMove.LinearVelocity    = WorldSpaceRootMotion.GetTranslation() / DeltaSeconds;
-		OutProposedMove.AngularVelocity   = WorldSpaceRootMotion.GetRotation().Rotator() * (1.f / DeltaSeconds);
+		OutProposedMove.AngularVelocityDegrees   = FMath::RadiansToDegrees(WorldSpaceRootMotion.GetRotation().ToRotationVector() / DeltaSeconds);
 	}
 
 #if !UE_BUILD_SHIPPING
diff --git a/Engine/Plugins/Experimental/Mover/Source/Mover/Private/DefaultMovementSet/LayeredMoves/RootMotionAttributeLayeredMove.cpp b/Engine/Plugins/Experimental/Mover/Source/Mover/Private/DefaultMovementSet/LayeredMoves/RootMotionAttributeLayeredMove.cpp
--- a/Engine/Plugins/Experimental/Mover/Source/Mover/Private/DefaultMovementSet/LayeredMoves/RootMotionAttributeLayeredMove.cpp
+++ b/Engine/Plugins/Experimental/Mover/Source/Mover/Private/DefaultMovementSet/LayeredMoves/RootMotionAttributeLayeredMove.cpp
@@ -161,14 +161,22 @@ bool FLayeredMove_RootMotionAttribute::GenerateMove(const FMoverTickStartData& S
 		}
 
 		// Convert the transform into linear and angular velocities
-		OutProposedMove.LinearVelocity  = WorldSpaceRootMotion.GetTranslation() * (1.f / DeltaSeconds);
-		OutProposedMove.AngularVelocity = WorldSpaceRootMotion.GetRotation().Rotator() * (1.f / DeltaSeconds);
+		if (DeltaSeconds > UE_SMALL_NUMBER)
+		{
+			OutProposedMove.LinearVelocity  = WorldSpaceRootMotion.GetTranslation() * (1.f / DeltaSeconds);
+			OutProposedMove.AngularVelocityDegrees = FMath::RadiansToDegrees(WorldSpaceRootMotion.GetRotation().ToRotationVector()  / DeltaSeconds);			
+		}
+		else
+		{
+			OutProposedMove.LinearVelocity = FVector::ZeroVector;
+			OutProposedMove.AngularVelocityDegrees = FVector::ZeroVector;
+		}
 
 #if !UE_BUILD_SHIPPING
 		UE_CLOG(CVarLogRootMotionAttrSteps->GetBool(), LogMover, Log, TEXT("RootMotionAttr. SimF %i (dt %.3f) => LocalT: %s (WST: %s)  XY Speed: %.6f Z: %.6f   AngV: %s"),
 			TimeStep.ServerFrame, DeltaSeconds,
 			*LocalRootMotion.GetTranslation().ToString(), *WorldSpaceRootMotion.GetTranslation().ToString(),
-			OutProposedMove.LinearVelocity.Size2D(), OutProposedMove.LinearVelocity.Z, *OutProposedMove.AngularVelocity.ToCompactString());
+			OutProposedMove.LinearVelocity.Size2D(), OutProposedMove.LinearVelocity.Z, *OutProposedMove.AngularVelocityDegrees.ToCompactString());
 
 		if (ExcessiveLinearVelocitySquaredThreshold > 0.f && ExcessiveLinearVelocitySquaredThreshold < OutProposedMove.LinearVelocity.SquaredLength())
 		{
diff --git a/Engine/Plugins/Experimental/Mover/Source/Mover/Private/DefaultMovementSet/Modes/AsyncFallingMode.cpp b/Engine/Plugins/Experimental/Mover/Source/Mover/Private/DefaultMovementSet/Modes/AsyncFallingMode.cpp
--- a/Engine/Plugins/Experimental/Mover/Source/Mover/Private/DefaultMovementSet/Modes/AsyncFallingMode.cpp
+++ b/Engine/Plugins/Experimental/Mover/Source/Mover/Private/DefaultMovementSet/Modes/AsyncFallingMode.cpp
@@ -181,7 +181,7 @@ void UAsyncFallingMode::SimulationTick_Implementation(const FSimulationTickParam
 
 	// Use the orientation intent directly. If no intent is provided, use last frame's orientation. Note that we are assuming rotation changes can't fail. 
 	const FRotator StartingOrient = StartingSyncState->GetOrientation_WorldSpace();
-	const FRotator TargetOrient = UMovementUtils::ApplyAngularVelocity(StartingOrient, ProposedMove.AngularVelocity, DeltaSeconds);
+	const FRotator TargetOrient = UMovementUtils::ApplyAngularVelocityToRotator(StartingOrient, ProposedMove.AngularVelocityDegrees, DeltaSeconds);
 
 	const FVector StartingFallingVelocity = StartingSyncState->GetVelocity_WorldSpace();
 
@@ -227,7 +227,7 @@ void UAsyncFallingMode::SimulationTick_Implementation(const FSimulationTickParam
 		{
 			// Adjust height to float slightly above walkable floor
 			LocationInProgress = UGroundMovementUtils::TestMoveToKeepMinHeightAboveFloor(Params.MovingComps, LocationInProgress, RotationInProgress, CommonLegacySettings->MaxWalkSlopeCosine, IN OUT LandingFloor, IN OUT MoveRecord);
-			CaptureFinalState(LocationInProgress, RotationInProgress.Rotator(), LandingFloor, DeltaSeconds, DeltaSeconds * PctTimeApplied, OutputSyncState, OutputState, IN OUT MoveRecord);
+			CaptureFinalState(LocationInProgress, RotationInProgress.Rotator(), LandingFloor, DeltaSeconds, DeltaSeconds * PctTimeApplied, ProposedMove.AngularVelocityDegrees, OutputSyncState, OutputState, IN OUT MoveRecord);
 			return;
 		}
 
@@ -264,7 +264,7 @@ void UAsyncFallingMode::SimulationTick_Implementation(const FSimulationTickParam
 		PctTimeApplied = 1.f;
 	}
 	
-	CaptureFinalState(LocationInProgress, RotationInProgress.Rotator(), LandingFloor, DeltaSeconds, DeltaSeconds* PctTimeApplied, OutputSyncState, OutputState, MoveRecord);
+	CaptureFinalState(LocationInProgress, RotationInProgress.Rotator(), LandingFloor, DeltaSeconds, DeltaSeconds* PctTimeApplied,ProposedMove.AngularVelocityDegrees, OutputSyncState, OutputState, MoveRecord);
 }
 
 
@@ -327,7 +327,7 @@ void UAsyncFallingMode::ProcessLanded(const FFloorCheckResult& FloorResult, FVec
 	}
 }
 
-void UAsyncFallingMode::CaptureFinalState(const FVector FinalLocation, const FRotator FinalRotation, const FFloorCheckResult& FloorResult, float DeltaSeconds, float DeltaSecondsUsed, FMoverDefaultSyncState& OutputSyncState, FMoverTickEndData& TickEndData, FMovementRecord& Record) const
+void UAsyncFallingMode::CaptureFinalState(const FVector FinalLocation, const FRotator FinalRotation, const FFloorCheckResult& FloorResult, float DeltaSeconds, float DeltaSecondsUsed, const FVector& AngularVelocityDegrees, FMoverDefaultSyncState& OutputSyncState, FMoverTickEndData& TickEndData, FMovementRecord& Record) const
 {
 	UMoverBlackboard* SimBlackboard = GetMoverComponent()->GetSimBlackboard_Mutable();
 
@@ -358,6 +358,7 @@ void UAsyncFallingMode::CaptureFinalState(const FVector FinalLocation, const FRo
 		OutputSyncState.SetTransforms_WorldSpace( FinalLocation,
 												  FinalRotation,
 												  EffectiveVelocity,
+												  AngularVelocityDegrees,
 												  MovementBaseInfo.MovementBase.Get(), MovementBaseInfo.BoneName);
 	}
 	else
@@ -365,6 +366,7 @@ void UAsyncFallingMode::CaptureFinalState(const FVector FinalLocation, const FRo
 		OutputSyncState.SetTransforms_WorldSpace( FinalLocation,
 												  FinalRotation,
 												  EffectiveVelocity,
+												  AngularVelocityDegrees,
 												  nullptr); // no movement base
 	}
 }
diff --git a/Engine/Plugins/Experimental/Mover/Source/Mover/Private/DefaultMovementSet/Modes/AsyncFlyingMode.cpp b/Engine/Plugins/Experimental/Mover/Source/Mover/Private/DefaultMovementSet/Modes/AsyncFlyingMode.cpp
--- a/Engine/Plugins/Experimental/Mover/Source/Mover/Private/DefaultMovementSet/Modes/AsyncFlyingMode.cpp
+++ b/Engine/Plugins/Experimental/Mover/Source/Mover/Private/DefaultMovementSet/Modes/AsyncFlyingMode.cpp
@@ -98,7 +98,7 @@ void UAsyncFlyingMode::SimulationTick_Implementation(const FSimulationTickParams
 
 	// Use the orientation intent directly. If no intent is provided, use last frame's orientation. Note that we are assuming rotation changes can't fail. 
 	const FRotator StartingOrient = StartingSyncState->GetOrientation_WorldSpace();
-	const FRotator TargetOrient = UMovementUtils::ApplyAngularVelocity(StartingOrient, ProposedMove.AngularVelocity, DeltaSeconds);
+	const FRotator TargetOrient = UMovementUtils::ApplyAngularVelocityToRotator(StartingOrient, ProposedMove.AngularVelocityDegrees, DeltaSeconds);
 	
 	const FVector StartLocation = StartingSyncState->GetLocation_WorldSpace();
 	const FVector TargetLocation = StartLocation + (ProposedMove.LinearVelocity * DeltaSeconds);
@@ -148,6 +148,7 @@ void UAsyncFlyingMode::SimulationTick_Implementation(const FSimulationTickParams
 		LocationInProgress,
 		RotationInProgress.Rotator(),
 		SweepRecord.GetRelevantVelocity(),
+		ProposedMove.AngularVelocityDegrees,
 		nullptr); // no movement base
 }
 
diff --git a/Engine/Plugins/Experimental/Mover/Source/Mover/Private/DefaultMovementSet/Modes/AsyncNavWalkingMode.cpp b/Engine/Plugins/Experimental/Mover/Source/Mover/Private/DefaultMovementSet/Modes/AsyncNavWalkingMode.cpp
--- a/Engine/Plugins/Experimental/Mover/Source/Mover/Private/DefaultMovementSet/Modes/AsyncNavWalkingMode.cpp
+++ b/Engine/Plugins/Experimental/Mover/Source/Mover/Private/DefaultMovementSet/Modes/AsyncNavWalkingMode.cpp
@@ -139,7 +139,7 @@ void UAsyncNavWalkingMode::GenerateMove_Implementation(const FMoverTickStartData
 
 	if (TurnGenerator)
 	{
-		OutProposedMove.AngularVelocity = ITurnGeneratorInterface::Execute_GetTurn(TurnGenerator, IntendedOrientation_WorldSpace, StartState, *StartingSyncState, TimeStep, OutProposedMove, SimBlackboard);
+		OutProposedMove.AngularVelocityDegrees = ITurnGeneratorInterface::Execute_GetTurn(TurnGenerator, IntendedOrientation_WorldSpace, StartState, *StartingSyncState, TimeStep, OutProposedMove, SimBlackboard);
 	}
 }
 
@@ -174,7 +174,7 @@ void UAsyncNavWalkingMode::SimulationTick_Implementation(const FSimulationTickPa
 
 	const FRotator StartingOrient = StartingSyncState->GetOrientation_WorldSpace();
 	
-	const FRotator TargetOrient = UMovementUtils::ApplyAngularVelocity(StartingOrient, ProposedMove.AngularVelocity, DeltaSeconds);
+	const FRotator TargetOrient = UMovementUtils::ApplyAngularVelocityToRotator(StartingOrient, ProposedMove.AngularVelocityDegrees, DeltaSeconds);
 
 	const FVector StartingFeetLocation = NavMoverComponent->GetFeetLocationAt(StartingSyncState->GetLocation_WorldSpace());
 	FVector TargetFeetLocation = StartingFeetLocation + (ProposedMove.LinearVelocity * DeltaSeconds);
@@ -320,7 +320,7 @@ void UAsyncNavWalkingMode::SimulationTick_Implementation(const FSimulationTickPa
 					break;
 			}
 
-			CaptureOutputState(*StartingSyncState, LocationInProgress, RotationInProgress.Rotator(), MoveRecord, OutputSyncState, OutputState);
+			CaptureOutputState(*StartingSyncState, LocationInProgress, RotationInProgress.Rotator(), MoveRecord, ProposedMove.AngularVelocityDegrees, OutputSyncState, OutputState);
 			return;
 		}
 
@@ -377,7 +377,7 @@ void UAsyncNavWalkingMode::SimulationTick_Implementation(const FSimulationTickPa
 		MoveRecord.SetDeltaSeconds((Params.TimeStep.StepMs - OutputState.MovementEndState.RemainingMs) * 0.001f);
 	}
 
-	CaptureOutputState(*StartingSyncState, LocationInProgress, RotationInProgress.Rotator(), MoveRecord, OutputSyncState, OutputState);
+	CaptureOutputState(*StartingSyncState, LocationInProgress, RotationInProgress.Rotator(), MoveRecord, ProposedMove.AngularVelocityDegrees, OutputSyncState, OutputState);
 }
 
 bool UAsyncNavWalkingMode::FindNavFloor(const FVector& TestLocation, FNavLocation& OutNavFloorLocation, const INavigationDataInterface* NavData) const
@@ -661,7 +661,7 @@ void UAsyncNavWalkingMode::OnUnregistered()
 	Super::OnUnregistered();
 }
 
-void UAsyncNavWalkingMode::CaptureOutputState(const FMoverDefaultSyncState& StartSyncState, const FVector& FinalLocation, const FRotator& FinalRotation, const FMovementRecord& Record, FMoverDefaultSyncState& OutputSyncState, FMoverTickEndData& TickEndData) const
+void UAsyncNavWalkingMode::CaptureOutputState(const FMoverDefaultSyncState& StartSyncState, const FVector& FinalLocation, const FRotator& FinalRotation, const FMovementRecord& Record, const FVector& AngularVelocityDegrees, FMoverDefaultSyncState& OutputSyncState, FMoverTickEndData& TickEndData) const
 {
 	UMoverBlackboard* SimBlackboard = GetMoverComponent()->GetSimBlackboard_Mutable();
 
@@ -672,6 +672,7 @@ void UAsyncNavWalkingMode::CaptureOutputState(const FMoverDefaultSyncState& Star
 	OutputSyncState.SetTransforms_WorldSpace(FinalLocation,
 		FinalRotation,
 		Record.GetRelevantVelocity(),
+		AngularVelocityDegrees,
 		nullptr);	// no movement base
 
 	TickEndData.MovementEndState.bEndedWithNoChanges = OutputSyncState.IsNearlyEqual(StartSyncState);
diff --git a/Engine/Plugins/Experimental/Mover/Source/Mover/Private/DefaultMovementSet/Modes/AsyncWalkingMode.cpp b/Engine/Plugins/Experimental/Mover/Source/Mover/Private/DefaultMovementSet/Modes/AsyncWalkingMode.cpp
--- a/Engine/Plugins/Experimental/Mover/Source/Mover/Private/DefaultMovementSet/Modes/AsyncWalkingMode.cpp
+++ b/Engine/Plugins/Experimental/Mover/Source/Mover/Private/DefaultMovementSet/Modes/AsyncWalkingMode.cpp
@@ -111,7 +111,7 @@ void UAsyncWalkingMode::GenerateMove_Implementation(const FMoverTickStartData& S
 
 	if (TurnGenerator)
 	{
-		OutProposedMove.AngularVelocity = ITurnGeneratorInterface::Execute_GetTurn(TurnGenerator, IntendedOrientation_WorldSpace, StartState, *StartingSyncState, TimeStep, OutProposedMove, SimBlackboard);
+		OutProposedMove.AngularVelocityDegrees = ITurnGeneratorInterface::Execute_GetTurn(TurnGenerator, IntendedOrientation_WorldSpace, StartState, *StartingSyncState, TimeStep, OutProposedMove, SimBlackboard);
 	}
 }
 
@@ -158,7 +158,7 @@ void UAsyncWalkingMode::SimulationTick_Implementation(const FSimulationTickParam
 	OutputSyncState.MoveDirectionIntent = (ProposedMove.bHasDirIntent ? ProposedMove.DirectionIntent : FVector::ZeroVector);
 
 	const FRotator StartingOrient = StartingSyncState->GetOrientation_WorldSpace();
-	const FRotator TargetOrient = UMovementUtils::ApplyAngularVelocity(StartingOrient, ProposedMove.AngularVelocity, DeltaSeconds);
+	const FRotator TargetOrient = UMovementUtils::ApplyAngularVelocityToRotator(StartingOrient, ProposedMove.AngularVelocityDegrees, DeltaSeconds);
 	const bool bIsOrientationChanging = !StartingOrient.Equals(TargetOrient);
 
 	const FQuat StartRotation = StartingOrient.Quaternion();
@@ -286,7 +286,7 @@ void UAsyncWalkingMode::SimulationTick_Implementation(const FSimulationTickParam
 			OutputState.MovementEndState.NextModeName = CommonLegacySettings->AirMovementModeName;
 			OutputState.MovementEndState.RemainingMs = Params.TimeStep.StepMs - (Params.TimeStep.StepMs * PercentTimeAppliedSoFar);
 			MoveRecord.SetDeltaSeconds((Params.TimeStep.StepMs - OutputState.MovementEndState.RemainingMs) * 0.001f);
-			CaptureFinalState(LocationInProgress, RotationInProgress.Rotator(), bDidAttemptMovement, CurrentFloor, MoveRecord, OutputSyncState);
+			CaptureFinalState(LocationInProgress, RotationInProgress.Rotator(), bDidAttemptMovement, CurrentFloor, MoveRecord, ProposedMove.AngularVelocityDegrees, OutputSyncState);
 			return;
 		}
 	}
@@ -328,7 +328,7 @@ void UAsyncWalkingMode::SimulationTick_Implementation(const FSimulationTickParam
 		*/
 	}
 
-	CaptureFinalState(LocationInProgress, RotationInProgress.Rotator(), bDidAttemptMovement, CurrentFloor, MoveRecord, OutputSyncState);
+	CaptureFinalState(LocationInProgress, RotationInProgress.Rotator(), bDidAttemptMovement, CurrentFloor, MoveRecord, ProposedMove.AngularVelocityDegrees, OutputSyncState);
 
 }
 
@@ -365,7 +365,7 @@ void UAsyncWalkingMode::OnUnregistered()
 	Super::OnUnregistered();
 }
 
-void UAsyncWalkingMode::CaptureFinalState(const FVector FinalLocation, const FRotator FinalRotation, bool bDidAttemptMovement, const FFloorCheckResult& FloorResult, const FMovementRecord& Record, FMoverDefaultSyncState& OutputSyncState) const
+void UAsyncWalkingMode::CaptureFinalState(const FVector FinalLocation, const FRotator FinalRotation, bool bDidAttemptMovement, const FFloorCheckResult& FloorResult, const FMovementRecord& Record, const FVector& AngularVelocityDegrees, FMoverDefaultSyncState& OutputSyncState) const
 {
 	FRelativeBaseInfo PriorBaseInfo;
 
@@ -392,6 +392,7 @@ void UAsyncWalkingMode::CaptureFinalState(const FVector FinalLocation, const FRo
 		OutputSyncState.SetTransforms_WorldSpace( FinalLocation,
                                                   FinalRotation,
                                                   Record.GetRelevantVelocity(),
+                                                  AngularVelocityDegrees,
                                                   CurrentBaseInfo.MovementBase.Get(), CurrentBaseInfo.BoneName);
 	}
 	else
@@ -401,6 +402,7 @@ void UAsyncWalkingMode::CaptureFinalState(const FVector FinalLocation, const FRo
 		OutputSyncState.SetTransforms_WorldSpace( FinalLocation,
                                                   FinalRotation,
                                                   Record.GetRelevantVelocity(),
+                                                  AngularVelocityDegrees,
                                                   nullptr);  // no movement base
 	}
 }
diff --git a/Engine/Plugins/Experimental/Mover/Source/Mover/Private/DefaultMovementSet/Modes/FallingMode.cpp b/Engine/Plugins/Experimental/Mover/Source/Mover/Private/DefaultMovementSet/Modes/FallingMode.cpp
--- a/Engine/Plugins/Experimental/Mover/Source/Mover/Private/DefaultMovementSet/Modes/FallingMode.cpp
+++ b/Engine/Plugins/Experimental/Mover/Source/Mover/Private/DefaultMovementSet/Modes/FallingMode.cpp
@@ -175,7 +175,7 @@ void UFallingMode::SimulationTick_Implementation(const FSimulationTickParams& Pa
 
 	// Use the orientation intent directly. If no intent is provided, use last frame's orientation. Note that we are assuming rotation changes can't fail. 
 	const FRotator StartingOrient = StartingSyncState->GetOrientation_WorldSpace();
-	const FRotator TargetOrient = UMovementUtils::ApplyAngularVelocity(StartingOrient, ProposedMove.AngularVelocity, DeltaSeconds);
+	const FRotator TargetOrient = UMovementUtils::ApplyAngularVelocityToRotator(StartingOrient, ProposedMove.AngularVelocityDegrees, DeltaSeconds);
 	
 	const FVector StartingFallingVelocity = StartingSyncState->GetVelocity_WorldSpace();
 
@@ -192,7 +192,7 @@ void UFallingMode::SimulationTick_Implementation(const FSimulationTickParams& Pa
 		if (FloorUnderActor.IsWalkableFloor())
 		{
 			UGroundMovementUtils::TryMoveToKeepMinHeightAboveFloor(MoverComponent, FloorUnderActor, CommonLegacySettings->MaxWalkSlopeCosine, MoveRecord);
-			CaptureFinalState(UpdatedComponent, *StartingSyncState, FloorUnderActor, DeltaSeconds, DeltaSeconds * PctTimeApplied, OutputSyncState, OutputState, MoveRecord);
+			CaptureFinalState(UpdatedComponent, *StartingSyncState, FloorUnderActor, DeltaSeconds, DeltaSeconds * PctTimeApplied, ProposedMove.AngularVelocityDegrees, OutputSyncState, OutputState, MoveRecord);
 			return;
 		}
 	}
@@ -225,7 +225,7 @@ void UFallingMode::SimulationTick_Implementation(const FSimulationTickParams& Pa
 			Hit, CommonLegacySettings->FloorSweepDistance, CommonLegacySettings->MaxWalkSlopeCosine, OUT LandingFloor))
 		{
 			UGroundMovementUtils::TryMoveToKeepMinHeightAboveFloor(MoverComponent, LandingFloor, CommonLegacySettings->MaxWalkSlopeCosine, MoveRecord); // make sure we maintain a small gap over walking surfaces
-			CaptureFinalState(UpdatedComponent, *StartingSyncState, LandingFloor, DeltaSeconds, DeltaSeconds * PctTimeApplied, OutputSyncState, OutputState, MoveRecord);
+			CaptureFinalState(UpdatedComponent, *StartingSyncState, LandingFloor, DeltaSeconds, DeltaSeconds * PctTimeApplied, ProposedMove.AngularVelocityDegrees, OutputSyncState, OutputState, MoveRecord);
 			return;
 		}
 		
@@ -245,7 +245,7 @@ void UFallingMode::SimulationTick_Implementation(const FSimulationTickParams& Pa
 		if (LandingFloor.IsWalkableFloor())
 		{
 			UGroundMovementUtils::TryMoveToKeepMinHeightAboveFloor(MoverComponent, LandingFloor, CommonLegacySettings->MaxWalkSlopeCosine, MoveRecord); // make sure we maintain a small gap over walking surfaces
-			CaptureFinalState(UpdatedComponent, *StartingSyncState, LandingFloor, DeltaSeconds, DeltaSeconds * PctTimeApplied, OutputSyncState, OutputState, MoveRecord);
+			CaptureFinalState(UpdatedComponent, *StartingSyncState, LandingFloor, DeltaSeconds, DeltaSeconds * PctTimeApplied, ProposedMove.AngularVelocityDegrees, OutputSyncState, OutputState, MoveRecord);
 			return;
 		}
 	}
@@ -255,7 +255,7 @@ void UFallingMode::SimulationTick_Implementation(const FSimulationTickParams& Pa
 		PctTimeApplied = 1.f;
 	}
 	
-	CaptureFinalState(UpdatedComponent, *StartingSyncState, LandingFloor, DeltaSeconds, DeltaSeconds* PctTimeApplied, OutputSyncState, OutputState, MoveRecord);
+	CaptureFinalState(UpdatedComponent, *StartingSyncState, LandingFloor, DeltaSeconds, DeltaSeconds* PctTimeApplied, ProposedMove.AngularVelocityDegrees, OutputSyncState, OutputState, MoveRecord);
 }
 
 
@@ -316,7 +316,7 @@ void UFallingMode::ProcessLanded(const FFloorCheckResult& FloorResult, FVector&
 	}
 }
 
-void UFallingMode::CaptureFinalState(USceneComponent* UpdatedComponent, const FMoverDefaultSyncState& StartSyncState, const FFloorCheckResult& FloorResult, float DeltaSeconds, float DeltaSecondsUsed, FMoverDefaultSyncState& OutputSyncState, FMoverTickEndData& TickEndData, FMovementRecord& Record) const
+void UFallingMode::CaptureFinalState(USceneComponent* UpdatedComponent, const FMoverDefaultSyncState& StartSyncState, const FFloorCheckResult& FloorResult, float DeltaSeconds, float DeltaSecondsUsed, const FVector& AngularVelocityDegrees, FMoverDefaultSyncState& OutputSyncState, FMoverTickEndData& TickEndData, FMovementRecord& Record) const
 {
 	UMoverBlackboard* SimBlackboard = GetMoverComponent()->GetSimBlackboard_Mutable();
 
@@ -350,6 +350,7 @@ void UFallingMode::CaptureFinalState(USceneComponent* UpdatedComponent, const FM
 		OutputSyncState.SetTransforms_WorldSpace( FinalLocation,
 												  UpdatedComponent->GetComponentRotation(),
 												  EffectiveVelocity,
+												  AngularVelocityDegrees,
 												  MovementBaseInfo.MovementBase.Get(), MovementBaseInfo.BoneName);
 	}
 	else
@@ -357,6 +358,7 @@ void UFallingMode::CaptureFinalState(USceneComponent* UpdatedComponent, const FM
 		OutputSyncState.SetTransforms_WorldSpace( FinalLocation,
 												  UpdatedComponent->GetComponentRotation(),
 												  EffectiveVelocity,
+												  AngularVelocityDegrees,
 												  nullptr); // no movement base
 	}
 
diff --git a/Engine/Plugins/Experimental/Mover/Source/Mover/Private/DefaultMovementSet/Modes/FlyingMode.cpp b/Engine/Plugins/Experimental/Mover/Source/Mover/Private/DefaultMovementSet/Modes/FlyingMode.cpp
--- a/Engine/Plugins/Experimental/Mover/Source/Mover/Private/DefaultMovementSet/Modes/FlyingMode.cpp
+++ b/Engine/Plugins/Experimental/Mover/Source/Mover/Private/DefaultMovementSet/Modes/FlyingMode.cpp
@@ -103,7 +103,7 @@ void UFlyingMode::SimulationTick_Implementation(const FSimulationTickParams& Par
 	// Use the orientation intent directly. If no intent is provided, use last frame's orientation. Note that we are assuming rotation changes can't fail. 
 	const FRotator StartingOrient = StartingSyncState->GetOrientation_WorldSpace();
 
-	const FRotator TargetOrient = UMovementUtils::ApplyAngularVelocity(StartingOrient, ProposedMove.AngularVelocity, DeltaSeconds);
+	const FRotator TargetOrient = UMovementUtils::ApplyAngularVelocityToRotator(StartingOrient, ProposedMove.AngularVelocityDegrees, DeltaSeconds);
 	const bool bIsOrientationChanging = !StartingOrient.Equals(TargetOrient);
 	
 	FVector MoveDelta = ProposedMove.LinearVelocity * DeltaSeconds;
@@ -142,11 +142,11 @@ void UFlyingMode::SimulationTick_Implementation(const FSimulationTickParams& Par
 		}
 	}
 
-	CaptureFinalState(UpdatedComponent, MoveRecord, *StartingSyncState, OutputSyncState, DeltaSeconds);
+	CaptureFinalState(UpdatedComponent, MoveRecord, *StartingSyncState, ProposedMove.AngularVelocityDegrees, OutputSyncState, DeltaSeconds);
 }
 
 // TODO: replace this function with simply looking at/collapsing the MovementRecord
-void UFlyingMode::CaptureFinalState(USceneComponent* UpdatedComponent, FMovementRecord& Record, const FMoverDefaultSyncState& StartSyncState, FMoverDefaultSyncState& OutputSyncState, const float DeltaSeconds) const
+void UFlyingMode::CaptureFinalState(USceneComponent* UpdatedComponent, FMovementRecord& Record, const FMoverDefaultSyncState& StartSyncState, const FVector& AngularVelocityDegrees, FMoverDefaultSyncState& OutputSyncState, const float DeltaSeconds) const
 {
 	const FVector FinalLocation = UpdatedComponent->GetComponentLocation();
 	const FVector FinalVelocity = Record.GetRelevantVelocity();
@@ -156,6 +156,7 @@ void UFlyingMode::CaptureFinalState(USceneComponent* UpdatedComponent, FMovement
 	OutputSyncState.SetTransforms_WorldSpace(FinalLocation,
 											  UpdatedComponent->GetComponentRotation(),
 											  FinalVelocity,
+											  AngularVelocityDegrees,
 											  nullptr); // no movement base
 
 	UpdatedComponent->ComponentVelocity = FinalVelocity;
diff --git a/Engine/Plugins/Experimental/Mover/Source/Mover/Private/DefaultMovementSet/Modes/NavWalkingMode.cpp b/Engine/Plugins/Experimental/Mover/Source/Mover/Private/DefaultMovementSet/Modes/NavWalkingMode.cpp
--- a/Engine/Plugins/Experimental/Mover/Source/Mover/Private/DefaultMovementSet/Modes/NavWalkingMode.cpp
+++ b/Engine/Plugins/Experimental/Mover/Source/Mover/Private/DefaultMovementSet/Modes/NavWalkingMode.cpp
@@ -114,7 +114,7 @@ void UNavWalkingMode::GenerateMove_Implementation(const FMoverTickStartData& Sta
 
 	if (TurnGenerator)
 	{
-		OutProposedMove.AngularVelocity = ITurnGeneratorInterface::Execute_GetTurn(TurnGenerator, IntendedOrientation_WorldSpace, StartState, *StartingSyncState, TimeStep, OutProposedMove, SimBlackboard);
+		OutProposedMove.AngularVelocityDegrees = ITurnGeneratorInterface::Execute_GetTurn(TurnGenerator, IntendedOrientation_WorldSpace, StartState, *StartingSyncState, TimeStep, OutProposedMove, SimBlackboard);
 	}
 }
 
@@ -155,7 +155,7 @@ void UNavWalkingMode::SimulationTick_Implementation(const FSimulationTickParams&
 	OutputSyncState.MoveDirectionIntent = (ProposedMove.bHasDirIntent ? ProposedMove.DirectionIntent : FVector::ZeroVector);
 
 	const FRotator StartingOrient = StartingSyncState->GetOrientation_WorldSpace();
-	const FRotator TargetOrient = UMovementUtils::ApplyAngularVelocity(StartingOrient, ProposedMove.AngularVelocity, DeltaSeconds);
+	const FRotator TargetOrient = UMovementUtils::ApplyAngularVelocityToRotator(StartingOrient, ProposedMove.AngularVelocityDegrees, DeltaSeconds);
 	const bool bIsOrientationChanging = !StartingOrient.Equals(TargetOrient);
 	
 	FQuat TargetOrientQuat = TargetOrient.Quaternion();
@@ -296,7 +296,7 @@ void UNavWalkingMode::SimulationTick_Implementation(const FSimulationTickParams&
 					break;
 			}
 
-			CaptureFinalState(UpdatedComponent, MoveRecord, OutputSyncState);
+			CaptureFinalState(UpdatedComponent, MoveRecord, ProposedMove.AngularVelocityDegrees, OutputSyncState);
 			return;
 		}
 
@@ -338,7 +338,7 @@ void UNavWalkingMode::SimulationTick_Implementation(const FSimulationTickParams&
 		MoveRecord.SetDeltaSeconds(0.0f);
 	}
 
-	CaptureFinalState(UpdatedComponent, MoveRecord, OutputSyncState);
+	CaptureFinalState(UpdatedComponent, MoveRecord, ProposedMove.AngularVelocityDegrees, OutputSyncState);
 }
 
 bool UNavWalkingMode::FindNavFloor(const FVector& TestLocation, FNavLocation& OutNavFloorLocation, const INavigationDataInterface* NavData) const
@@ -641,7 +641,7 @@ void UNavWalkingMode::OnUnregistered()
 	Super::OnUnregistered();
 }
 
-void UNavWalkingMode::CaptureFinalState(USceneComponent* UpdatedComponent, const FMovementRecord& Record, FMoverDefaultSyncState& OutputSyncState) const
+void UNavWalkingMode::CaptureFinalState(USceneComponent* UpdatedComponent, const FMovementRecord& Record, const FVector& AngularVelocityDegrees, FMoverDefaultSyncState& OutputSyncState) const
 {
 	UMoverBlackboard* SimBlackboard = GetMoverComponent()->GetSimBlackboard_Mutable();
 
@@ -650,6 +650,7 @@ void UNavWalkingMode::CaptureFinalState(USceneComponent* UpdatedComponent, const
 	OutputSyncState.SetTransforms_WorldSpace(UpdatedComponent->GetComponentLocation(),
 		UpdatedComponent->GetComponentRotation(),
 		Record.GetRelevantVelocity(),
+		AngularVelocityDegrees,
 		nullptr);	// no movement base
 
 	UpdatedComponent->ComponentVelocity = OutputSyncState.GetVelocity_WorldSpace();
diff --git a/Engine/Plugins/Experimental/Mover/Source/Mover/Private/DefaultMovementSet/Modes/SimpleSpringState.cpp b/Engine/Plugins/Experimental/Mover/Source/Mover/Private/DefaultMovementSet/Modes/SimpleSpringState.cpp
--- a/Engine/Plugins/Experimental/Mover/Source/Mover/Private/DefaultMovementSet/Modes/SimpleSpringState.cpp
+++ b/Engine/Plugins/Experimental/Mover/Source/Mover/Private/DefaultMovementSet/Modes/SimpleSpringState.cpp
@@ -19,9 +19,7 @@ bool FSimpleSpringState::NetSerialize(FArchive& Ar, UPackageMap* Map, bool& bOut
 	bool bSuccess = Super::NetSerialize(Ar, Map, bOutSuccess);
 
 	// Could be quantized to save bandwidth
-	Ar << CurrentVelocity;
 	Ar << CurrentAccel;
-	Ar << CurrentAngularVelocity;
 
 	return bSuccess;
 }
@@ -30,16 +28,15 @@ void FSimpleSpringState::ToString(FAnsiStringBuilderBase& Out) const
 {
 	Super::ToString(Out);
 
-	Out.Appendf("Vel=%s Accel=%s AngVel=%s\n", *CurrentVelocity.ToCompactString(), *CurrentAccel.ToCompactString(), *CurrentAngularVelocity.ToCompactString());
+	Out.Appendf("Accel=%s\n", *CurrentAccel.ToCompactString());
 }
 
 bool FSimpleSpringState::ShouldReconcile(const FMoverDataStructBase& AuthorityState) const
 {
 	const FSimpleSpringState* AuthoritySpringState = static_cast<const FSimpleSpringState*>(&AuthorityState);
 
-	return (!(CurrentVelocity - AuthoritySpringState->CurrentVelocity).IsNearlyZero() ||
-			!(CurrentAccel - AuthoritySpringState->CurrentAccel).IsNearlyZero() ||
-			!(CurrentAngularVelocity - AuthoritySpringState->CurrentAngularVelocity).IsNearlyZero());
+	return !(CurrentAccel - AuthoritySpringState->CurrentAccel).IsNearlyZero();
+		
 }
 
 
@@ -48,7 +45,5 @@ void FSimpleSpringState::Interpolate(const FMoverDataStructBase& From, const FMo
 	const FSimpleSpringState* FromState = static_cast<const FSimpleSpringState*>(&From);
 	const FSimpleSpringState* ToState   = static_cast<const FSimpleSpringState*>(&To);
 
-	CurrentVelocity        = FMath::Lerp(FromState->CurrentVelocity, ToState->CurrentVelocity, Pct);
 	CurrentAccel           = FMath::Lerp(FromState->CurrentAccel, ToState->CurrentAccel, Pct);
-	CurrentAngularVelocity = FMath::Lerp(FromState->CurrentAngularVelocity, ToState->CurrentAngularVelocity, Pct);
 }
diff --git a/Engine/Plugins/Experimental/Mover/Source/Mover/Private/DefaultMovementSet/Modes/SimpleSpringState.h b/Engine/Plugins/Experimental/Mover/Source/Mover/Private/DefaultMovementSet/Modes/SimpleSpringState.h
--- a/Engine/Plugins/Experimental/Mover/Source/Mover/Private/DefaultMovementSet/Modes/SimpleSpringState.h
+++ b/Engine/Plugins/Experimental/Mover/Source/Mover/Private/DefaultMovementSet/Modes/SimpleSpringState.h
@@ -8,7 +8,7 @@
 #include "SimpleSpringState.generated.h"
 
 /**
-* Internal state data for the SimpleWalkingMode
+* Internal state data for the SimpleSpringWalkingMode
 */
 USTRUCT()
 struct FSimpleSpringState : public FMoverDataStructBase
@@ -22,16 +22,7 @@ struct FSimpleSpringState : public FMoverDataStructBase
 	virtual bool ShouldReconcile(const FMoverDataStructBase& AuthorityState) const override;
 	virtual void Interpolate(const FMoverDataStructBase& From, const FMoverDataStructBase& To, float Pct) override;
 
-
-	// Velocity of internal spring model. Not necessarily the velocity of the character
-	UPROPERTY(BlueprintReadOnly, Category = "Mover|Experimental")
-	FVector CurrentVelocity = FVector::ZeroVector;
-
 	// Acceleration of internal spring model
 	UPROPERTY(BlueprintReadOnly, Category = "Mover|Experimental")
 	FVector CurrentAccel = FVector::ZeroVector;
-
-	// Angular velocity of internal spring model in radians per second
-	UPROPERTY(BlueprintReadOnly, Category = "Mover|Experimental")
-	FVector CurrentAngularVelocity = FVector::ZeroVector;
 };
diff --git a/Engine/Plugins/Experimental/Mover/Source/Mover/Private/DefaultMovementSet/Modes/SimpleSpringWalkingMode.cpp b/Engine/Plugins/Experimental/Mover/Source/Mover/Private/DefaultMovementSet/Modes/SimpleSpringWalkingMode.cpp
--- a/Engine/Plugins/Experimental/Mover/Source/Mover/Private/DefaultMovementSet/Modes/SimpleSpringWalkingMode.cpp
+++ b/Engine/Plugins/Experimental/Mover/Source/Mover/Private/DefaultMovementSet/Modes/SimpleSpringWalkingMode.cpp
@@ -7,6 +7,7 @@
 #include "DefaultMovementSet/Settings/CommonLegacyMovementSettings.h"
 #include "Animation/SpringMath.h"
 
+
 #include UE_INLINE_GENERATED_CPP_BY_NAME(SimpleSpringWalkingMode)
 
 void USimpleSpringWalkingMode::SimulationTick_Implementation(const FSimulationTickParams& Params, FMoverTickEndData& OutputState)
@@ -22,20 +23,19 @@ void USimpleSpringWalkingMode::SimulationTick_Implementation(const FSimulationTi
 }
 
 void USimpleSpringWalkingMode::GenerateWalkMove_Implementation(FMoverTickStartData& StartState, float DeltaSeconds, const FVector& DesiredVelocity,
-	const FQuat& DesiredFacing, FVector& OutVelocity, FQuat& OutFacing)
+	const FQuat& DesiredFacing, const FQuat& CurrentFacing, FVector& InOutAngularVelocityDegrees, FVector& InOutVelocity)
 {
-	const FMoverDefaultSyncState* StartingSyncState = StartState.SyncState.SyncStateCollection.FindDataByType<FMoverDefaultSyncState>();
-	FVector CurrentPosition = StartingSyncState->GetLocation_WorldSpace();
-	FVector PrevPosition = CurrentPosition;
-
 	FSimpleSpringState& SpringState = StartState.SyncState.SyncStateCollection.FindOrAddMutableDataByType<FSimpleSpringState>();
 
 	// Linear //
-	SpringMath::SpringCharacterUpdate(CurrentPosition, SpringState.CurrentVelocity, SpringState.CurrentAccel, DesiredVelocity, VelocitySmoothingTime, DeltaSeconds, VelocityDeadzoneThreshold);
-	OutVelocity = (CurrentPosition - PrevPosition) / DeltaSeconds;
-
+	
+	SpringMath::CriticalSpringDamper(InOutVelocity, SpringState.CurrentAccel, DesiredVelocity, VelocitySmoothingTime, DeltaSeconds);
+	
 	// Angular //
-	OutFacing = StartingSyncState->GetOrientation_WorldSpace().Quaternion();
-	SpringMath::CriticalSpringDamperQuat(OutFacing, SpringState.CurrentAngularVelocity, DesiredFacing, FacingSmoothingTime, DeltaSeconds);
+	
+	FVector CurrentAngularVelocityRad = FMath::DegreesToRadians(InOutAngularVelocityDegrees);
+	FQuat UpdatedFacing = CurrentFacing;
+	SpringMath::CriticalSpringDamperQuat(UpdatedFacing, CurrentAngularVelocityRad, DesiredFacing, FacingSmoothingTime, DeltaSeconds);
+	InOutAngularVelocityDegrees = FMath::RadiansToDegrees(CurrentAngularVelocityRad);
 }
 
diff --git a/Engine/Plugins/Experimental/Mover/Source/Mover/Private/DefaultMovementSet/Modes/SimpleWalkingMode.cpp b/Engine/Plugins/Experimental/Mover/Source/Mover/Private/DefaultMovementSet/Modes/SimpleWalkingMode.cpp
--- a/Engine/Plugins/Experimental/Mover/Source/Mover/Private/DefaultMovementSet/Modes/SimpleWalkingMode.cpp
+++ b/Engine/Plugins/Experimental/Mover/Source/Mover/Private/DefaultMovementSet/Modes/SimpleWalkingMode.cpp
@@ -83,36 +83,31 @@ void USimpleWalkingMode::GenerateMove_Implementation(const FMoverTickStartData&
 	DesiredFacingDir -= DesiredFacingDir.ProjectOnTo(GetMoverComponent()->GetUpDirection());
 	FQuat CurrentFacing = StartingSyncState->GetOrientation_WorldSpace().Quaternion();
 	FQuat DesiredFacing = CurrentFacing;
-	FQuat PrevFacing = CurrentFacing;
 	
 	if (DesiredFacingDir.Normalize())
 	{
 		DesiredFacing = FQuat::FindBetween(FVector::ForwardVector, DesiredFacingDir);
 	}
 
-	FQuat OutputFacing = CurrentFacing;
 	OutProposedMove.LinearVelocity = StartingSyncState->GetVelocity_WorldSpace();
+	FVector AngularVelocityDegrees = StartingSyncState->GetAngularVelocityDegrees_WorldSpace();
 	
 	// Hack const_cast stuff
 	// Why is this needed?
 	// Because some modes want to mutate their data inside the generate walk move
 	USimpleWalkingMode* MutableSimpleWalkMode  = const_cast<USimpleWalkingMode*>(this);
-	MutableSimpleWalkMode->GenerateWalkMove(const_cast<FMoverTickStartData&>(StartState), DeltaSeconds, DesiredVelocity, DesiredFacing, OutProposedMove.LinearVelocity, OutputFacing);
+	MutableSimpleWalkMode->GenerateWalkMove(const_cast<FMoverTickStartData&>(StartState), DeltaSeconds, DesiredVelocity, DesiredFacing, CurrentFacing, AngularVelocityDegrees, OutProposedMove.LinearVelocity);
 
-	// Calc angular velocity
-	FQuat Diff = (PrevFacing.Inverse() * OutputFacing);
-	Diff.Normalize();
-	Diff.EnforceShortestArcWith(FQuat::Identity);
-
-	FRotator AngularVelocity = Diff.Rotator();
-	AngularVelocity *= (1.0f / DeltaSeconds);
-	OutProposedMove.AngularVelocity = AngularVelocity;
+	// Calc angular velocity from final facing
+	OutProposedMove.AngularVelocityDegrees = AngularVelocityDegrees;
 }
 
 void USimpleWalkingMode::GenerateWalkMove_Implementation(FMoverTickStartData& StartState, float DeltaSeconds, const FVector& DesiredVelocity,
-	const FQuat& DesiredFacing, FVector& OutVelocity, FQuat& OutFacing)
+	const FQuat& DesiredFacing, const FQuat& CurrentFacing, FVector& InOutAngularVelocityDegrees, FVector& InOutVelocity)
 {
-	OutVelocity = DesiredVelocity;
-	OutFacing = DesiredFacing;
+	InOutVelocity = DesiredVelocity;
+
+	FQuat ToFacing = CurrentFacing.Inverse() * DesiredFacing;
+	InOutAngularVelocityDegrees = DeltaSeconds > 0.0f ? FMath::RadiansToDegrees(ToFacing.ToRotationVector() / DeltaSeconds) : FVector::ZeroVector;
 }
 
diff --git a/Engine/Plugins/Experimental/Mover/Source/Mover/Private/DefaultMovementSet/Modes/SmoothWalkingMode.cpp b/Engine/Plugins/Experimental/Mover/Source/Mover/Private/DefaultMovementSet/Modes/SmoothWalkingMode.cpp
--- a/Engine/Plugins/Experimental/Mover/Source/Mover/Private/DefaultMovementSet/Modes/SmoothWalkingMode.cpp
+++ b/Engine/Plugins/Experimental/Mover/Source/Mover/Private/DefaultMovementSet/Modes/SmoothWalkingMode.cpp
@@ -22,7 +22,7 @@ void USmoothWalkingMode::SimulationTick_Implementation(const FSimulationTickPara
 }
 
 void USmoothWalkingMode::GenerateWalkMove_Implementation(FMoverTickStartData& StartState, float DeltaSeconds, const FVector& DesiredVelocity,
-	const FQuat& DesiredFacing, FVector& InOutVelocity, FQuat& InOutFacing) 
+	const FQuat& DesiredFacing, const FQuat& CurrentFacing, FVector& InOutAngularVelocityDegrees, FVector& InOutVelocity) 
 {
 	if (DeltaSeconds <= FLT_EPSILON)
 	{
@@ -47,9 +47,8 @@ void USmoothWalkingMode::GenerateWalkMove_Implementation(FMoverTickStartData& St
 		SpringState.SpringVelocity = InOutVelocity;
 		SpringState.SpringAcceleration = FVector::ZeroVector;
 		SpringState.IntermediateVelocity = InOutVelocity;
-		SpringState.IntermediateFacing = InOutFacing;
+		SpringState.IntermediateFacing = CurrentFacing;
 		SpringState.IntermediateAngularVelocity = FVector::ZeroVector;
-		SpringState.CurrentAngularVelocity = FVector::ZeroVector;
 	}
 
 	// We can compute the degree to which the internal velocity matches the actual velocity by projecting the spring velocity onto the actual velocity 
@@ -146,33 +145,41 @@ void USmoothWalkingMode::GenerateWalkMove_Implementation(FMoverTickStartData& St
 	// Update the intermediate velocity
 	SpringState.IntermediateVelocity = NextVelocity;
 
+	FVector CurrentAngularVelocityRadians = FMath::DegreesToRadians(InOutAngularVelocityDegrees);
+	FQuat UpdatedFacing = CurrentFacing;
+
 	// For controlling the facing direction we use either a spring or a double spring
 	if (bSmoothFacingWithDoubleSpring)
 	{
 		SpringMath::CriticalSpringDamperQuat(SpringState.IntermediateFacing, SpringState.IntermediateAngularVelocity, DesiredFacing, FacingSmoothingTime / 2.0f, DeltaSeconds);
-		SpringMath::CriticalSpringDamperQuat(InOutFacing, SpringState.CurrentAngularVelocity, SpringState.IntermediateFacing, FacingSmoothingTime / 2.0f, DeltaSeconds);
+		SpringMath::CriticalSpringDamperQuat(UpdatedFacing, CurrentAngularVelocityRadians, SpringState.IntermediateFacing, FacingSmoothingTime / 2.0f, DeltaSeconds);
 	}
 	else
 	{
 		SpringState.IntermediateFacing = DesiredFacing;
-		SpringState.IntermediateAngularVelocity = SpringState.CurrentAngularVelocity;
-		SpringMath::CriticalSpringDamperQuat(InOutFacing, SpringState.CurrentAngularVelocity, DesiredFacing, FacingSmoothingTime, DeltaSeconds);
+		SpringState.IntermediateAngularVelocity = CurrentAngularVelocityRadians;
+		SpringMath::CriticalSpringDamperQuat(UpdatedFacing, CurrentAngularVelocityRadians, DesiredFacing, FacingSmoothingTime, DeltaSeconds);
 	}
 
 	// Snap the facing to the desired facing based on the dead-zone
-	if (DesiredFacing.AngularDistance(InOutFacing) < FMath::DegreesToRadians(FacingDeadzoneThreshold))
+	if (DesiredFacing.AngularDistance(UpdatedFacing) < FMath::DegreesToRadians(FacingDeadzoneThreshold))
 	{
 		// We reached our target
-		InOutFacing = DesiredFacing;
+		// Ensure the output angular velocity will snap perfectly to the target
+		// Note we don't do this normally because its better to have a consistent angular velocity
+		// If we output the angular velocity based on the updated facing every frame it can cause errors at low dt due to inaccuracy in the inverse
+		// exponential approximation inside thne spring damper
+		CurrentAngularVelocityRadians = DeltaSeconds > 0.0f ? (CurrentFacing.Inverse() * UpdatedFacing).ToRotationVector() / DeltaSeconds : FVector::ZeroVector;
 		SpringState.IntermediateFacing = DesiredFacing;
 
 		// If we've reached our target then also snap the angular velocity to zero once it is close enough
-		if (SpringState.CurrentAngularVelocity.SquaredLength() < FMath::Square(FMath::DegreesToRadians(AngularVelocityDeadzoneThreshold)))
+		if (CurrentAngularVelocityRadians.SquaredLength() < FMath::Square(FMath::DegreesToRadians(AngularVelocityDeadzoneThreshold)))
 		{
-			SpringState.CurrentAngularVelocity = FVector::ZeroVector;
 			SpringState.IntermediateAngularVelocity = FVector::ZeroVector;
 		}
 	}
+
+	InOutAngularVelocityDegrees = FMath::RadiansToDegrees(CurrentAngularVelocityRadians);
 }
 
 
diff --git a/Engine/Plugins/Experimental/Mover/Source/Mover/Private/DefaultMovementSet/Modes/SmoothWalkingState.cpp b/Engine/Plugins/Experimental/Mover/Source/Mover/Private/DefaultMovementSet/Modes/SmoothWalkingState.cpp
--- a/Engine/Plugins/Experimental/Mover/Source/Mover/Private/DefaultMovementSet/Modes/SmoothWalkingState.cpp
+++ b/Engine/Plugins/Experimental/Mover/Source/Mover/Private/DefaultMovementSet/Modes/SmoothWalkingState.cpp
@@ -24,7 +24,6 @@ bool FSmoothWalkingState::NetSerialize(FArchive& Ar, UPackageMap* Map, bool& bOu
 	Ar << IntermediateVelocity;
 	Ar << IntermediateFacing;
 	Ar << IntermediateAngularVelocity;
-	Ar << CurrentAngularVelocity;
 
 	return bSuccess;
 }
@@ -33,13 +32,12 @@ void FSmoothWalkingState::ToString(FAnsiStringBuilderBase& Out) const
 {
 	Super::ToString(Out);
 
-	Out.Appendf("SpringVelocity=%s SpringAcceleration=%s IntVel=%s IntFac=%s IntAng=%s AngVel=%s\n",
+	Out.Appendf("SpringVelocity=%s SpringAcceleration=%s IntVel=%s IntFac=%s IntAng=%s\n",
 		*SpringVelocity.ToCompactString(),
 		*SpringAcceleration.ToCompactString(),
 		*IntermediateVelocity.ToCompactString(),
 		*IntermediateFacing.ToString(), 
-		*IntermediateAngularVelocity.ToString(),  
-		*CurrentAngularVelocity.ToCompactString());
+		*IntermediateAngularVelocity.ToString());
 }
 
 bool FSmoothWalkingState::ShouldReconcile(const FMoverDataStructBase& AuthorityState) const
@@ -50,8 +48,7 @@ bool FSmoothWalkingState::ShouldReconcile(const FMoverDataStructBase& AuthorityS
 			!(SpringAcceleration - AuthoritySpringState->SpringAcceleration).IsNearlyZero() ||
 			!(IntermediateVelocity - AuthoritySpringState->IntermediateVelocity).IsNearlyZero() ||
 			 (IntermediateFacing.AngularDistance(AuthoritySpringState->IntermediateFacing) > UE_SMALL_NUMBER ||
-			!(IntermediateAngularVelocity - AuthoritySpringState->IntermediateAngularVelocity).IsNearlyZero() ||
-			!(CurrentAngularVelocity - AuthoritySpringState->CurrentAngularVelocity).IsNearlyZero()));
+			!(IntermediateAngularVelocity - AuthoritySpringState->IntermediateAngularVelocity).IsNearlyZero()));
 }
 
 
@@ -65,5 +62,4 @@ void FSmoothWalkingState::Interpolate(const FMoverDataStructBase& From, const FM
 	IntermediateVelocity = FMath::Lerp(FromState->IntermediateVelocity, ToState->IntermediateVelocity, Pct);
 	IntermediateFacing = FQuat::Slerp(FromState->IntermediateFacing, ToState->IntermediateFacing, Pct);
 	IntermediateAngularVelocity = FMath::Lerp(FromState->IntermediateAngularVelocity, ToState->IntermediateAngularVelocity, Pct);
-	CurrentAngularVelocity = FMath::Lerp(FromState->CurrentAngularVelocity, ToState->CurrentAngularVelocity, Pct);
 }
diff --git a/Engine/Plugins/Experimental/Mover/Source/Mover/Private/DefaultMovementSet/Modes/SmoothWalkingState.h b/Engine/Plugins/Experimental/Mover/Source/Mover/Private/DefaultMovementSet/Modes/SmoothWalkingState.h
--- a/Engine/Plugins/Experimental/Mover/Source/Mover/Private/DefaultMovementSet/Modes/SmoothWalkingState.h
+++ b/Engine/Plugins/Experimental/Mover/Source/Mover/Private/DefaultMovementSet/Modes/SmoothWalkingState.h
@@ -41,8 +41,4 @@ struct FSmoothWalkingState : public FMoverDataStructBase
 	// Angular velocity of the intermediate spring when using a double spring
 	UPROPERTY(BlueprintReadOnly, Category = "Mover|Experimental")
 	FVector IntermediateAngularVelocity = FVector::ZeroVector;
-
-	// Angular velocity of rotational spring model in radians per second
-	UPROPERTY(BlueprintReadOnly, Category = "Mover|Experimental")
-	FVector CurrentAngularVelocity = FVector::ZeroVector;
 };
diff --git a/Engine/Plugins/Experimental/Mover/Source/Mover/Private/DefaultMovementSet/Modes/WalkingMode.cpp b/Engine/Plugins/Experimental/Mover/Source/Mover/Private/DefaultMovementSet/Modes/WalkingMode.cpp
--- a/Engine/Plugins/Experimental/Mover/Source/Mover/Private/DefaultMovementSet/Modes/WalkingMode.cpp
+++ b/Engine/Plugins/Experimental/Mover/Source/Mover/Private/DefaultMovementSet/Modes/WalkingMode.cpp
@@ -105,7 +105,7 @@ void UWalkingMode::GenerateMove_Implementation(const FMoverTickStartData& StartS
 
 	if (TurnGenerator)
 	{
-		OutProposedMove.AngularVelocity = ITurnGeneratorInterface::Execute_GetTurn(TurnGenerator, IntendedOrientation_WorldSpace, StartState, *StartingSyncState, TimeStep, OutProposedMove, SimBlackboard);
+		OutProposedMove.AngularVelocityDegrees = ITurnGeneratorInterface::Execute_GetTurn(TurnGenerator, IntendedOrientation_WorldSpace, StartState, *StartingSyncState, TimeStep, OutProposedMove, SimBlackboard);
 	}
 }
 
@@ -150,7 +150,8 @@ void UWalkingMode::SimulationTick_Implementation(const FSimulationTickParams& Pa
 	OutputSyncState.MoveDirectionIntent = (ProposedMove.bHasDirIntent ? ProposedMove.DirectionIntent : FVector::ZeroVector);
 
 	const FRotator StartingOrient = StartingSyncState->GetOrientation_WorldSpace();
-	const FRotator TargetOrient = UMovementUtils::ApplyAngularVelocity(StartingOrient, ProposedMove.AngularVelocity, DeltaSeconds);
+	const FRotator TargetOrient = UMovementUtils::ApplyAngularVelocityToRotator(StartingOrient, ProposedMove.AngularVelocityDegrees, DeltaSeconds);
+	
 	const bool bIsOrientationChanging = !StartingOrient.Equals(TargetOrient);
 
 	FQuat TargetOrientQuat = TargetOrient.Quaternion();
@@ -245,7 +246,7 @@ void UWalkingMode::SimulationTick_Implementation(const FSimulationTickParams& Pa
 			OutputState.MovementEndState.NextModeName = CommonLegacySettings->AirMovementModeName;
 			OutputState.MovementEndState.RemainingMs = Params.TimeStep.StepMs - (Params.TimeStep.StepMs * PercentTimeAppliedSoFar);
 			MoveRecord.SetDeltaSeconds((Params.TimeStep.StepMs - OutputState.MovementEndState.RemainingMs) * 0.001f);
-			CaptureFinalState(UpdatedComponent, bDidAttemptMovement, CurrentFloor, MoveRecord, OutputSyncState);
+			CaptureFinalState(UpdatedComponent, bDidAttemptMovement, CurrentFloor, MoveRecord, ProposedMove.AngularVelocityDegrees, OutputSyncState);
 			return;
 		}
 	}
@@ -277,13 +278,13 @@ void UWalkingMode::SimulationTick_Implementation(const FSimulationTickParams& Pa
 				OutputState.MovementEndState.NextModeName = CommonLegacySettings->AirMovementModeName;
 				OutputState.MovementEndState.RemainingMs = Params.TimeStep.StepMs;
 				MoveRecord.SetDeltaSeconds((Params.TimeStep.StepMs - OutputState.MovementEndState.RemainingMs) * 0.001f);
-				CaptureFinalState(UpdatedComponent, bDidAttemptMovement, CurrentFloor, MoveRecord, OutputSyncState);
+				CaptureFinalState(UpdatedComponent, bDidAttemptMovement, CurrentFloor, MoveRecord, ProposedMove.AngularVelocityDegrees, OutputSyncState);
 				return;
 			}
 		}
 	}
 
-	CaptureFinalState(UpdatedComponent, bDidAttemptMovement, CurrentFloor, MoveRecord, OutputSyncState);
+	CaptureFinalState(UpdatedComponent, bDidAttemptMovement, CurrentFloor, MoveRecord, ProposedMove.AngularVelocityDegrees, OutputSyncState);
 
 }
 
@@ -320,7 +321,7 @@ void UWalkingMode::OnUnregistered()
 	Super::OnUnregistered();
 }
 
-void UWalkingMode::CaptureFinalState(USceneComponent* UpdatedComponent, bool bDidAttemptMovement, const FFloorCheckResult& FloorResult, const FMovementRecord& Record, FMoverDefaultSyncState& OutputSyncState) const
+void UWalkingMode::CaptureFinalState(USceneComponent* UpdatedComponent, bool bDidAttemptMovement, const FFloorCheckResult& FloorResult, const FMovementRecord& Record, const FVector& AngularVelocityDegrees, FMoverDefaultSyncState& OutputSyncState) const
 {
 	FRelativeBaseInfo PriorBaseInfo;
 
@@ -349,6 +350,8 @@ void UWalkingMode::CaptureFinalState(USceneComponent* UpdatedComponent, bool bDi
 		OutputSyncState.SetTransforms_WorldSpace( UpdatedComponent->GetComponentLocation(),
 												  UpdatedComponent->GetComponentRotation(),
 												  Record.GetRelevantVelocity(),
+												  AngularVelocityDegrees,
+												  
 												  CurrentBaseInfo.MovementBase.Get(), CurrentBaseInfo.BoneName);
 	}
 	else
@@ -358,6 +361,7 @@ void UWalkingMode::CaptureFinalState(USceneComponent* UpdatedComponent, bool bDi
 		OutputSyncState.SetTransforms_WorldSpace( UpdatedComponent->GetComponentLocation(),
 												  UpdatedComponent->GetComponentRotation(),
 												  Record.GetRelevantVelocity(),
+												  AngularVelocityDegrees,
 												  nullptr);	// no movement base
 	}
 
diff --git a/Engine/Plugins/Experimental/Mover/Source/Mover/Private/MoveLibrary/AirMovementUtils.cpp b/Engine/Plugins/Experimental/Mover/Source/Mover/Private/MoveLibrary/AirMovementUtils.cpp
--- a/Engine/Plugins/Experimental/Mover/Source/Mover/Private/MoveLibrary/AirMovementUtils.cpp
+++ b/Engine/Plugins/Experimental/Mover/Source/Mover/Private/MoveLibrary/AirMovementUtils.cpp
@@ -32,7 +32,7 @@ FProposedMove UAirMovementUtils::ComputeControlledFreeMove(const FFreeMoveParams
 	ComputeVelocityParams.Friction = InParams.Friction;
 	
 	OutMove.LinearVelocity = UMovementUtils::ComputeVelocity(ComputeVelocityParams);
-	OutMove.AngularVelocity = UMovementUtils::ComputeAngularVelocity(InParams.PriorOrientation, InParams.OrientationIntent, InParams.WorldToGravityQuat, InParams.DeltaSeconds, InParams.TurningRate);
+	OutMove.AngularVelocityDegrees = UMovementUtils::ComputeAngularVelocityDegrees(InParams.PriorOrientation, InParams.OrientationIntent, InParams.DeltaSeconds, InParams.TurningRate);
 
 	return OutMove;
 }
diff --git a/Engine/Plugins/Experimental/Mover/Source/Mover/Private/MoveLibrary/BasedMovementUtils.cpp b/Engine/Plugins/Experimental/Mover/Source/Mover/Private/MoveLibrary/BasedMovementUtils.cpp
--- a/Engine/Plugins/Experimental/Mover/Source/Mover/Private/MoveLibrary/BasedMovementUtils.cpp
+++ b/Engine/Plugins/Experimental/Mover/Source/Mover/Private/MoveLibrary/BasedMovementUtils.cpp
@@ -445,6 +445,7 @@ void UBasedMovementUtils::UpdateSimpleBasedMovement(UMoverComponent* TargetMover
 							NewSyncTransformWs.GetLocation(),
 							NewSyncTransformWs.GetRotation().Rotator(),
 							PendingMoverState->GetVelocity_WorldSpace(),	// keep same velocity and base
+							PendingMoverState->GetAngularVelocityDegrees_WorldSpace(),
 							PendingMoverState->GetMovementBase(), PendingMoverState->GetMovementBaseBoneName());
 
 						TargetMoverComp->BackendLiaisonComp->WritePendingSyncState(PendingSimSyncState);	// writes pending Simulation state
@@ -467,6 +468,7 @@ void UBasedMovementUtils::UpdateSimpleBasedMovement(UMoverComponent* TargetMover
 										NewSyncTransformWs.GetLocation(),
 										NewSyncTransformWs.GetRotation().Rotator(),
 										PresentationMoverState->GetVelocity_WorldSpace(),	// keep same velocity and base
+										PresentationMoverState->GetAngularVelocityDegrees_WorldSpace(),
 										PresentationMoverState->GetMovementBase(), PresentationMoverState->GetMovementBaseBoneName());
 
 									TargetMoverComp->BackendLiaisonComp->WritePresentationSyncState(PresentationSyncState);
@@ -486,6 +488,7 @@ void UBasedMovementUtils::UpdateSimpleBasedMovement(UMoverComponent* TargetMover
 										NewSyncTransformWs.GetLocation(),
 										NewSyncTransformWs.GetRotation().Rotator(),
 										PrevPresentationMoverState->GetVelocity_WorldSpace(),	// keep same velocity and base
+										PrevPresentationMoverState->GetAngularVelocityDegrees_WorldSpace(),
 										PrevPresentationMoverState->GetMovementBase(), PrevPresentationMoverState->GetMovementBaseBoneName());
 
 									TargetMoverComp->BackendLiaisonComp->WritePrevPresentationSyncState(PrevPresentationSyncState);	
diff --git a/Engine/Plugins/Experimental/Mover/Source/Mover/Private/MoveLibrary/GroundMovementUtils.cpp b/Engine/Plugins/Experimental/Mover/Source/Mover/Private/MoveLibrary/GroundMovementUtils.cpp
--- a/Engine/Plugins/Experimental/Mover/Source/Mover/Private/MoveLibrary/GroundMovementUtils.cpp
+++ b/Engine/Plugins/Experimental/Mover/Source/Mover/Private/MoveLibrary/GroundMovementUtils.cpp
@@ -45,7 +45,7 @@ FProposedMove UGroundMovementUtils::ComputeControlledGroundMove(const FGroundMov
 	OutMove.LinearVelocity = UMovementUtils::ConstrainToPlane(Velocity, GroundSurfacePlane, true);
 
 	// Linearly rotate in place
-	OutMove.AngularVelocity = UMovementUtils::ComputeAngularVelocity(InParams.PriorOrientation, InParams.OrientationIntent, InParams.WorldToGravityQuat, InParams.DeltaSeconds, InParams.TurningRate);
+	OutMove.AngularVelocityDegrees = UMovementUtils::ComputeAngularVelocityDegrees(InParams.PriorOrientation, InParams.OrientationIntent, InParams.DeltaSeconds, InParams.TurningRate);
 
 	return OutMove;
 }
diff --git a/Engine/Plugins/Experimental/Mover/Source/Mover/Private/MoveLibrary/ModularMovement.cpp b/Engine/Plugins/Experimental/Mover/Source/Mover/Private/MoveLibrary/ModularMovement.cpp
--- a/Engine/Plugins/Experimental/Mover/Source/Mover/Private/MoveLibrary/ModularMovement.cpp
+++ b/Engine/Plugins/Experimental/Mover/Source/Mover/Private/MoveLibrary/ModularMovement.cpp
@@ -3,13 +3,14 @@
 #include "MoveLibrary/ModularMovement.h"
 #include "MoverTypes.h"
 #include "MoverSimulationTypes.h"
+#include "MoveLibrary/MovementUtils.h"
 
 #include UE_INLINE_GENERATED_CPP_BY_NAME(ModularMovement)
 
 
-FRotator ULinearTurnGenerator::GetTurn_Implementation(FRotator TargetOrientation, const FMoverTickStartData& FullStartState, const FMoverDefaultSyncState& MoverState, const FMoverTimeStep& TimeStep, const FProposedMove& ProposedMove, UMoverBlackboard* SimBlackboard)
+FVector ULinearTurnGenerator::GetTurn_Implementation(FRotator TargetOrientation, const FMoverTickStartData& FullStartState, const FMoverDefaultSyncState& MoverState, const FMoverTimeStep& TimeStep, const FProposedMove& ProposedMove, UMoverBlackboard* SimBlackboard)
 {
-	FRotator AngularVelocityDpS(FRotator::ZeroRotator);
+	FVector AngularVelocityDpS(FVector::ZeroVector);
 
 	const float DeltaSeconds = TimeStep.StepMs * 0.001f;
 
@@ -20,24 +21,25 @@ FRotator ULinearTurnGenerator::GetTurn_Implementation(FRotator TargetOrientation
 
 		AngularDelta.GetWindingAndRemainder(Winding, Remainder);	// to find the fastest turn, just get the (-180,180] remainder
 
-		AngularVelocityDpS = Remainder * (1.f / DeltaSeconds);
-
-		FRotator TurnRates(PitchRate, HeadingRate, RollRate);
+		FRotator AngularVelocityRotator = Remainder * (1.f / DeltaSeconds);
 
 		if (HeadingRate >= 0.f)
 		{
-			AngularVelocityDpS.Yaw = FMath::Clamp(AngularVelocityDpS.Yaw, -HeadingRate, HeadingRate);
+			AngularVelocityRotator.Yaw = FMath::Clamp(AngularVelocityRotator.Yaw, -HeadingRate, HeadingRate);
 		}
 
 		if (PitchRate >= 0.f)
 		{
-			AngularVelocityDpS.Pitch = FMath::Clamp(AngularVelocityDpS.Pitch, -PitchRate, PitchRate);
+			AngularVelocityRotator.Pitch = FMath::Clamp(AngularVelocityRotator.Pitch, -PitchRate, PitchRate);
 		}
 
 		if (RollRate >= 0.f)
 		{
-			AngularVelocityDpS.Roll = FMath::Clamp(AngularVelocityDpS.Roll, -RollRate, RollRate);
+			AngularVelocityRotator.Roll = FMath::Clamp(AngularVelocityRotator.Roll, -RollRate, RollRate);
 		}
+		
+		AngularDelta = AngularVelocityRotator * DeltaSeconds;
+		AngularVelocityDpS = FMath::RadiansToDegrees(AngularDelta.Quaternion().ToRotationVector()) / DeltaSeconds;
 	}
 
 	return AngularVelocityDpS;
@@ -50,9 +52,9 @@ static float CalcExactDampedInterpolation(float CurrentVal, float TargetVal, flo
 	return FMath::Lerp(CurrentVal, TargetVal, 1.f - FMath::Pow(2.f, (-DeltaTime / HalflifeTime)));
 }
 
-FRotator UExactDampedTurnGenerator::GetTurn_Implementation(FRotator TargetOrientation, const FMoverTickStartData& FullStartState, const FMoverDefaultSyncState& MoverState, const FMoverTimeStep& TimeStep, const FProposedMove& ProposedMove, UMoverBlackboard* SimBlackboard)
+FVector UExactDampedTurnGenerator::GetTurn_Implementation(FRotator TargetOrientation, const FMoverTickStartData& FullStartState, const FMoverDefaultSyncState& MoverState, const FMoverTimeStep& TimeStep, const FProposedMove& ProposedMove, UMoverBlackboard* SimBlackboard)
 {
-	FRotator AngularVelocityDpS(FRotator::ZeroRotator);
+	FVector AngularVelocityDpS(FVector::ZeroVector);
 
 	const float DeltaSeconds = TimeStep.StepMs * 0.001f;
 
@@ -64,10 +66,14 @@ FRotator UExactDampedTurnGenerator::GetTurn_Implementation(FRotator TargetOrient
 		AngularDelta.GetWindingAndRemainder(Winding, Remainder);	// to find the fastest turn, just get the (-180,180] remainder
 
 		const float OneOverDeltaSecs = 1.f / DeltaSeconds;
+		FRotator AngularVelocityRotator;
+
+		AngularVelocityRotator.Yaw   = CalcExactDampedInterpolation(0.f, Remainder.Yaw,   HalfLifeSeconds, DeltaSeconds) * OneOverDeltaSecs;
+		AngularVelocityRotator.Pitch = CalcExactDampedInterpolation(0.f, Remainder.Pitch, HalfLifeSeconds, DeltaSeconds) * OneOverDeltaSecs;
+		AngularVelocityRotator.Roll  = CalcExactDampedInterpolation(0.f, Remainder.Roll,  HalfLifeSeconds, DeltaSeconds) * OneOverDeltaSecs;
 
-		AngularVelocityDpS.Yaw   = CalcExactDampedInterpolation(0.f, Remainder.Yaw,   HalfLifeSeconds, DeltaSeconds) * OneOverDeltaSecs;
-		AngularVelocityDpS.Pitch = CalcExactDampedInterpolation(0.f, Remainder.Pitch, HalfLifeSeconds, DeltaSeconds) * OneOverDeltaSecs;
-		AngularVelocityDpS.Roll  = CalcExactDampedInterpolation(0.f, Remainder.Roll,  HalfLifeSeconds, DeltaSeconds) * OneOverDeltaSecs;
+		AngularDelta = AngularVelocityRotator * DeltaSeconds;
+		AngularVelocityDpS = FMath::RadiansToDegrees(AngularDelta.Quaternion().ToRotationVector()) / DeltaSeconds;
 	}
 
 	return AngularVelocityDpS;
diff --git a/Engine/Plugins/Experimental/Mover/Source/Mover/Private/MoveLibrary/MovementMixer.cpp b/Engine/Plugins/Experimental/Mover/Source/Mover/Private/MoveLibrary/MovementMixer.cpp
--- a/Engine/Plugins/Experimental/Mover/Source/Mover/Private/MoveLibrary/MovementMixer.cpp
+++ b/Engine/Plugins/Experimental/Mover/Source/Mover/Private/MoveLibrary/MovementMixer.cpp
@@ -48,7 +48,7 @@ void UMovementMixer::MixLayeredMove(const FLayeredMoveBase& ActiveMove, const FP
 				
 			OutCumulativeMove.MixMode = EMoveMixMode::OverrideVelocity;
 			OutCumulativeMove.LinearVelocity  = MoveStep.LinearVelocity;
-			OutCumulativeMove.AngularVelocity = MoveStep.AngularVelocity;
+			OutCumulativeMove.AngularVelocityDegrees = MoveStep.AngularVelocityDegrees;
 		}
 	}
 	else if (MoveStep.MixMode == EMoveMixMode::AdditiveVelocity)
@@ -61,7 +61,7 @@ void UMovementMixer::MixLayeredMove(const FLayeredMoveBase& ActiveMove, const FP
 			}
 
 			OutCumulativeMove.LinearVelocity += MoveStep.LinearVelocity;
-			OutCumulativeMove.AngularVelocity += MoveStep.AngularVelocity;
+			OutCumulativeMove.AngularVelocityDegrees += MoveStep.AngularVelocityDegrees;
 		}
 	}
 	else if (MoveStep.MixMode == EMoveMixMode::OverrideAll)
@@ -131,7 +131,7 @@ void UMovementMixer::MixLayeredMove(const FLayeredMoveInstance& ActiveMove, cons
 				
 			OutCumulativeMove.MixMode = EMoveMixMode::OverrideVelocity;
 			OutCumulativeMove.LinearVelocity  = MoveStep.LinearVelocity;
-			OutCumulativeMove.AngularVelocity = MoveStep.AngularVelocity;
+			OutCumulativeMove.AngularVelocityDegrees = MoveStep.AngularVelocityDegrees;
 		}
 	}
 	else if (MoveStep.MixMode == EMoveMixMode::AdditiveVelocity)
@@ -144,7 +144,7 @@ void UMovementMixer::MixLayeredMove(const FLayeredMoveInstance& ActiveMove, cons
 			}
 
 			OutCumulativeMove.LinearVelocity += MoveStep.LinearVelocity;
-			OutCumulativeMove.AngularVelocity += MoveStep.AngularVelocity;
+			OutCumulativeMove.AngularVelocityDegrees += MoveStep.AngularVelocityDegrees;
 		}
 	}
 	else if (MoveStep.MixMode == EMoveMixMode::OverrideAll)
@@ -195,12 +195,12 @@ void UMovementMixer::MixProposedMoves(const FProposedMove& MoveToMix, FVector Up
 	else if (MoveToMix.MixMode == EMoveMixMode::AdditiveVelocity)
 	{
 		OutCumulativeMove.LinearVelocity += MoveToMix.LinearVelocity;
-		OutCumulativeMove.AngularVelocity += MoveToMix.AngularVelocity;
+		OutCumulativeMove.AngularVelocityDegrees += MoveToMix.AngularVelocityDegrees;
 	}
 	else if (MoveToMix.MixMode == EMoveMixMode::OverrideVelocity)
 	{
 		OutCumulativeMove.LinearVelocity = MoveToMix.LinearVelocity;
-		OutCumulativeMove.AngularVelocity = MoveToMix.AngularVelocity;
+		OutCumulativeMove.AngularVelocityDegrees = MoveToMix.AngularVelocityDegrees;
 	}
 	else if (MoveToMix.MixMode == EMoveMixMode::OverrideAllExceptVerticalVelocity)
 	{
diff --git a/Engine/Plugins/Experimental/Mover/Source/Mover/Private/MoveLibrary/MovementUtils.cpp b/Engine/Plugins/Experimental/Mover/Source/Mover/Private/MoveLibrary/MovementUtils.cpp
--- a/Engine/Plugins/Experimental/Mover/Source/Mover/Private/MoveLibrary/MovementUtils.cpp
+++ b/Engine/Plugins/Experimental/Mover/Source/Mover/Private/MoveLibrary/MovementUtils.cpp
@@ -604,6 +604,24 @@ FRotator UMovementUtils::ComputeAngularVelocity(const FRotator& FromOrientation,
 	return AngularVelocityDpS;
 }
 
+FVector UMovementUtils::ComputeAngularVelocityDegrees(const FRotator& From, const FRotator& To, float DeltaSeconds, float TurningRateLimit)
+{
+	if (DeltaSeconds <= UE_SMALL_NUMBER)
+	{
+		return FVector::ZeroVector;
+	}
+
+	FQuat Diff = From.Quaternion().Inverse() * To.Quaternion();
+	Diff.EnforceShortestArcWith(FQuat::Identity);
+	FVector AngularVelocity = FMath::RadiansToDegrees(Diff.ToRotationVector() / DeltaSeconds);
+
+	if (TurningRateLimit >= 0.0f)
+	{
+		AngularVelocity = AngularVelocity.GetClampedToMaxSize(TurningRateLimit);
+	}
+
+	return AngularVelocity;
+}
 
 FVector UMovementUtils::ComputeDirectionIntent(const FVector& MoveInput, EMoveInputType MoveInputType, float MaxSpeed)
 {
@@ -634,8 +652,10 @@ bool UMovementUtils::IsAngularVelocityZero(const FRotator& AngularVelocity)
 
 FRotator UMovementUtils::ApplyAngularVelocity(const FRotator& StartingOrient, const FRotator& AngularVelocity, float DeltaSeconds)
 {
+PRAGMA_DISABLE_DEPRECATION_WARNINGS
 	if (!UMovementUtils::IsAngularVelocityZero(AngularVelocity))
 	{
+PRAGMA_ENABLE_DEPRECATION_WARNINGS
 		const FQuat ProposedMoveQuat(AngularVelocity * DeltaSeconds);
 		const FQuat TargetOrientQuat = FQuat(StartingOrient) * ProposedMoveQuat;
 		return TargetOrientQuat.Rotator();
@@ -644,6 +664,21 @@ FRotator UMovementUtils::ApplyAngularVelocity(const FRotator& StartingOrient, co
 	return StartingOrient;
 }
 
+FRotator UMovementUtils::ApplyAngularVelocityToRotator(const FRotator& StartingOrient, const FVector& AngularVelocityDegrees, float DeltaSeconds)
+{
+	return ApplyAngularVelocityToQuat(StartingOrient.Quaternion(), AngularVelocityDegrees, DeltaSeconds).Rotator();
+}
+
+FQuat UMovementUtils::ApplyAngularVelocityToQuat(const FQuat& StartingOrient, const FVector& AngularVelocityDegrees, float DeltaSeconds)
+{
+	if (!AngularVelocityDegrees.IsZero())
+	{
+		const FQuat TargetOrientQuat = StartingOrient * FQuat::MakeFromRotationVector(FMath::DegreesToRadians(AngularVelocityDegrees) * DeltaSeconds);
+		return TargetOrientQuat;
+	}
+
+	return StartingOrient;
+}
 
 bool UMovementUtils::FindTeleportSpot(const UMoverComponent* MoverComp, FVector& TestLocation, FRotator TestRotation)
 {
diff --git a/Engine/Plugins/Experimental/Mover/Source/Mover/Private/MoveLibrary/WaterMovementUtils.cpp b/Engine/Plugins/Experimental/Mover/Source/Mover/Private/MoveLibrary/WaterMovementUtils.cpp
--- a/Engine/Plugins/Experimental/Mover/Source/Mover/Private/MoveLibrary/WaterMovementUtils.cpp
+++ b/Engine/Plugins/Experimental/Mover/Source/Mover/Private/MoveLibrary/WaterMovementUtils.cpp
@@ -27,7 +27,7 @@ FProposedMove UWaterMovementUtils::ComputeControlledWaterMove(const FWaterMovePa
 
 	// JAH TODO: this is where we can perform turning, based on aux settings. For now, just snap to the intended final orientation.
 	FVector IntendedFacingDir = InParams.OrientationIntent.RotateVector(FVector::ForwardVector).GetSafeNormal();
-	OutMove.AngularVelocity = UMovementUtils::ComputeAngularVelocity(InParams.PriorOrientation, IntendedFacingDir.ToOrientationRotator(), InParams.WorldToGravityQuat, InParams.DeltaSeconds, InParams.TurningRate);
+	OutMove.AngularVelocityDegrees = UMovementUtils::ComputeAngularVelocityDegrees(InParams.PriorOrientation, IntendedFacingDir.ToOrientationRotator(), InParams.DeltaSeconds, InParams.TurningRate);
 
 	return OutMove;
 }
diff --git a/Engine/Plugins/Experimental/Mover/Source/Mover/Private/MoverComponent.cpp b/Engine/Plugins/Experimental/Mover/Source/Mover/Private/MoverComponent.cpp
--- a/Engine/Plugins/Experimental/Mover/Source/Mover/Private/MoverComponent.cpp
+++ b/Engine/Plugins/Experimental/Mover/Source/Mover/Private/MoverComponent.cpp
@@ -969,7 +969,8 @@ void UMoverComponent::CreateDefaultInputAndState(FMoverInputCmdContext& OutInput
 		MoverState->SetTransforms_WorldSpace(
 			UpdatedComponent->GetComponentLocation(),
 			UpdatedComponent->GetComponentRotation(),
-			FVector::ZeroVector);	// no initial velocity
+			FVector::ZeroVector, // no initial velocity
+			FVector::ZeroVector);
 	}
 
 	OutSyncState.MovementMode = StartingMovementMode;
@@ -2309,8 +2310,9 @@ TArray<FTrajectorySampleInfo> UMoverComponent::GetPredictedTrajectory(FMoverPred
 
 					// Advance state based on move
 					StepSyncState->SetTransforms_WorldSpace(StepSyncState->GetLocation_WorldSpace() + (StepMove.LinearVelocity * PredictionParams.SecondsPerSample),
-						StepSyncState->GetOrientation_WorldSpace() + (StepMove.AngularVelocity * PredictionParams.SecondsPerSample),
+						UMovementUtils::ApplyAngularVelocityToRotator(StepSyncState->GetOrientation_WorldSpace(),StepMove.AngularVelocityDegrees, PredictionParams.SecondsPerSample),
 						StepMove.LinearVelocity,
+						StepMove.AngularVelocityDegrees,
 						StepSyncState->GetMovementBase(),
 						StepSyncState->GetMovementBaseBoneName());
 
@@ -2693,7 +2695,8 @@ void UMoverComponent::CheckForExternalMovement(const FMoverTickStartData& SimSta
 
 				MutableSyncState->SetTransforms_WorldSpace(ComponentTransform.GetLocation(), 
 				                                           ComponentTransform.GetRotation().Rotator(),
-				                                           MutableSyncState->GetVelocity_WorldSpace());
+				                                           MutableSyncState->GetVelocity_WorldSpace(),
+				                                           MutableSyncState->GetAngularVelocityDegrees_WorldSpace());
 			}
 		}
 	}
diff --git a/Engine/Plugins/Experimental/Mover/Source/Mover/Private/MoverDataModelTypes.cpp b/Engine/Plugins/Experimental/Mover/Source/Mover/Private/MoverDataModelTypes.cpp
--- a/Engine/Plugins/Experimental/Mover/Source/Mover/Private/MoverDataModelTypes.cpp
+++ b/Engine/Plugins/Experimental/Mover/Source/Mover/Private/MoverDataModelTypes.cpp
@@ -174,6 +174,7 @@ bool FMoverDefaultSyncState::NetSerialize(FArchive& Ar, UPackageMap* Map, bool&
 	SerializePackedVector<100, 30>(Location, Ar);
 	SerializeFixedVector<2, 8>(MoveDirectionIntent, Ar);
 	SerializePackedVector<10, 16>(Velocity, Ar);
+	SerializePackedVector<10, 16>(AngularVelocityDegrees, Ar);
 	Orientation.SerializeCompressedShort(Ar);
 
 	// Optional movement base
@@ -204,6 +205,7 @@ void FMoverDefaultSyncState::ToString(FAnsiStringBuilderBase& Out) const
 	Out.Appendf("Loc: X=%.2f Y=%.2f Z=%.2f\n", Location.X, Location.Y, Location.Z);
 	Out.Appendf("Intent: X=%.2f Y=%.2f Z=%.2f\n", MoveDirectionIntent.X, MoveDirectionIntent.Y, MoveDirectionIntent.Z);
 	Out.Appendf("Vel: X=%.2f Y=%.2f Z=%.2f\n", Velocity.X, Velocity.Y, Velocity.Z);
+	Out.Appendf("Ang Vel: X=%.2f Y=%.2f Z=%.2f\n", AngularVelocityDegrees.X, AngularVelocityDegrees.Y, AngularVelocityDegrees.Z);
 	Out.Appendf("Orient: P=%.2f Y=%.2f R=%.2f\n", Orientation.Pitch, Orientation.Yaw, Orientation.Roll);
 
 	if (const UPrimitiveComponent* MovementBasePtr = MovementBase.Get())
@@ -264,7 +266,7 @@ void FMoverDefaultSyncState::Interpolate(const FMoverDataStructBase& From, const
 		MovementBasePos      = ToState->MovementBasePos;
 		MovementBaseQuat     = ToState->MovementBaseQuat;
 
-		FVector FromLocation_ToSpace, FromIntent_ToSpace, FromVelocity_ToSpace;
+		FVector FromLocation_ToSpace, FromIntent_ToSpace, FromVelocity_ToSpace, FromAngularVelocity_ToSpace;
 		FRotator FromOrientation_ToSpace;
 
 	
@@ -280,6 +282,7 @@ void FMoverDefaultSyncState::Interpolate(const FMoverDataStructBase& From, const
 			FromLocation_ToSpace    = FromState->Location;
 			FromIntent_ToSpace      = FromState->MoveDirectionIntent;
 			FromVelocity_ToSpace    = FromState->Velocity;
+			FromAngularVelocity_ToSpace    = FromState->AngularVelocityDegrees;
 			FromOrientation_ToSpace = FromState->Orientation;
 		}
 		else if (ToState->MovementBase.IsValid())	// if moving onto a different base, regardless of coming from a prior base or not
@@ -287,6 +290,7 @@ void FMoverDefaultSyncState::Interpolate(const FMoverDataStructBase& From, const
 			UBasedMovementUtils::TransformLocationToLocal(ToState->MovementBasePos, ToState->MovementBaseQuat, FromState->GetLocation_WorldSpace(), OUT FromLocation_ToSpace);
 			UBasedMovementUtils::TransformDirectionToLocal(ToState->MovementBaseQuat, FromState->GetIntent_WorldSpace(), OUT FromIntent_ToSpace);
 			UBasedMovementUtils::TransformDirectionToLocal(ToState->MovementBaseQuat, FromState->GetVelocity_WorldSpace(), OUT FromVelocity_ToSpace);
+			UBasedMovementUtils::TransformDirectionToLocal(ToState->MovementBaseQuat, FromState->GetAngularVelocityDegrees_WorldSpace(), OUT FromAngularVelocity_ToSpace);
 			UBasedMovementUtils::TransformRotatorToLocal(ToState->MovementBaseQuat, FromState->GetOrientation_WorldSpace(), OUT FromOrientation_ToSpace);
 		}
 		else if (FromState->MovementBase.IsValid())	// if leaving a base
@@ -294,24 +298,33 @@ void FMoverDefaultSyncState::Interpolate(const FMoverDataStructBase& From, const
 			FromLocation_ToSpace	= FromState->GetLocation_WorldSpace();
 			FromIntent_ToSpace		= FromState->GetIntent_WorldSpace();
 			FromVelocity_ToSpace	= FromState->GetVelocity_WorldSpace();
+			FromAngularVelocity_ToSpace	= FromState->GetAngularVelocityDegrees_WorldSpace();
 			FromOrientation_ToSpace = FromState->GetOrientation_WorldSpace();
 		}
 
 		Location			= FMath::Lerp(FromLocation_ToSpace,		ToState->Location, Pct);
 		MoveDirectionIntent = FMath::Lerp(FromIntent_ToSpace,		ToState->MoveDirectionIntent, Pct);
 		Velocity			= FMath::Lerp(FromVelocity_ToSpace,		ToState->Velocity, Pct);
+		AngularVelocityDegrees = FMath::Lerp(FromAngularVelocity_ToSpace,ToState->AngularVelocityDegrees, Pct);
 		Orientation			= FMath::Lerp(FromOrientation_ToSpace,	ToState->Orientation, Pct);
 
 	}
 }
 
-void FMoverDefaultSyncState::SetTransforms_WorldSpace(FVector WorldLocation, FRotator WorldOrient, FVector WorldVelocity, UPrimitiveComponent* Base, FName BaseBone)
+void FMoverDefaultSyncState::SetTransforms_WorldSpace(FVector WorldLocation, FRotator WorldOrient, FVector WorldVelocity, UPrimitiveComponent* Base,
+	FName BaseBone)
+{
+	SetTransforms_WorldSpace(WorldLocation, WorldOrient, WorldVelocity, FVector::ZeroVector, Base, BaseBone);
+}
+
+void FMoverDefaultSyncState::SetTransforms_WorldSpace(FVector WorldLocation, FRotator WorldOrient, FVector WorldVelocity, FVector WorldAngularVelocityDegrees, UPrimitiveComponent* Base, FName BaseBone)
 {
 	if (SetMovementBase(Base, BaseBone))
 	{
 		UBasedMovementUtils::TransformLocationToLocal(  MovementBasePos,  MovementBaseQuat, WorldLocation, OUT Location);
 		UBasedMovementUtils::TransformRotatorToLocal(   MovementBaseQuat, WorldOrient, OUT Orientation);
 		UBasedMovementUtils::TransformDirectionToLocal( MovementBaseQuat, WorldVelocity, OUT Velocity);
+		UBasedMovementUtils::TransformDirectionToLocal( MovementBaseQuat, WorldAngularVelocityDegrees, OUT AngularVelocityDegrees);
 	}
 	else
 	{
@@ -323,6 +336,7 @@ void FMoverDefaultSyncState::SetTransforms_WorldSpace(FVector WorldLocation, FRo
 		Location = WorldLocation;
 		Orientation = WorldOrient;
 		Velocity = WorldVelocity;
+		AngularVelocityDegrees = WorldAngularVelocityDegrees;
 	}
 }
 
@@ -366,10 +380,10 @@ bool FMoverDefaultSyncState::IsNearlyEqual(const FMoverDefaultSyncState& Other)
 											 MovementBaseQuat.Equals(Other.MovementBaseQuat));
 
 	return (Location-Other.Location).IsNearlyZero() &&
-		   (Orientation-Other.Orientation).IsNearlyZero() &&
-		   (Velocity-Other.Velocity).IsNearlyZero() &&
-		   (MoveDirectionIntent-Other.MoveDirectionIntent).IsNearlyZero() &&
-		    bHasSameBaseBaseInfo;
+		(Orientation-Other.Orientation).IsNearlyZero() &&
+		(Velocity-Other.Velocity).IsNearlyZero() &&(AngularVelocityDegrees-Other.AngularVelocityDegrees).IsNearlyZero() &&
+		(MoveDirectionIntent-Other.MoveDirectionIntent).IsNearlyZero() &&
+			bHasSameBaseBaseInfo;
 }
 
 FVector FMoverDefaultSyncState::GetLocation_WorldSpace() const
@@ -450,6 +464,21 @@ FTransform FMoverDefaultSyncState::GetTransform_BaseSpace() const
 	return FTransform(Orientation, Location);
 }
 
+FVector FMoverDefaultSyncState::GetAngularVelocityDegrees_WorldSpace() const
+{
+	if (MovementBase.IsValid())
+	{
+		return MovementBaseQuat.RotateVector(AngularVelocityDegrees);
+	}
+
+	return AngularVelocityDegrees; // if no base, assumed to be in world space
+}
+
+FVector FMoverDefaultSyncState::GetAngularVelocityDegrees_BaseSpace() const
+{
+	return AngularVelocityDegrees;
+}
+
 
 // UMoverDataModelBlueprintLibrary ///////////////////////////////////////////////////
 
@@ -478,6 +507,11 @@ FVector UMoverDataModelBlueprintLibrary::GetVelocityFromSyncState(const FMoverDe
 	return SyncState.GetVelocity_WorldSpace();
 }
 
+FVector UMoverDataModelBlueprintLibrary::GetAngularVelocityDegreesFromSyncState(const FMoverDefaultSyncState& SyncState)
+{
+	return SyncState.GetAngularVelocityDegrees_WorldSpace();
+}
+
 FRotator UMoverDataModelBlueprintLibrary::GetOrientationFromSyncState(const FMoverDefaultSyncState& SyncState)
 {
 	return SyncState.GetOrientation_WorldSpace();
diff --git a/Engine/Plugins/Experimental/Mover/Source/Mover/Private/PhysicsMover/InstantMovementEffects/ApplyVelocityPhysicsMovementEffect.cpp b/Engine/Plugins/Experimental/Mover/Source/Mover/Private/PhysicsMover/InstantMovementEffects/ApplyVelocityPhysicsMovementEffect.cpp
--- a/Engine/Plugins/Experimental/Mover/Source/Mover/Private/PhysicsMover/InstantMovementEffects/ApplyVelocityPhysicsMovementEffect.cpp
+++ b/Engine/Plugins/Experimental/Mover/Source/Mover/Private/PhysicsMover/InstantMovementEffects/ApplyVelocityPhysicsMovementEffect.cpp
@@ -24,15 +24,18 @@ bool FApplyVelocityPhysicsEffect::ApplyMovementEffect(FApplyMovementEffectParams
 	// If not there, get from the particle corresponding to the updated component
 	FVector Position;
 	FRotator Orientation;
+	FVector AngularVelocityDegrees = FVector::ZeroVector;
 	if (const FMoverDefaultSyncState* CurrentSyncState = OutputState.SyncStateCollection.FindDataByType<FMoverDefaultSyncState>())
 	{
 		Position = CurrentSyncState->GetLocation_WorldSpace();
 		Orientation = CurrentSyncState->GetOrientation_WorldSpace();
+		AngularVelocityDegrees = CurrentSyncState->GetAngularVelocityDegrees_WorldSpace();
 	}
 	else if (Chaos::FPBDRigidParticleHandle* ParticleHandle = UPhysicsMovementUtils::GetRigidParticleHandleFromComponent(ApplyEffectParams.UpdatedPrimitive))
 	{
 		Position = ParticleHandle->GetX();
 		Orientation = FRotator(ParticleHandle->GetR());
+		AngularVelocityDegrees = FMath::RadiansToDegrees(ParticleHandle->GetW());
 	}
 	else
 	{
@@ -62,6 +65,7 @@ bool FApplyVelocityPhysicsEffect::ApplyMovementEffect(FApplyMovementEffectParams
 		Position,
 		Orientation,
 		Velocity,
+		AngularVelocityDegrees,
 		MovementBaseInfo.MovementBase.Get(),
 		MovementBaseInfo.BoneName);
 
diff --git a/Engine/Plugins/Experimental/Mover/Source/Mover/Private/PhysicsMover/Modes/PhysicsDrivenFallingMode.cpp b/Engine/Plugins/Experimental/Mover/Source/Mover/Private/PhysicsMover/Modes/PhysicsDrivenFallingMode.cpp
--- a/Engine/Plugins/Experimental/Mover/Source/Mover/Private/PhysicsMover/Modes/PhysicsDrivenFallingMode.cpp
+++ b/Engine/Plugins/Experimental/Mover/Source/Mover/Private/PhysicsMover/Modes/PhysicsDrivenFallingMode.cpp
@@ -135,7 +135,7 @@ void UPhysicsDrivenFallingMode::SimulationTick_Implementation(const FSimulationT
 	}
 
 	// In air steering
-	const FRotator TargetOrient = UMovementUtils::ApplyAngularVelocity(StartingSyncState->GetOrientation_WorldSpace(), ProposedMove.AngularVelocity, DeltaSeconds);
+	const FRotator TargetOrient = UMovementUtils::ApplyAngularVelocityToRotator(StartingSyncState->GetOrientation_WorldSpace(), ProposedMove.AngularVelocityDegrees, DeltaSeconds);
 
 	FVector TargetVel = ProposedMove.LinearVelocity;
 	if (const APhysicsVolume* CurPhysVolume = UpdatedComponent->GetPhysicsVolume())
@@ -168,7 +168,8 @@ void UPhysicsDrivenFallingMode::SimulationTick_Implementation(const FSimulationT
 	OutputSyncState.SetTransforms_WorldSpace(
 		TargetPos,
 		TargetOrient,
-		TargetVel);
+		TargetVel,
+		ProposedMove.AngularVelocityDegrees);
 }
 
 #undef LOCTEXT_NAMESPACE
\ No newline at end of file
diff --git a/Engine/Plugins/Experimental/Mover/Source/Mover/Private/PhysicsMover/Modes/PhysicsDrivenFlyingMode.cpp b/Engine/Plugins/Experimental/Mover/Source/Mover/Private/PhysicsMover/Modes/PhysicsDrivenFlyingMode.cpp
--- a/Engine/Plugins/Experimental/Mover/Source/Mover/Private/PhysicsMover/Modes/PhysicsDrivenFlyingMode.cpp
+++ b/Engine/Plugins/Experimental/Mover/Source/Mover/Private/PhysicsMover/Modes/PhysicsDrivenFlyingMode.cpp
@@ -66,7 +66,7 @@ void UPhysicsDrivenFlyingMode::SimulationTick_Implementation(const FSimulationTi
 	}
 
 	// In air steering
-	const FRotator TargetOrient = UMovementUtils::ApplyAngularVelocity(StartingSyncState->GetOrientation_WorldSpace(), ProposedMove.AngularVelocity, DeltaSeconds);
+	const FRotator TargetOrient = UMovementUtils::ApplyAngularVelocityToRotator(StartingSyncState->GetOrientation_WorldSpace(), ProposedMove.AngularVelocityDegrees, DeltaSeconds);
 
 	FVector TargetVel = ProposedMove.LinearVelocity;
 	if (const APhysicsVolume* CurPhysVolume = UpdatedComponent->GetPhysicsVolume())
@@ -84,7 +84,8 @@ void UPhysicsDrivenFlyingMode::SimulationTick_Implementation(const FSimulationTi
 	OutputSyncState.SetTransforms_WorldSpace(
 		TargetPos,
 		TargetOrient,
-		TargetVel);
+		TargetVel,
+		ProposedMove.AngularVelocityDegrees);
 }
 
 #undef LOCTEXT_NAMESPACE
\ No newline at end of file
diff --git a/Engine/Plugins/Experimental/Mover/Source/Mover/Private/PhysicsMover/Modes/PhysicsDrivenSwimmingMode.cpp b/Engine/Plugins/Experimental/Mover/Source/Mover/Private/PhysicsMover/Modes/PhysicsDrivenSwimmingMode.cpp
--- a/Engine/Plugins/Experimental/Mover/Source/Mover/Private/PhysicsMover/Modes/PhysicsDrivenSwimmingMode.cpp
+++ b/Engine/Plugins/Experimental/Mover/Source/Mover/Private/PhysicsMover/Modes/PhysicsDrivenSwimmingMode.cpp
@@ -153,7 +153,7 @@ void UPhysicsDrivenSwimmingMode::SimulationTick_Implementation(const FSimulation
 		const bool bWalkTrigger = WaterResult.WaterSplineData.ImmersionDepth < CommonLegacySettings->SwimmingStopImmersionDepth;
 		const bool bFallTrigger = FMath::Clamp((WaterResult.WaterSplineData.ImmersionDepth + TargetHeight) / (2 * TargetHeight), -2.f, 2.f) < -1.f;
 	
-		const FRotator TargetOrient = UMovementUtils::ApplyAngularVelocity(StartingSyncState->GetOrientation_WorldSpace(), ProposedMove.AngularVelocity, DeltaSeconds);
+		const FRotator TargetOrient = UMovementUtils::ApplyAngularVelocityToRotator(StartingSyncState->GetOrientation_WorldSpace(), ProposedMove.AngularVelocityDegrees, DeltaSeconds);
 
 		FVector TargetVel = ProposedMove.LinearVelocity;		
 		if (const APhysicsVolume* CurPhysVolume = UpdatedComponent->GetPhysicsVolume())
@@ -168,7 +168,8 @@ void UPhysicsDrivenSwimmingMode::SimulationTick_Implementation(const FSimulation
 		OutputSyncState.SetTransforms_WorldSpace(
 			TargetPos,
 			TargetOrient,
-			TargetVel);
+			TargetVel,
+			ProposedMove.AngularVelocityDegrees);
 	
 		if (bWalkTrigger && bIsWithinReach)
 		{
diff --git a/Engine/Plugins/Experimental/Mover/Source/Mover/Private/PhysicsMover/Modes/PhysicsDrivenWalkingMode.cpp b/Engine/Plugins/Experimental/Mover/Source/Mover/Private/PhysicsMover/Modes/PhysicsDrivenWalkingMode.cpp
--- a/Engine/Plugins/Experimental/Mover/Source/Mover/Private/PhysicsMover/Modes/PhysicsDrivenWalkingMode.cpp
+++ b/Engine/Plugins/Experimental/Mover/Source/Mover/Private/PhysicsMover/Modes/PhysicsDrivenWalkingMode.cpp
@@ -447,7 +447,7 @@ void UPhysicsDrivenWalkingMode::SimulationTick_Implementation(const FSimulationT
 
 		// Target orientation
 		// This is always applied regardless of whether the character is supported
-		const FRotator TargetOrientation = UMovementUtils::ApplyAngularVelocity(StartingSyncState->GetOrientation_WorldSpace(), ProposedMove.AngularVelocity, DeltaSeconds);
+		const FRotator TargetOrientation = UMovementUtils::ApplyAngularVelocityToRotator(StartingSyncState->GetOrientation_WorldSpace(), ProposedMove.AngularVelocityDegrees, DeltaSeconds);
 
 		if (bIsSupported)
 		{
@@ -457,6 +457,7 @@ void UPhysicsDrivenWalkingMode::SimulationTick_Implementation(const FSimulationT
 				TargetPosition,
 				TargetOrientation,
 				TargetVelocity,
+				ProposedMove.AngularVelocityDegrees,
 				nullptr);
 		}
 		else
@@ -490,6 +491,7 @@ void UPhysicsDrivenWalkingMode::SwitchToState(const FName& StateName, const FSim
 		StartingSyncState->GetLocation_WorldSpace(),
 		StartingSyncState->GetOrientation_WorldSpace(),
 		StartingSyncState->GetVelocity_WorldSpace(),
+		StartingSyncState->GetAngularVelocityDegrees_WorldSpace(),
 		nullptr);
 }
 
diff --git a/Engine/Plugins/Experimental/Mover/Source/Mover/Public/DefaultMovementSet/Modes/AsyncFallingMode.h b/Engine/Plugins/Experimental/Mover/Source/Mover/Public/DefaultMovementSet/Modes/AsyncFallingMode.h
--- a/Engine/Plugins/Experimental/Mover/Source/Mover/Public/DefaultMovementSet/Modes/AsyncFallingMode.h
+++ b/Engine/Plugins/Experimental/Mover/Source/Mover/Public/DefaultMovementSet/Modes/AsyncFallingMode.h
@@ -91,7 +91,7 @@ protected:
 	UFUNCTION(BlueprintCallable, Category=Mover)
 	UE_API virtual void ProcessLanded(const FFloorCheckResult& FloorResult, FVector& Velocity, FRelativeBaseInfo& BaseInfo, FMoverTickEndData& TickEndData) const;
 
-	UE_API void CaptureFinalState(const FVector FinalLocation, const FRotator FinalRotation, const FFloorCheckResult& FloorResult, float DeltaSeconds, float DeltaSecondsUsed, FMoverDefaultSyncState& OutputSyncState, FMoverTickEndData& TickEndData, FMovementRecord& Record) const;
+	UE_API void CaptureFinalState(const FVector FinalLocation, const FRotator FinalRotation, const FFloorCheckResult& FloorResult, float DeltaSeconds, float DeltaSecondsUsed, const FVector& AngularVelocityDegrees, FMoverDefaultSyncState& OutputSyncState, FMoverTickEndData& TickEndData, FMovementRecord& Record) const;
 
 	TWeakObjectPtr<const UCommonLegacyMovementSettings> CommonLegacySettings;
 };
diff --git a/Engine/Plugins/Experimental/Mover/Source/Mover/Public/DefaultMovementSet/Modes/AsyncNavWalkingMode.h b/Engine/Plugins/Experimental/Mover/Source/Mover/Public/DefaultMovementSet/Modes/AsyncNavWalkingMode.h
--- a/Engine/Plugins/Experimental/Mover/Source/Mover/Public/DefaultMovementSet/Modes/AsyncNavWalkingMode.h
+++ b/Engine/Plugins/Experimental/Mover/Source/Mover/Public/DefaultMovementSet/Modes/AsyncNavWalkingMode.h
@@ -129,7 +129,7 @@ protected:
 	UE_API virtual void OnRegistered(const FName ModeName) override;
 	UE_API virtual void OnUnregistered() override;
 
-	UE_API void CaptureOutputState(const FMoverDefaultSyncState& StartSyncState, const FVector& FinalLocation, const FRotator& FinalRotation, const FMovementRecord& Record, FMoverDefaultSyncState& OutputSyncState, FMoverTickEndData& OutputState) const;
+	UE_API void CaptureOutputState(const FMoverDefaultSyncState& StartSyncState, const FVector& FinalLocation, const FRotator& FinalRotation, const FMovementRecord& Record, const FVector& AngularVelocityDegrees, FMoverDefaultSyncState& OutputSyncState, FMoverTickEndData& OutputState) const;
 
 	TWeakObjectPtr<const UCommonLegacyMovementSettings> CommonLegacySettings;
  
diff --git a/Engine/Plugins/Experimental/Mover/Source/Mover/Public/DefaultMovementSet/Modes/AsyncWalkingMode.h b/Engine/Plugins/Experimental/Mover/Source/Mover/Public/DefaultMovementSet/Modes/AsyncWalkingMode.h
--- a/Engine/Plugins/Experimental/Mover/Source/Mover/Public/DefaultMovementSet/Modes/AsyncWalkingMode.h
+++ b/Engine/Plugins/Experimental/Mover/Source/Mover/Public/DefaultMovementSet/Modes/AsyncWalkingMode.h
@@ -52,7 +52,7 @@ protected:
 	UE_API virtual void OnRegistered(const FName ModeName) override; 
 	UE_API virtual void OnUnregistered() override;
 
-	UE_API void CaptureFinalState(const FVector FinalLocation, const FRotator FinalRotation, bool bDidAttemptMovement, const FFloorCheckResult& FloorResult, const FMovementRecord& Record, FMoverDefaultSyncState& OutputSyncState) const;
+	UE_API void CaptureFinalState(const FVector FinalLocation, const FRotator FinalRotation, bool bDidAttemptMovement, const FFloorCheckResult& FloorResult, const FMovementRecord& Record, const FVector& AngularVelocityDegrees, FMoverDefaultSyncState& OutputSyncState) const;
 
 	UE_API FRelativeBaseInfo UpdateFloorAndBaseInfo(const FFloorCheckResult& FloorResult) const;
 
diff --git a/Engine/Plugins/Experimental/Mover/Source/Mover/Public/DefaultMovementSet/Modes/FallingMode.h b/Engine/Plugins/Experimental/Mover/Source/Mover/Public/DefaultMovementSet/Modes/FallingMode.h
--- a/Engine/Plugins/Experimental/Mover/Source/Mover/Public/DefaultMovementSet/Modes/FallingMode.h
+++ b/Engine/Plugins/Experimental/Mover/Source/Mover/Public/DefaultMovementSet/Modes/FallingMode.h
@@ -102,7 +102,7 @@ protected:
 	UE_API virtual void OnRegistered(const FName ModeName) override;
 	UE_API virtual void OnUnregistered() override;
 
-	UE_API void CaptureFinalState(USceneComponent* UpdatedComponent, const FMoverDefaultSyncState& StartSyncState, const FFloorCheckResult& FloorResult, float DeltaSeconds, float DeltaSecondsUsed, FMoverDefaultSyncState& OutputSyncState, FMoverTickEndData& TickEndData, FMovementRecord& Record) const;
+	UE_API void CaptureFinalState(USceneComponent* UpdatedComponent, const FMoverDefaultSyncState& StartSyncState, const FFloorCheckResult& FloorResult, float DeltaSeconds, float DeltaSecondsUsed, const FVector& AngularVelocityDegrees, FMoverDefaultSyncState& OutputSyncState, FMoverTickEndData& TickEndData, FMovementRecord& Record) const;
 
 	TObjectPtr<const UCommonLegacyMovementSettings> CommonLegacySettings;
 };
diff --git a/Engine/Plugins/Experimental/Mover/Source/Mover/Public/DefaultMovementSet/Modes/FlyingMode.h b/Engine/Plugins/Experimental/Mover/Source/Mover/Public/DefaultMovementSet/Modes/FlyingMode.h
--- a/Engine/Plugins/Experimental/Mover/Source/Mover/Public/DefaultMovementSet/Modes/FlyingMode.h
+++ b/Engine/Plugins/Experimental/Mover/Source/Mover/Public/DefaultMovementSet/Modes/FlyingMode.h
@@ -37,7 +37,7 @@ protected:
 	UE_API virtual void OnRegistered(const FName ModeName) override;
 	UE_API virtual void OnUnregistered() override;
 
-	UE_API void CaptureFinalState(USceneComponent* UpdatedComponent, FMovementRecord& Record, const FMoverDefaultSyncState& StartSyncState, FMoverDefaultSyncState& OutputSyncState, const float DeltaSeconds) const;
+	UE_API void CaptureFinalState(USceneComponent* UpdatedComponent, FMovementRecord& Record, const FMoverDefaultSyncState& StartSyncState, const FVector& AngularVelocityDegrees, FMoverDefaultSyncState& OutputSyncState, const float DeltaSeconds) const;
 
 	TObjectPtr<const UCommonLegacyMovementSettings> CommonLegacySettings;
 };
diff --git a/Engine/Plugins/Experimental/Mover/Source/Mover/Public/DefaultMovementSet/Modes/NavWalkingMode.h b/Engine/Plugins/Experimental/Mover/Source/Mover/Public/DefaultMovementSet/Modes/NavWalkingMode.h
--- a/Engine/Plugins/Experimental/Mover/Source/Mover/Public/DefaultMovementSet/Modes/NavWalkingMode.h
+++ b/Engine/Plugins/Experimental/Mover/Source/Mover/Public/DefaultMovementSet/Modes/NavWalkingMode.h
@@ -153,7 +153,7 @@ protected:
 	UE_API virtual void OnRegistered(const FName ModeName) override;
 	UE_API virtual void OnUnregistered() override;
 
-	UE_API void CaptureFinalState(USceneComponent* UpdatedComponent, const FMovementRecord& Record, FMoverDefaultSyncState& OutputSyncState) const;
+	UE_API void CaptureFinalState(USceneComponent* UpdatedComponent, const FMovementRecord& Record, const FVector& AngularVelocityDegrees, FMoverDefaultSyncState& OutputSyncState) const;
 
 	TObjectPtr<const UCommonLegacyMovementSettings> CommonLegacySettings;
 
diff --git a/Engine/Plugins/Experimental/Mover/Source/Mover/Public/DefaultMovementSet/Modes/SimpleSpringWalkingMode.h b/Engine/Plugins/Experimental/Mover/Source/Mover/Public/DefaultMovementSet/Modes/SimpleSpringWalkingMode.h
--- a/Engine/Plugins/Experimental/Mover/Source/Mover/Public/DefaultMovementSet/Modes/SimpleSpringWalkingMode.h
+++ b/Engine/Plugins/Experimental/Mover/Source/Mover/Public/DefaultMovementSet/Modes/SimpleSpringWalkingMode.h
@@ -20,7 +20,7 @@ class USimpleSpringWalkingMode : public USimpleWalkingMode
 public:
 	UE_API virtual void SimulationTick_Implementation(const FSimulationTickParams& Params, FMoverTickEndData& OutputState) override;
 	UE_API virtual void GenerateWalkMove_Implementation(FMoverTickStartData& StartState, float DeltaSeconds, const FVector& DesiredVelocity,
-									 const FQuat& DesiredFacing, FVector& OutVelocity, FQuat& OutFacing) override;
+									 const FQuat& DesiredFacing, const FQuat& CurrentFacing, FVector& InOutAngularVelocityDegrees, FVector& InOutVelocity) override;
 
 protected:
 	UPROPERTY(EditAnywhere, BlueprintReadWrite, Category = "Mover|Spring Settings", meta = (ClampMin = "0", UIMin = "0", ForceUnits = "s"))
diff --git a/Engine/Plugins/Experimental/Mover/Source/Mover/Public/DefaultMovementSet/Modes/SimpleWalkingMode.h b/Engine/Plugins/Experimental/Mover/Source/Mover/Public/DefaultMovementSet/Modes/SimpleWalkingMode.h
--- a/Engine/Plugins/Experimental/Mover/Source/Mover/Public/DefaultMovementSet/Modes/SimpleWalkingMode.h
+++ b/Engine/Plugins/Experimental/Mover/Source/Mover/Public/DefaultMovementSet/Modes/SimpleWalkingMode.h
@@ -23,7 +23,7 @@ public:
 	// Override this to make a simple walking mode
 	UFUNCTION(BlueprintNativeEvent, meta = (DisplayName = "Generate Simple Walk Move"))
 	UE_API void GenerateWalkMove(UPARAM(ref) FMoverTickStartData& StartState, float DeltaSeconds, const FVector& DesiredVelocity,
-	                                     const FQuat& DesiredFacing,UPARAM(ref)  FVector& InOutVelocity, UPARAM(ref) FQuat& InOutFacing);
+	                                     const FQuat& DesiredFacing, const FQuat& CurrentFacing, UPARAM(ref) FVector& InAngularVelocityDegrees, UPARAM(ref) FVector& InOutVelocity);
 
 	// If this value is greater or equal to 0, this will override the max speed read from the common legacy shared walk settings
 	UPROPERTY(EditAnywhere, BlueprintReadWrite, Category = "Mover|Walking Settings", meta = (ForceUnits = "cm/s"))
diff --git a/Engine/Plugins/Experimental/Mover/Source/Mover/Public/DefaultMovementSet/Modes/SmoothWalkingMode.h b/Engine/Plugins/Experimental/Mover/Source/Mover/Public/DefaultMovementSet/Modes/SmoothWalkingMode.h
--- a/Engine/Plugins/Experimental/Mover/Source/Mover/Public/DefaultMovementSet/Modes/SmoothWalkingMode.h
+++ b/Engine/Plugins/Experimental/Mover/Source/Mover/Public/DefaultMovementSet/Modes/SmoothWalkingMode.h
@@ -20,7 +20,7 @@ class USmoothWalkingMode : public USimpleWalkingMode
 public:
 	UE_API virtual void SimulationTick_Implementation(const FSimulationTickParams& Params, FMoverTickEndData& OutputState) override;
 	UE_API virtual void GenerateWalkMove_Implementation(FMoverTickStartData& StartState, float DeltaSeconds, const FVector& DesiredVelocity,
-									 const FQuat& DesiredFacing, FVector& InOutVelocity, FQuat& InOutFacing) override;
+									 const FQuat& DesiredFacing, const FQuat& CurrentFacing, FVector& InOutAngularVelocityDegrees, FVector& InOutVelocity) override;
 
 
 protected: // Velocity Controls
diff --git a/Engine/Plugins/Experimental/Mover/Source/Mover/Public/DefaultMovementSet/Modes/WalkingMode.h b/Engine/Plugins/Experimental/Mover/Source/Mover/Public/DefaultMovementSet/Modes/WalkingMode.h
--- a/Engine/Plugins/Experimental/Mover/Source/Mover/Public/DefaultMovementSet/Modes/WalkingMode.h
+++ b/Engine/Plugins/Experimental/Mover/Source/Mover/Public/DefaultMovementSet/Modes/WalkingMode.h
@@ -62,7 +62,7 @@ protected:
 	UE_API virtual void OnRegistered(const FName ModeName) override; 
 	UE_API virtual void OnUnregistered() override;
 
-	UE_API void CaptureFinalState(USceneComponent* UpdatedComponent, bool bDidAttemptMovement, const FFloorCheckResult& FloorResult, const FMovementRecord& Record, FMoverDefaultSyncState& OutputSyncState) const;
+	UE_API void CaptureFinalState(USceneComponent* UpdatedComponent, bool bDidAttemptMovement, const FFloorCheckResult& FloorResult, const FMovementRecord& Record, const FVector& AngularVelocityDegrees, FMoverDefaultSyncState& OutputSyncState) const;
 
 	UE_API FRelativeBaseInfo UpdateFloorAndBaseInfo(const FFloorCheckResult& FloorResult) const;
 
diff --git a/Engine/Plugins/Experimental/Mover/Source/Mover/Public/MoveLibrary/ModularMovement.h b/Engine/Plugins/Experimental/Mover/Source/Mover/Public/MoveLibrary/ModularMovement.h
--- a/Engine/Plugins/Experimental/Mover/Source/Mover/Public/MoveLibrary/ModularMovement.h
+++ b/Engine/Plugins/Experimental/Mover/Source/Mover/Public/MoveLibrary/ModularMovement.h
@@ -30,9 +30,9 @@ class ITurnGeneratorInterface : public IInterface
 public:
 	// TODO: consider making GetTurn take a params struct instead of a fixed arg list
 
-	/** Returns an additive angular velocity (per-axis degrees/second) based on the starting state and timestep */
+	/** Returns an additive angular velocity (degrees/second) based on the starting state and timestep. The vector points in the direction of the rotation axis */
 	UFUNCTION(BlueprintNativeEvent, BlueprintCallable, Category=Mover)
-	MOVER_API FRotator GetTurn(FRotator TargetOrientation, const FMoverTickStartData& FullStartState, const FMoverDefaultSyncState& MoverState, const FMoverTimeStep& TimeStep, const FProposedMove& ProposedMove, UMoverBlackboard* SimBlackboard);
+	MOVER_API FVector GetTurn(FRotator TargetOrientation, const FMoverTickStartData& FullStartState, const FMoverDefaultSyncState& MoverState, const FMoverTimeStep& TimeStep, const FProposedMove& ProposedMove, UMoverBlackboard* SimBlackboard);
 };
 
 
@@ -59,7 +59,7 @@ public:
 	float RollRate = -1.f;
 
 
-	UE_API virtual FRotator GetTurn_Implementation(FRotator TargetOrientation, const FMoverTickStartData& FullStartState, const FMoverDefaultSyncState& MoverState, const FMoverTimeStep& TimeStep, const FProposedMove& ProposedMove, UMoverBlackboard* SimBlackboard) override;
+	UE_API virtual FVector GetTurn_Implementation(FRotator TargetOrientation, const FMoverTickStartData& FullStartState, const FMoverDefaultSyncState& MoverState, const FMoverTimeStep& TimeStep, const FProposedMove& ProposedMove, UMoverBlackboard* SimBlackboard) override;
 };
 
 
@@ -77,7 +77,7 @@ public:
 	UPROPERTY(BlueprintReadWrite, EditAnywhere, Category = "Exact Damped Turning", meta = (ClampMin = "0.01", UIMin = "0", ForceUnits = "seconds"))
 	float HalfLifeSeconds = 0.15f;
 
-	UE_API virtual FRotator GetTurn_Implementation(FRotator TargetOrientation, const FMoverTickStartData& FullStartState, const FMoverDefaultSyncState& MoverState, const FMoverTimeStep& TimeStep, const FProposedMove& ProposedMove, UMoverBlackboard* SimBlackboard) override;
+	UE_API virtual FVector GetTurn_Implementation(FRotator TargetOrientation, const FMoverTickStartData& FullStartState, const FMoverDefaultSyncState& MoverState, const FMoverTimeStep& TimeStep, const FProposedMove& ProposedMove, UMoverBlackboard* SimBlackboard) override;
 };
 
 
diff --git a/Engine/Plugins/Experimental/Mover/Source/Mover/Public/MoveLibrary/MovementUtils.h b/Engine/Plugins/Experimental/Mover/Source/Mover/Public/MoveLibrary/MovementUtils.h
--- a/Engine/Plugins/Experimental/Mover/Source/Mover/Public/MoveLibrary/MovementUtils.h
+++ b/Engine/Plugins/Experimental/Mover/Source/Mover/Public/MoveLibrary/MovementUtils.h
@@ -280,19 +280,34 @@ public:
 	static UE_API FVector ComputeVelocityFromPositions(const FVector& FromPos, const FVector& ToPos, float DeltaSeconds);
 
 	/** Computes the angular velocity needed to change from one orientation to another within a time frame. Use the optional TurningRateLimit to clamp to a maximum step (negative=unlimited). */
+	UE_DEPRECATED(5.7, "Rotator is no longer used for angular velocity")
 	UFUNCTION(BlueprintCallable, Category = Mover)
 	static UE_API FRotator ComputeAngularVelocity(const FRotator& From, const FRotator& To, const FQuat& WorldToGravity, float DeltaSeconds, float TurningRateLimit=-1.f);
+	
+	/** Computes the angular velocity needed to change from one orientation to another within a time frame. Use the optional TurningRateLimit to clamp to a maximum step in degrees per second (negative=unlimited). */
+	UFUNCTION(BlueprintCallable, Category = Mover)
+	static UE_API FVector ComputeAngularVelocityDegrees(const FRotator& From, const FRotator& To, float DeltaSeconds, float TurningRateLimit = -1.0f);
 
 	/** Computes the directional movement intent based on input vector and associated type */
 	UFUNCTION(BlueprintCallable, Category = Mover)
 	static UE_API FVector ComputeDirectionIntent(const FVector& MoveInput, EMoveInputType MoveInputType, float MaxSpeed);
-
+	
 	/** Returns whether this rotator representing angular velocity has any non-zero values. This function exists due to FRotator::IsZero queries performing undesired wrapping and clamping. */
+	UE_DEPRECATED(5.7, "Rotator is no longer used for angular velocity")
 	UFUNCTION(BlueprintCallable, Category = Mover)
 	static UE_API bool IsAngularVelocityZero(const FRotator& AngularVelocity);
 
+	/** Applies a desired angular velocity to a starting rotator over time. Returns the new orientation rotator. */
+	UFUNCTION(BlueprintPure, Category = Mover)
+	static UE_API FRotator ApplyAngularVelocityToRotator(const FRotator& StartingOrient, const FVector& AngularVelocityDegrees, float DeltaSeconds);
+
+	/** Applies a desired angular velocity to a starting quaternion over time. Returns the new orientation quaternion. */
+	UFUNCTION(BlueprintPure, Category = Mover)
+	static UE_API FQuat ApplyAngularVelocityToQuat(const FQuat& StartingOrient, const FVector& AngularVelocityDegrees, float DeltaSeconds);
+
 	/** Applies a desired angular velocity to a starting rotator over time. Returns the new orientation rotator. */
 	UFUNCTION(BlueprintCallable, Category = Mover)
+	UE_DEPRECATED(5.7, "Rotator is no longer used for angular velocity")
 	static UE_API FRotator ApplyAngularVelocity(const FRotator& StartingOrient, const FRotator& AngularVelocity, float DeltaSeconds);
 
 	/**
diff --git a/Engine/Plugins/Experimental/Mover/Source/Mover/Public/MoveLibrary/MovementUtilsTypes.h b/Engine/Plugins/Experimental/Mover/Source/Mover/Public/MoveLibrary/MovementUtilsTypes.h
--- a/Engine/Plugins/Experimental/Mover/Source/Mover/Public/MoveLibrary/MovementUtilsTypes.h
+++ b/Engine/Plugins/Experimental/Mover/Source/Mover/Public/MoveLibrary/MovementUtilsTypes.h
@@ -53,10 +53,10 @@ struct FProposedMove
 	// Units per second, world space, possibly mapped onto walking surface
 	UPROPERTY(EditAnywhere, BlueprintReadWrite, Category = Mover)
 	FVector LinearVelocity = FVector::ZeroVector;
-	
-	// Degrees per second, local space
+
+	// Angular velocity in degrees per second. Direction points along rotation axis.
 	UPROPERTY(EditAnywhere, BlueprintReadWrite, Category = Mover)
-	FRotator AngularVelocity = FRotator::ZeroRotator;
+	FVector AngularVelocityDegrees = FVector::ZeroVector;
 	
 	// Signals whether there was any directional intent specified
 	UPROPERTY(EditAnywhere, BlueprintReadWrite, Category = Mover)
diff --git a/Engine/Plugins/Experimental/Mover/Source/Mover/Public/MoverDataModelTypes.h b/Engine/Plugins/Experimental/Mover/Source/Mover/Public/MoverDataModelTypes.h
--- a/Engine/Plugins/Experimental/Mover/Source/Mover/Public/MoverDataModelTypes.h
+++ b/Engine/Plugins/Experimental/Mover/Source/Mover/Public/MoverDataModelTypes.h
@@ -158,6 +158,10 @@ protected:
 	UPROPERTY(BlueprintReadWrite, EditAnywhere, Category = Mover)
 	FVector Velocity;
 
+	// Angular velocity, degrees per second, relative to MovementBase if set, world space otherwise.
+	UPROPERTY(BlueprintReadWrite, EditAnywhere, Category = Mover)
+	FVector AngularVelocityDegrees;
+
 public:
 	// Movement intent direction relative to MovementBase if set, world space otherwise. Magnitude of range (0-1)
 	UPROPERTY(BlueprintReadWrite, EditAnywhere, Category = Mover)
@@ -184,6 +188,7 @@ public:
 		: Location(ForceInitToZero)
 		, Orientation(ForceInitToZero)
 		, Velocity(ForceInitToZero)
+		, AngularVelocityDegrees(ForceInitToZero)
 		, MoveDirectionIntent(ForceInitToZero)
 		, MovementBase(nullptr)
 		, MovementBaseBoneName(NAME_None)
@@ -207,7 +212,10 @@ public:
 
 	UE_API virtual void Interpolate(const FMoverDataStructBase& From, const FMoverDataStructBase& To, float Pct) override;
 
+	UE_DEPRECATED(5.7, "Use the SetTransforms_WorldSpace with an angular velocity")
 	UE_API void SetTransforms_WorldSpace(FVector WorldLocation, FRotator WorldOrient, FVector WorldVelocity, UPrimitiveComponent* Base=nullptr, FName BaseBone = NAME_None);
+	
+	UE_API void SetTransforms_WorldSpace(FVector WorldLocation, FRotator WorldOrient, FVector WorldVelocity, FVector WorldAngularVelocityDegrees, UPrimitiveComponent* Base=nullptr, FName BaseBone = NAME_None);
 
 	// Returns whether the base setting succeeded
 	UE_API bool SetMovementBase(UPrimitiveComponent* Base, FName BaseBone=NAME_None);
@@ -237,6 +245,9 @@ public:
 
 	UE_API FTransform GetTransform_WorldSpace() const;
 	UE_API FTransform GetTransform_BaseSpace() const;
+
+	UE_API FVector GetAngularVelocityDegrees_WorldSpace() const;
+	UE_API FVector GetAngularVelocityDegrees_BaseSpace() const;
 };
 
 template<>
@@ -283,6 +294,10 @@ public:	// FMoverDefaultSyncState
 	UFUNCTION(BlueprintCallable, Category = Mover)
 	static UE_API FVector GetVelocityFromSyncState(const FMoverDefaultSyncState& SyncState);
 
+	/** Returns the angular velocity in world space, in degrees per second */
+	UFUNCTION(BlueprintCallable, Category = Mover)
+	static UE_API FVector GetAngularVelocityDegreesFromSyncState(const FMoverDefaultSyncState& SyncState);
+
 	/** Returns the orientation in world space */
 	UFUNCTION(BlueprintCallable, Category = Mover)
 	static UE_API FRotator GetOrientationFromSyncState(const FMoverDefaultSyncState& SyncState);
diff --git a/Engine/Plugins/Experimental/MoverExamples/MoverExamples.uplugin b/Engine/Plugins/Experimental/MoverExamples/MoverExamples.uplugin
--- a/Engine/Plugins/Experimental/MoverExamples/MoverExamples.uplugin
+++ b/Engine/Plugins/Experimental/MoverExamples/MoverExamples.uplugin
@@ -34,6 +34,10 @@
 		{
 			"Name": "EnhancedInput",
 			"Enabled": true
+		},
+		{
+			"Name": "CableComponent",
+			"Enabled": true
 		}
 	]
 }
\ No newline at end of file
diff --git a/Engine/Plugins/Experimental/MoverExamples/Source/MoverExamples/MoverExamples.Build.cs b/Engine/Plugins/Experimental/MoverExamples/Source/MoverExamples/MoverExamples.Build.cs
--- a/Engine/Plugins/Experimental/MoverExamples/Source/MoverExamples/MoverExamples.Build.cs
+++ b/Engine/Plugins/Experimental/MoverExamples/Source/MoverExamples/MoverExamples.Build.cs
@@ -14,7 +14,8 @@ public class MoverExamples : ModuleRules
 				"Engine", 
 				"InputCore", 
 				"Mover",
-				"EnhancedInput"
+				"EnhancedInput",
+				"CableComponent"
 			});
 	}
 }
diff --git a/Engine/Plugins/Experimental/MoverExamples/Source/MoverExamples/Private/CharacterVariants/Ziplining/ZipliningMode.cpp b/Engine/Plugins/Experimental/MoverExamples/Source/MoverExamples/Private/CharacterVariants/Ziplining/ZipliningMode.cpp
--- a/Engine/Plugins/Experimental/MoverExamples/Source/MoverExamples/Private/CharacterVariants/Ziplining/ZipliningMode.cpp
+++ b/Engine/Plugins/Experimental/MoverExamples/Source/MoverExamples/Private/CharacterVariants/Ziplining/ZipliningMode.cpp
@@ -159,7 +159,7 @@ void UZipliningMode::SimulationTick_Implementation(const FSimulationTickParams&
 				DefaultAirMode = LegacySettings->AirMovementModeName;
 			}
 
-			OutputState.MovementEndState.NextModeName = DefaultModeNames::Falling;
+			OutputState.MovementEndState.NextModeName = DefaultAirMode;
 			OutputState.MovementEndState.RemainingMs = Params.TimeStep.StepMs;
 			return;
 		}
@@ -222,6 +222,7 @@ void UZipliningMode::SimulationTick_Implementation(const FSimulationTickParams&
 	OutputSyncState.SetTransforms_WorldSpace(FinalLocation,
 		UpdatedComponent->GetComponentRotation(),
 		FinalVelocity,
+		FVector::ZeroVector,
 		nullptr); // no movement base
 
 	UpdatedComponent->ComponentVelocity = FinalVelocity;
diff --git a/Engine/Plugins/Experimental/MoverExamples/Source/MoverExamples/Private/MovementBases/FollowPathMode.cpp b/Engine/Plugins/Experimental/MoverExamples/Source/MoverExamples/Private/MovementBases/FollowPathMode.cpp
--- a/Engine/Plugins/Experimental/MoverExamples/Source/MoverExamples/Private/MovementBases/FollowPathMode.cpp
+++ b/Engine/Plugins/Experimental/MoverExamples/Source/MoverExamples/Private/MovementBases/FollowPathMode.cpp
@@ -30,6 +30,7 @@ void UFollowPathMode::SimulationTick_Implementation(const FSimulationTickParams&
 	FFollowPathState& OutputPathState		= OutputState.SyncState.SyncStateCollection.FindOrAddMutableDataByType<FFollowPathState>();
 
 	FVector StartingLocation = StartingMoveState->GetLocation_WorldSpace();
+	FRotator StartingRotation = StartingMoveState->GetOrientation_WorldSpace();
 
 	const float DeltaSeconds = Params.TimeStep.StepMs * 0.001f;
 
@@ -94,9 +95,15 @@ void UFollowPathMode::SimulationTick_Implementation(const FSimulationTickParams&
 
 	// Capture final state
 	FVector Velocity = (UpdatedComponent->GetComponentLocation() - StartingLocation) / DeltaSeconds;
+
+	// Compute angular velocity
+	FRotator NewRotation = UpdatedComponent->GetComponentRotation();
+	FVector AngularVelocityDegrees = UMovementUtils::ComputeAngularVelocityDegrees(StartingRotation, NewRotation, DeltaSeconds);
+	
 	OutputMoveState.SetTransforms_WorldSpace( UpdatedComponent->GetComponentLocation(),
-		                                      UpdatedComponent->GetComponentRotation(),
+		                                      NewRotation,
 											  Velocity,
+											  AngularVelocityDegrees,
 											  nullptr ); // no movement base
 
 	UpdatedComponent->ComponentVelocity = Velocity;
diff --git a/Engine/Plugins/Experimental/MoverExamples/Source/MoverExamples/Private/MovementBases/FollowSplineMode.cpp b/Engine/Plugins/Experimental/MoverExamples/Source/MoverExamples/Private/MovementBases/FollowSplineMode.cpp
--- a/Engine/Plugins/Experimental/MoverExamples/Source/MoverExamples/Private/MovementBases/FollowSplineMode.cpp
+++ b/Engine/Plugins/Experimental/MoverExamples/Source/MoverExamples/Private/MovementBases/FollowSplineMode.cpp
@@ -1,6 +1,8 @@
 // Copyright Epic Games, Inc. All Rights Reserved.
 
 #include "MovementBases/FollowSplineMode.h"
+
+#include "MoverComponent.h"
 #include "MoveLibrary/MovementUtils.h"
 #include "Components/SplineComponent.h"
 #include "Curves/CurveFloat.h"
@@ -24,7 +26,7 @@ namespace FollowSplineMode::Utils::Private
 		case ESplineOffsetUnit::DistanceAbsolute:
 		{
 			const float SplineTime = ControlSpline->GetTimeAtDistanceAlongSpline(RangeOffsetInput.Value);
-			return FMath::Clamp(RangeOffsetInput.Value, 0.0f, ControlSpline->Duration);
+			return FMath::Clamp(SplineTime, 0.0f, ControlSpline->Duration);
 		}
 		default:
 		{
@@ -63,6 +65,7 @@ void UFollowSplineMode::SimulationTick_Implementation(const FSimulationTickParam
 	FFollowSplineState& OutputPathState		= OutputState.SyncState.SyncStateCollection.FindOrAddMutableDataByType<FFollowSplineState>();
 
 	FVector StartingLocation = StartingMoveState->GetLocation_WorldSpace();
+	FRotator StartingOrientation = StartingMoveState->GetOrientation_WorldSpace();
 
 	const float DeltaSeconds = Params.TimeStep.StepMs * 0.001f;
 
@@ -123,7 +126,7 @@ void UFollowSplineMode::SimulationTick_Implementation(const FSimulationTickParam
 		// Move the object
 		const FTransform SplineTransform = GetTransformAtTime(MappedSplineTime, MovingComps.UpdatedComponent->GetComponentRotation());
 		const FVector MoveDelta = SplineTransform.GetLocation() - StartingLocation;
-		const FVector Velocity = MoveDelta / DeltaSeconds;
+		const FVector Velocity = DeltaSeconds > UE_SMALL_NUMBER ? MoveDelta / DeltaSeconds : FVector::ZeroVector;
 
 		FMovementRecord& MoveRecord = OutputState.MoveRecord;
 		MoveRecord.SetDeltaSeconds(DeltaSeconds);
@@ -133,11 +136,14 @@ void UFollowSplineMode::SimulationTick_Implementation(const FSimulationTickParam
 		UMovementUtils::TrySafeMoveUpdatedComponent(MovingComps, MoveDelta, SplineTransform.GetRotation(), true, MoveHitResult, ETeleportType::None, MoveRecord);
 
 		UpdatePathState(OutputPathState);
-
+		FRotator NewRotation = MovingComps.UpdatedComponent->GetComponentRotation();
+		FVector AngularVelocityDegrees = UMovementUtils::ComputeAngularVelocityDegrees(StartingOrientation, NewRotation, DeltaSeconds);
+		
 		// Update Move State
 		OutputMoveState.SetTransforms_WorldSpace(MovingComps.UpdatedComponent->GetComponentLocation(),
-			MovingComps.UpdatedComponent->GetComponentRotation(),
+			NewRotation,
 			Velocity,
+			AngularVelocityDegrees,
 			nullptr); // no movement base
 	}
 }
