Mover: adds option for flat-bottom floor checks, exposed in default movement shared settings and through floor util query API

diff --git a/Engine/Plugins/Experimental/Mover/Source/Mover/Private/DefaultMovementSet/Modes/AsyncFallingMode.cpp b/Engine/Plugins/Experimental/Mover/Source/Mover/Private/DefaultMovementSet/Modes/AsyncFallingMode.cpp
--- a/Engine/Plugins/Experimental/Mover/Source/Mover/Private/DefaultMovementSet/Modes/AsyncFallingMode.cpp
+++ b/Engine/Plugins/Experimental/Mover/Source/Mover/Private/DefaultMovementSet/Modes/AsyncFallingMode.cpp
@@ -223,7 +223,7 @@ void UAsyncFallingMode::SimulationTick_Implementation(const FSimulationTickParam
 
 		// Check for hitting a landing surface
 		if (UAirMovementUtils::IsValidLandingSpot(Params.MovingComps, LocationInProgress,
-				SweepHit, CommonLegacySettings->FloorSweepDistance, CommonLegacySettings->MaxWalkSlopeCosine, OUT LandingFloor))
+				SweepHit, CommonLegacySettings->FloorSweepDistance, CommonLegacySettings->MaxWalkSlopeCosine, CommonLegacySettings->bUseFlatBaseForFloorChecks, LandingFloor))
 		{
 			// Adjust height to float slightly above walkable floor
 			LocationInProgress = UGroundMovementUtils::TestMoveToKeepMinHeightAboveFloor(Params.MovingComps, LocationInProgress, RotationInProgress, CommonLegacySettings->MaxWalkSlopeCosine, IN OUT LandingFloor, IN OUT MoveRecord);
@@ -241,7 +241,7 @@ void UAsyncFallingMode::SimulationTick_Implementation(const FSimulationTickParam
 		// We didn't land on a walkable surface, so let's try to slide along it
 		const float PctOfTimeUsedForSliding = UAirMovementUtils::TestFallingMoveAlongHitSurface(
 			Params.MovingComps, TargetLocation-StartLocation, LocationInProgress, TargetRotation, 
-			/*bHandleImpact=*/true, CommonLegacySettings->FloorSweepDistance, CommonLegacySettings->MaxWalkSlopeCosine,
+			/*bHandleImpact=*/true, CommonLegacySettings->FloorSweepDistance, CommonLegacySettings->MaxWalkSlopeCosine, CommonLegacySettings->bUseFlatBaseForFloorChecks,
 			IN OUT SweepHit, OUT LandingFloor, IN OUT MoveRecord);
 
 		if (PctOfTimeUsedForSliding > 0.f)
diff --git a/Engine/Plugins/Experimental/Mover/Source/Mover/Private/DefaultMovementSet/Modes/AsyncFlyingMode.cpp b/Engine/Plugins/Experimental/Mover/Source/Mover/Private/DefaultMovementSet/Modes/AsyncFlyingMode.cpp
--- a/Engine/Plugins/Experimental/Mover/Source/Mover/Private/DefaultMovementSet/Modes/AsyncFlyingMode.cpp
+++ b/Engine/Plugins/Experimental/Mover/Source/Mover/Private/DefaultMovementSet/Modes/AsyncFlyingMode.cpp
@@ -136,7 +136,7 @@ void UAsyncFlyingMode::SimulationTick_Implementation(const FSimulationTickParams
 	{
 		// If we are very close to a walkable surface, make sure we maintain a small gap over it
 		FFloorCheckResult FloorUnderActor;
-		UFloorQueryUtils::FindFloor(Params.MovingComps, CommonLegacySettings->FloorSweepDistance, CommonLegacySettings->MaxWalkSlopeCosine, LocationInProgress, FloorUnderActor);
+		UFloorQueryUtils::FindFloor(Params.MovingComps, CommonLegacySettings->FloorSweepDistance, CommonLegacySettings->MaxWalkSlopeCosine, CommonLegacySettings->bUseFlatBaseForFloorChecks, LocationInProgress, FloorUnderActor);
 
 		if (FloorUnderActor.IsWalkableFloor())
 		{
diff --git a/Engine/Plugins/Experimental/Mover/Source/Mover/Private/DefaultMovementSet/Modes/AsyncWalkingMode.cpp b/Engine/Plugins/Experimental/Mover/Source/Mover/Private/DefaultMovementSet/Modes/AsyncWalkingMode.cpp
--- a/Engine/Plugins/Experimental/Mover/Source/Mover/Private/DefaultMovementSet/Modes/AsyncWalkingMode.cpp
+++ b/Engine/Plugins/Experimental/Mover/Source/Mover/Private/DefaultMovementSet/Modes/AsyncWalkingMode.cpp
@@ -152,7 +152,7 @@ void UAsyncWalkingMode::SimulationTick_Implementation(const FSimulationTickParam
 	// If we don't have cached floor information, we need to search for it again
 	if (!SimBlackboard->TryGet(CommonBlackboard::LastFloorResult, CurrentFloor))
 	{
-		UFloorQueryUtils::FindFloor(Params.MovingComps, CommonLegacySettings->FloorSweepDistance, CommonLegacySettings->MaxWalkSlopeCosine, StartLocation, CurrentFloor);
+		UFloorQueryUtils::FindFloor(Params.MovingComps, CommonLegacySettings->FloorSweepDistance, CommonLegacySettings->MaxWalkSlopeCosine, CommonLegacySettings->bUseFlatBaseForFloorChecks, StartLocation, CurrentFloor);
 	}
 
 	OutputSyncState.MoveDirectionIntent = (ProposedMove.bHasDirIntent ? ProposedMove.DirectionIntent : FVector::ZeroVector);
@@ -236,7 +236,7 @@ void UAsyncWalkingMode::SimulationTick_Implementation(const FSimulationTickParam
 					FOptionalFloorCheckResult StepUpFloorResult;	// passed to sub-operations, so we can use their final floor results if they did a test
 					FVector PostStepUpLocation; // Valid if step-up succeeded
 
-					if (UGroundMovementUtils::TestMoveToStepOver(Params.MovingComps, DownwardDir, CommonLegacySettings->MaxStepHeight, CommonLegacySettings->MaxWalkSlopeCosine, CommonLegacySettings->FloorSweepDistance, OrigMoveDelta * (1.f - PercentTimeAppliedSoFar), RotationInProgress, MoveHitResult, CurrentFloor, false, &StepUpFloorResult, PostStepUpLocation, MoveRecord))
+					if (UGroundMovementUtils::TestMoveToStepOver(Params.MovingComps, DownwardDir, CommonLegacySettings->MaxStepHeight, CommonLegacySettings->MaxWalkSlopeCosine, CommonLegacySettings->bUseFlatBaseForFloorChecks, CommonLegacySettings->FloorSweepDistance, OrigMoveDelta * (1.f - PercentTimeAppliedSoFar), RotationInProgress, MoveHitResult, CurrentFloor, false, &StepUpFloorResult, PostStepUpLocation, MoveRecord))
 					{
 						LocationInProgress = PostStepUpLocation;
 						RotationInProgress = TargetRotation;
@@ -273,7 +273,7 @@ void UAsyncWalkingMode::SimulationTick_Implementation(const FSimulationTickParam
 		}
 
 		// Search for the floor we've ended up on
-		UFloorQueryUtils::FindFloor(Params.MovingComps, CommonLegacySettings->FloorSweepDistance, CommonLegacySettings->MaxWalkSlopeCosine, LocationInProgress, CurrentFloor);
+		UFloorQueryUtils::FindFloor(Params.MovingComps, CommonLegacySettings->FloorSweepDistance, CommonLegacySettings->MaxWalkSlopeCosine, CommonLegacySettings->bUseFlatBaseForFloorChecks, LocationInProgress, CurrentFloor);
 
 		if (CurrentFloor.IsWalkableFloor())
 		{
diff --git a/Engine/Plugins/Experimental/Mover/Source/Mover/Private/DefaultMovementSet/Modes/FallingMode.cpp b/Engine/Plugins/Experimental/Mover/Source/Mover/Private/DefaultMovementSet/Modes/FallingMode.cpp
--- a/Engine/Plugins/Experimental/Mover/Source/Mover/Private/DefaultMovementSet/Modes/FallingMode.cpp
+++ b/Engine/Plugins/Experimental/Mover/Source/Mover/Private/DefaultMovementSet/Modes/FallingMode.cpp
@@ -187,7 +187,7 @@ void UFallingMode::SimulationTick_Implementation(const FSimulationTickParams& Pa
 	{
 		// If we are very close to a walkable floor, make sure we're maintaining the correct distance from it
 		FFloorCheckResult FloorUnderActor;
-		UFloorQueryUtils::FindFloor(Params.MovingComps, CommonLegacySettings->FloorSweepDistance, CommonLegacySettings->MaxWalkSlopeCosine, UpdatedComponent->GetComponentLocation(), OUT FloorUnderActor);
+		UFloorQueryUtils::FindFloor(Params.MovingComps, CommonLegacySettings->FloorSweepDistance, CommonLegacySettings->MaxWalkSlopeCosine, CommonLegacySettings->bUseFlatBaseForFloorChecks, UpdatedComponent->GetComponentLocation(), OUT FloorUnderActor);
 
 		if (FloorUnderActor.IsWalkableFloor())
 		{
@@ -222,7 +222,7 @@ void UFallingMode::SimulationTick_Implementation(const FSimulationTickParams& Pa
 
 		// Check for hitting a landing surface
 		if (UAirMovementUtils::IsValidLandingSpot(Params.MovingComps, UpdatedComponent->GetComponentLocation(),
-			Hit, CommonLegacySettings->FloorSweepDistance, CommonLegacySettings->MaxWalkSlopeCosine, OUT LandingFloor))
+			Hit, CommonLegacySettings->FloorSweepDistance, CommonLegacySettings->MaxWalkSlopeCosine, CommonLegacySettings->bUseFlatBaseForFloorChecks, OUT LandingFloor))
 		{
 			UGroundMovementUtils::TryMoveToKeepMinHeightAboveFloor(MoverComponent, LandingFloor, CommonLegacySettings->MaxWalkSlopeCosine, MoveRecord); // make sure we maintain a small gap over walking surfaces
 			CaptureFinalState(UpdatedComponent, *StartingSyncState, LandingFloor, DeltaSeconds, DeltaSeconds * PctTimeApplied, ProposedMove.AngularVelocityDegrees, OutputSyncState, OutputState, MoveRecord);
@@ -238,7 +238,7 @@ void UFallingMode::SimulationTick_Implementation(const FSimulationTickParams& Pa
 		// We didn't land on a walkable surface, so let's try to slide along it
 		UAirMovementUtils::TryMoveToFallAlongSurface(Params.MovingComps, MoveDelta,
 			(1.f - Hit.Time), TargetOrientQuat, Hit.Normal, Hit, true,
-			CommonLegacySettings->FloorSweepDistance, CommonLegacySettings->MaxWalkSlopeCosine, LandingFloor, MoveRecord);
+			CommonLegacySettings->FloorSweepDistance, CommonLegacySettings->MaxWalkSlopeCosine, CommonLegacySettings->bUseFlatBaseForFloorChecks, LandingFloor, MoveRecord);
 
 		PctTimeApplied += Hit.Time * (1.f - PctTimeApplied);
 
diff --git a/Engine/Plugins/Experimental/Mover/Source/Mover/Private/DefaultMovementSet/Modes/FlyingMode.cpp b/Engine/Plugins/Experimental/Mover/Source/Mover/Private/DefaultMovementSet/Modes/FlyingMode.cpp
--- a/Engine/Plugins/Experimental/Mover/Source/Mover/Private/DefaultMovementSet/Modes/FlyingMode.cpp
+++ b/Engine/Plugins/Experimental/Mover/Source/Mover/Private/DefaultMovementSet/Modes/FlyingMode.cpp
@@ -134,7 +134,7 @@ void UFlyingMode::SimulationTick_Implementation(const FSimulationTickParams& Par
 	{
 		// If we are very close to a walkable surface, make sure we maintain a small gap over it
 		FFloorCheckResult FloorUnderActor;
-		UFloorQueryUtils::FindFloor(Params.MovingComps, CommonLegacySettings->FloorSweepDistance, CommonLegacySettings->MaxWalkSlopeCosine, UpdatedComponent->GetComponentLocation(), OUT FloorUnderActor);
+		UFloorQueryUtils::FindFloor(Params.MovingComps, CommonLegacySettings->FloorSweepDistance, CommonLegacySettings->MaxWalkSlopeCosine, CommonLegacySettings->bUseFlatBaseForFloorChecks, UpdatedComponent->GetComponentLocation(), OUT FloorUnderActor);
 
 		if (FloorUnderActor.IsWalkableFloor())
 		{
diff --git a/Engine/Plugins/Experimental/Mover/Source/Mover/Private/DefaultMovementSet/Modes/WalkingMode.cpp b/Engine/Plugins/Experimental/Mover/Source/Mover/Private/DefaultMovementSet/Modes/WalkingMode.cpp
--- a/Engine/Plugins/Experimental/Mover/Source/Mover/Private/DefaultMovementSet/Modes/WalkingMode.cpp
+++ b/Engine/Plugins/Experimental/Mover/Source/Mover/Private/DefaultMovementSet/Modes/WalkingMode.cpp
@@ -144,7 +144,7 @@ void UWalkingMode::SimulationTick_Implementation(const FSimulationTickParams& Pa
 	// If we don't have cached floor information, we need to search for it again
 	if (!SimBlackboard->TryGet(CommonBlackboard::LastFloorResult, CurrentFloor))
 	{
-		UFloorQueryUtils::FindFloor(MovingComponents, CommonLegacySettings->FloorSweepDistance, CommonLegacySettings->MaxWalkSlopeCosine, UpdatedComponent->GetComponentLocation(), CurrentFloor);
+		UFloorQueryUtils::FindFloor(MovingComponents, CommonLegacySettings->FloorSweepDistance, CommonLegacySettings->MaxWalkSlopeCosine, CommonLegacySettings->bUseFlatBaseForFloorChecks, UpdatedComponent->GetComponentLocation(), CurrentFloor);
 	}
 
 	OutputSyncState.MoveDirectionIntent = (ProposedMove.bHasDirIntent ? ProposedMove.DirectionIntent : FVector::ZeroVector);
@@ -212,7 +212,7 @@ void UWalkingMode::SimulationTick_Implementation(const FSimulationTickParams& Pa
 					const FVector PreStepUpLocation = UpdatedComponent->GetComponentLocation();
 					const FVector DownwardDir = -MoverComp->GetUpDirection();
 
-					if (!UGroundMovementUtils::TryMoveToStepUp(Params.MovingComps, DownwardDir, CommonLegacySettings->MaxStepHeight, CommonLegacySettings->MaxWalkSlopeCosine, CommonLegacySettings->FloorSweepDistance, OrigMoveDelta * (1.f - PercentTimeAppliedSoFar), MoveHitResult, CurrentFloor, false, &StepUpFloorResult, MoveRecord))
+					if (!UGroundMovementUtils::TryMoveToStepUp(Params.MovingComps, DownwardDir, CommonLegacySettings->MaxStepHeight, CommonLegacySettings->MaxWalkSlopeCosine, CommonLegacySettings->bUseFlatBaseForFloorChecks, CommonLegacySettings->FloorSweepDistance, OrigMoveDelta * (1.f - PercentTimeAppliedSoFar), MoveHitResult, CurrentFloor, false, &StepUpFloorResult, MoveRecord))
 					{
 						FMoverOnImpactParams ImpactParams(DefaultModeNames::Walking, MoveHitResult, OrigMoveDelta);
 						MoverComp->HandleImpact(ImpactParams);
@@ -233,7 +233,7 @@ void UWalkingMode::SimulationTick_Implementation(const FSimulationTickParams& Pa
 		}
 
 		// Search for the floor we've ended up on
-		UFloorQueryUtils::FindFloor(MovingComponents, CommonLegacySettings->FloorSweepDistance, CommonLegacySettings->MaxWalkSlopeCosine, UpdatedComponent->GetComponentLocation(), CurrentFloor);
+		UFloorQueryUtils::FindFloor(MovingComponents, CommonLegacySettings->FloorSweepDistance, CommonLegacySettings->MaxWalkSlopeCosine, CommonLegacySettings->bUseFlatBaseForFloorChecks, UpdatedComponent->GetComponentLocation(), CurrentFloor);
 
 		if (CurrentFloor.IsWalkableFloor())
 		{
@@ -256,7 +256,7 @@ void UWalkingMode::SimulationTick_Implementation(const FSimulationTickParams& Pa
 		if ((FloorCheckPolicy == EStaticFloorCheckPolicy::Always) || 
 		    (FloorCheckPolicy == EStaticFloorCheckPolicy::OnDynamicBaseOnly && StartingSyncState->GetMovementBase()))
 		{
-			UFloorQueryUtils::FindFloor(MovingComponents, CommonLegacySettings->FloorSweepDistance, CommonLegacySettings->MaxWalkSlopeCosine, UpdatedComponent->GetComponentLocation(), CurrentFloor);
+			UFloorQueryUtils::FindFloor(MovingComponents, CommonLegacySettings->FloorSweepDistance, CommonLegacySettings->MaxWalkSlopeCosine, CommonLegacySettings->bUseFlatBaseForFloorChecks, UpdatedComponent->GetComponentLocation(), CurrentFloor);
 		
 			FHitResult Hit(CurrentFloor.HitResult);
 			if (Hit.bStartPenetrating)
diff --git a/Engine/Plugins/Experimental/Mover/Source/Mover/Private/MoveLibrary/AirMovementUtils.cpp b/Engine/Plugins/Experimental/Mover/Source/Mover/Private/MoveLibrary/AirMovementUtils.cpp
--- a/Engine/Plugins/Experimental/Mover/Source/Mover/Private/MoveLibrary/AirMovementUtils.cpp
+++ b/Engine/Plugins/Experimental/Mover/Source/Mover/Private/MoveLibrary/AirMovementUtils.cpp
@@ -37,7 +37,7 @@ FProposedMove UAirMovementUtils::ComputeControlledFreeMove(const FFreeMoveParams
 	return OutMove;
 }
 
-bool UAirMovementUtils::IsValidLandingSpot(const FMovingComponentSet& MovingComps, const FVector& Location, const FHitResult& Hit, float FloorSweepDistance, float MaxWalkSlopeCosine, FFloorCheckResult& OutFloorResult)
+bool UAirMovementUtils::IsValidLandingSpot(const FMovingComponentSet& MovingComps, const FVector& Location, const FHitResult& Hit, float FloorSweepDistance, float MaxWalkSlopeCosine, bool bUseFlatBaseForFloorChecks, FFloorCheckResult& OutFloorResult)
 {
 	OutFloorResult.Clear();
 
@@ -58,7 +58,7 @@ bool UAirMovementUtils::IsValidLandingSpot(const FMovingComponentSet& MovingComp
 	}
 
 	// Make sure floor test passes here.
-	UFloorQueryUtils::FindFloor(MovingComps, FloorSweepDistance, MaxWalkSlopeCosine, Location, OutFloorResult);
+	UFloorQueryUtils::FindFloor(MovingComps, FloorSweepDistance, MaxWalkSlopeCosine, bUseFlatBaseForFloorChecks, Location, OutFloorResult);
 
 	if (!OutFloorResult.IsWalkableFloor())
 	{
@@ -68,7 +68,7 @@ bool UAirMovementUtils::IsValidLandingSpot(const FMovingComponentSet& MovingComp
 	return true;
 }
 
-float UAirMovementUtils::TryMoveToFallAlongSurface(const FMovingComponentSet& MovingComps, const FVector& Delta, float PctOfDeltaToMove, const FQuat Rotation, const FVector& Normal, FHitResult& Hit, bool bHandleImpact, float FloorSweepDistance, float MaxWalkSlopeCosine, FFloorCheckResult& OutFloorResult, FMovementRecord& MoveRecord)
+float UAirMovementUtils::TryMoveToFallAlongSurface(const FMovingComponentSet& MovingComps, const FVector& Delta, float PctOfDeltaToMove, const FQuat Rotation, const FVector& Normal, FHitResult& Hit, bool bHandleImpact, float FloorSweepDistance, float MaxWalkSlopeCosine, bool bUseFlatBaseForFloorChecks, FFloorCheckResult& OutFloorResult, FMovementRecord& MoveRecord)
 {
 	OutFloorResult.Clear();
 
@@ -102,7 +102,7 @@ float UAirMovementUtils::TryMoveToFallAlongSurface(const FMovingComponentSet& Mo
 
 			// Check if we landed
 			if (!IsValidLandingSpot(MovingComps, UpdatedPrimitive->GetComponentLocation(),
-				Hit, FloorSweepDistance, MaxWalkSlopeCosine, OutFloorResult))
+				Hit, FloorSweepDistance, MaxWalkSlopeCosine, bUseFlatBaseForFloorChecks, OutFloorResult))
 			{
 				// We've hit another surface during our first move, so let's try to slide along both of them together
 
@@ -125,7 +125,7 @@ float UAirMovementUtils::TryMoveToFallAlongSurface(const FMovingComponentSet& Mo
 
 					// Check if we've landed, to acquire floor result
 					IsValidLandingSpot(MovingComps, UpdatedPrimitive->GetComponentLocation(),
-						Hit, FloorSweepDistance, MaxWalkSlopeCosine, OutFloorResult);
+						Hit, FloorSweepDistance, MaxWalkSlopeCosine, bUseFlatBaseForFloorChecks, OutFloorResult);
 				}
 			}
 		}
@@ -138,7 +138,7 @@ float UAirMovementUtils::TryMoveToFallAlongSurface(const FMovingComponentSet& Mo
 
 
 /* static */
-float UAirMovementUtils::TestFallingMoveAlongHitSurface(const FMovingComponentSet& MovingComps, const FVector& OriginalMoveDelta, const FVector& LocationAtHit, const FQuat& TargetRotation, bool bHandleImpact, float FloorSweepDistance, float MaxWalkSlopeCosine, FHitResult& InOutHit, FFloorCheckResult& OutFloorResult, FMovementRecord& InOutMoveRecord)
+float UAirMovementUtils::TestFallingMoveAlongHitSurface(const FMovingComponentSet& MovingComps, const FVector& OriginalMoveDelta, const FVector& LocationAtHit, const FQuat& TargetRotation, bool bHandleImpact, float FloorSweepDistance, float MaxWalkSlopeCosine, bool bUseFlatBaseForFloorChecks, FHitResult& InOutHit, FFloorCheckResult& OutFloorResult, FMovementRecord& InOutMoveRecord)
 {
 	OutFloorResult.Clear();
 
@@ -176,7 +176,7 @@ float UAirMovementUtils::TestFallingMoveAlongHitSurface(const FMovingComponentSe
 
 			// Check if we landed
 			if (!UAirMovementUtils::IsValidLandingSpot(MovingComps, LocationAfter1stSlide,
-				InOutHit, FloorSweepDistance, MaxWalkSlopeCosine, OutFloorResult))
+				InOutHit, FloorSweepDistance, MaxWalkSlopeCosine, bUseFlatBaseForFloorChecks, OutFloorResult))
 			{
 				// We've hit another surface during our first move, so let's try to slide along both of them together
 
@@ -201,7 +201,7 @@ float UAirMovementUtils::TestFallingMoveAlongHitSurface(const FMovingComponentSe
 
 					// Check if we've landed, to acquire floor result
 					UAirMovementUtils::IsValidLandingSpot(MovingComps, LocationAfter2ndSlide,
-						InOutHit, FloorSweepDistance, MaxWalkSlopeCosine, OutFloorResult);
+						InOutHit, FloorSweepDistance, MaxWalkSlopeCosine, bUseFlatBaseForFloorChecks, OutFloorResult);
 				}
 			}
 		}
diff --git a/Engine/Plugins/Experimental/Mover/Source/Mover/Private/MoveLibrary/FloorQueryUtils.cpp b/Engine/Plugins/Experimental/Mover/Source/Mover/Private/MoveLibrary/FloorQueryUtils.cpp
--- a/Engine/Plugins/Experimental/Mover/Source/Mover/Private/MoveLibrary/FloorQueryUtils.cpp
+++ b/Engine/Plugins/Experimental/Mover/Source/Mover/Private/MoveLibrary/FloorQueryUtils.cpp
@@ -18,7 +18,7 @@ namespace UE::FloorQueryUtility
 	const float SWEEP_EDGE_REJECT_DISTANCE = 0.15f;
 }
 
-void UFloorQueryUtils::FindFloor(const FMovingComponentSet& MovingComps, float FloorSweepDistance, float MaxWalkSlopeCosine, const FVector& Location, FFloorCheckResult& OutFloorResult)
+void UFloorQueryUtils::FindFloor(const FMovingComponentSet& MovingComps, float FloorSweepDistance, float MaxWalkSlopeCosine, bool bUseFlatBaseForFloorChecks, const FVector& Location, FFloorCheckResult& OutFloorResult)
 {
 	if (!MovingComps.UpdatedComponent->IsQueryCollisionEnabled())
 	{
@@ -28,10 +28,10 @@ void UFloorQueryUtils::FindFloor(const FMovingComponentSet& MovingComps, float F
 
 	// Sweep for the floor
 	// TODO: Might need to plug in a different value for LineTraceDistance - using the same value as FloorSweepDistance for now - function takes both so we can plug in different values if needed
-	ComputeFloorDist(MovingComps,FloorSweepDistance, FloorSweepDistance, MaxWalkSlopeCosine, Location, OutFloorResult);
+	ComputeFloorDist(MovingComps,FloorSweepDistance, FloorSweepDistance, MaxWalkSlopeCosine, Location, bUseFlatBaseForFloorChecks, OutFloorResult);
 }
 
-void UFloorQueryUtils::ComputeFloorDist(const FMovingComponentSet& MovingComps, float LineTraceDistance, float FloorSweepDistance, float MaxWalkSlopeCosine, const FVector& Location, FFloorCheckResult& OutFloorResult)
+void UFloorQueryUtils::ComputeFloorDist(const FMovingComponentSet& MovingComps, float LineTraceDistance, float FloorSweepDistance, float MaxWalkSlopeCosine, const FVector& Location, bool bUseFlatBaseForFloorChecks, FFloorCheckResult& OutFloorResult)
 {
 	OutFloorResult.Clear();
 
@@ -63,7 +63,7 @@ void UFloorQueryUtils::ComputeFloorDist(const FMovingComponentSet& MovingComps,
 
 		FHitResult Hit(1.f);
 		FVector SweepDirection = UpDirection * -TraceDist;
-		bBlockingHit = FloorSweepTest(MovingComps, Hit, Location, Location + SweepDirection, CollisionChannel, CapsuleShape, QueryParams, ResponseParam);
+		bBlockingHit = FloorSweepTest(MovingComps, Hit, Location, Location + SweepDirection, CollisionChannel, CapsuleShape, QueryParams, ResponseParam, bUseFlatBaseForFloorChecks);
 
 		if (bBlockingHit)
 		{
@@ -82,7 +82,7 @@ void UFloorQueryUtils::ComputeFloorDist(const FMovingComponentSet& MovingComps,
 					CapsuleShape.Capsule.HalfHeight = FMath::Max(PawnHalfHeight - ShrinkHeight, CapsuleShape.Capsule.Radius);
 					Hit.Reset(1.f, false);
 
-					bBlockingHit = FloorSweepTest(MovingComps, Hit, Location, Location + SweepDirection, CollisionChannel, CapsuleShape, QueryParams, ResponseParam);
+					bBlockingHit = FloorSweepTest(MovingComps, Hit, Location, Location + SweepDirection, CollisionChannel, CapsuleShape, QueryParams, ResponseParam, bUseFlatBaseForFloorChecks);
 				}
 			}
 
@@ -145,13 +145,42 @@ void UFloorQueryUtils::ComputeFloorDist(const FMovingComponentSet& MovingComps,
 	OutFloorResult.bWalkableFloor = false;
 }
 
-bool UFloorQueryUtils::FloorSweepTest(const FMovingComponentSet& MovingComps, FHitResult& OutHit, const FVector& Start, const FVector& End, ECollisionChannel TraceChannel, const struct FCollisionShape& CollisionShape, const struct FCollisionQueryParams& Params, const struct FCollisionResponseParams& ResponseParam)
+bool UFloorQueryUtils::FloorSweepTest(const FMovingComponentSet& MovingComps, FHitResult& OutHit, const FVector& Start, const FVector& End, ECollisionChannel TraceChannel, const struct FCollisionShape& CollisionShape, const struct FCollisionQueryParams& Params, const struct FCollisionResponseParams& ResponseParam, bool bUseFlatBaseForFloorChecks)
 {
+	const FVector DirAwayFromFloor = (Start - End).GetSafeNormal();
+	if (!MovingComps.UpdatedPrimitive.IsValid() || !DirAwayFromFloor.IsNormalized())
+	{
+		return false;
+	}
+
+	UWorld* World = MovingComps.UpdatedPrimitive->GetWorld();
+
+	const FQuat UpDirOrientation = FRotationMatrix::MakeFromZX(DirAwayFromFloor, FVector::ForwardVector).ToQuat();
+
 	bool bBlockingHit = false;
 
-	if (MovingComps.UpdatedPrimitive.IsValid())
+	if (bUseFlatBaseForFloorChecks)
+	{
+		// Test with a box that is enclosed by the capsule: 2 checks at different rotations to get a good approximation of a circular flat bottom.
+		const float CapsuleRadius = CollisionShape.GetCapsuleRadius();
+		const float CapsuleHeight = CollisionShape.GetCapsuleHalfHeight();
+		const FCollisionShape BoxShape = FCollisionShape::MakeBox(FVector(CapsuleRadius * 0.707f, CapsuleRadius * 0.707f, CapsuleHeight));
+
+		// First test with the box rotated so the corners are along the major axes (ie rotated 45 degrees).
+		const FQuat Rotate45LocalYaw = FQuat::MakeFromEuler(FVector(0.0f, 0.0f, 45.0f));
+
+		bBlockingHit = World->SweepSingleByChannel(OutHit, Start, End, UpDirOrientation * Rotate45LocalYaw, TraceChannel, BoxShape, Params, ResponseParam);
+
+		if (!bBlockingHit)
+		{
+			// Test again with the same box, not rotated.
+			OutHit.Reset(1.f, false);
+			bBlockingHit = World->SweepSingleByChannel(OutHit, Start, End, UpDirOrientation, TraceChannel, BoxShape, Params, ResponseParam);
+		}
+	}
+	else
 	{
-		bBlockingHit = MovingComps.UpdatedPrimitive->GetWorld()->SweepSingleByChannel(OutHit, Start, End, MovingComps.UpdatedComponent->GetComponentRotation().Quaternion(), TraceChannel, CollisionShape, Params, ResponseParam);
+		bBlockingHit = World->SweepSingleByChannel(OutHit, Start, End, UpDirOrientation, TraceChannel, CollisionShape, Params, ResponseParam);
 	}
 
 	return bBlockingHit;
diff --git a/Engine/Plugins/Experimental/Mover/Source/Mover/Private/MoveLibrary/GroundMovementUtils.cpp b/Engine/Plugins/Experimental/Mover/Source/Mover/Private/MoveLibrary/GroundMovementUtils.cpp
--- a/Engine/Plugins/Experimental/Mover/Source/Mover/Private/MoveLibrary/GroundMovementUtils.cpp
+++ b/Engine/Plugins/Experimental/Mover/Source/Mover/Private/MoveLibrary/GroundMovementUtils.cpp
@@ -55,7 +55,7 @@ static const FName StepFwdSubstepName = "StepFwd";
 static const FName StepDownSubstepName = "StepDown";
 static const FName SlideSubstepName = "SlideFromStep";
 
-bool UGroundMovementUtils::TryMoveToStepUp(const FMovingComponentSet& MovingComps, const FVector& GravDir, float MaxStepHeight, float MaxWalkSlopeCosine, float FloorSweepDistance, const FVector& MoveDelta, const FHitResult& MoveHitResult, const FFloorCheckResult& CurrentFloor, bool bIsFalling, FOptionalFloorCheckResult* OutFloorTestResult, FMovementRecord& MoveRecord)
+bool UGroundMovementUtils::TryMoveToStepUp(const FMovingComponentSet& MovingComps, const FVector& GravDir, float MaxStepHeight, float MaxWalkSlopeCosine, bool bUseFlatBaseForFloorChecks, float FloorSweepDistance, const FVector& MoveDelta, const FHitResult& MoveHitResult, const FFloorCheckResult& CurrentFloor, bool bIsFalling, FOptionalFloorCheckResult* OutFloorTestResult, FMovementRecord& MoveRecord)
 {
 	UCapsuleComponent* CapsuleComponent = Cast<UCapsuleComponent>(MovingComps.UpdatedPrimitive.Get());
 
@@ -308,7 +308,7 @@ bool UGroundMovementUtils::TryMoveToStepUp(const FMovingComponentSet& MovingComp
 		if (OutFloorTestResult != NULL)
 		{
 
-			UFloorQueryUtils::FindFloor(MovingComps, FloorSweepDistance, MaxWalkSlopeCosine, CapsuleComponent->GetComponentLocation(), StepDownResult.FloorTestResult);
+			UFloorQueryUtils::FindFloor(MovingComps, FloorSweepDistance, MaxWalkSlopeCosine, bUseFlatBaseForFloorChecks, CapsuleComponent->GetComponentLocation(), StepDownResult.FloorTestResult);
 
 			// Reject unwalkable normals if we end up higher than our initial height.
 			// It's fine to walk down onto an unwalkable surface, don't reject those moves.
@@ -759,7 +759,7 @@ FVector UGroundMovementUtils::TestMoveToKeepMinHeightAboveFloor(const FMovingCom
 
 
 /* static */
-bool UGroundMovementUtils::TestMoveToStepOver(const FMovingComponentSet& MovingComps, const FVector& GravDir, float MaxStepHeight, float MaxWalkSlopeCosine, float FloorSweepDistance, const FVector& MoveDelta, const FQuat& Rotation, const FHitResult& MoveHitResult, const FFloorCheckResult& CurrentFloor, bool bIsFalling, FOptionalFloorCheckResult* OutFloorTestResult, FVector& OutFinalLocation, FMovementRecord& InOutMoveRecord)
+bool UGroundMovementUtils::TestMoveToStepOver(const FMovingComponentSet& MovingComps, const FVector& GravDir, float MaxStepHeight, float MaxWalkSlopeCosine, bool bUseFlatBaseForFloorChecks, float FloorSweepDistance, const FVector& MoveDelta, const FQuat& Rotation, const FHitResult& MoveHitResult, const FFloorCheckResult& CurrentFloor, bool bIsFalling, FOptionalFloorCheckResult* OutFloorTestResult, FVector& OutFinalLocation, FMovementRecord& InOutMoveRecord)
 {
 	UCapsuleComponent* CapsuleComponent = Cast<UCapsuleComponent>(MovingComps.UpdatedPrimitive.Get());
 
@@ -1012,7 +1012,7 @@ bool UGroundMovementUtils::TestMoveToStepOver(const FMovingComponentSet& MovingC
 		if (OutFloorTestResult != NULL)
 		{
 
-			UFloorQueryUtils::FindFloor(MovingComps, FloorSweepDistance, MaxWalkSlopeCosine, LocationInProgress, StepDownResult.FloorTestResult);
+			UFloorQueryUtils::FindFloor(MovingComps, FloorSweepDistance, MaxWalkSlopeCosine, bUseFlatBaseForFloorChecks, LocationInProgress, StepDownResult.FloorTestResult);
 
 			// Reject unwalkable normals if we end up higher than our initial height.
 			// It's fine to walk down onto an unwalkable surface, don't reject those moves.
diff --git a/Engine/Plugins/Experimental/Mover/Source/Mover/Public/DefaultMovementSet/Settings/CommonLegacyMovementSettings.h b/Engine/Plugins/Experimental/Mover/Source/Mover/Public/DefaultMovementSet/Settings/CommonLegacyMovementSettings.h
--- a/Engine/Plugins/Experimental/Mover/Source/Mover/Public/DefaultMovementSet/Settings/CommonLegacyMovementSettings.h
+++ b/Engine/Plugins/Experimental/Mover/Source/Mover/Public/DefaultMovementSet/Settings/CommonLegacyMovementSettings.h
@@ -43,6 +43,12 @@ public:
 	UPROPERTY(BlueprintReadWrite, EditAnywhere, Category="Ground Movement", meta = (ClampMin = "0", UIMin = "0", ForceUnits = "cm"))
 	float FloorSweepDistance = 40.0f;
 
+	/** Performs floor checks as if the object is using a shape with a flat base.
+	 *  This avoids the situation where we slowly lower off the side of a ledge when using shapes with rounded bottoms, like a capsule. 
+	 */
+	UPROPERTY(BlueprintReadWrite, EditAnywhere, Category = "Ground Movement")
+	bool bUseFlatBaseForFloorChecks = true;
+
 	/** Mover actors will be able to step up onto or over obstacles shorter than this */
 	UPROPERTY(BlueprintReadWrite, EditAnywhere, Category="Ground Movement", meta = (ClampMin = "0", UIMin = "0", ForceUnits = "cm"))
 	float MaxStepHeight = 40.0f;
diff --git a/Engine/Plugins/Experimental/Mover/Source/Mover/Public/MoveLibrary/AirMovementUtils.h b/Engine/Plugins/Experimental/Mover/Source/Mover/Public/MoveLibrary/AirMovementUtils.h
--- a/Engine/Plugins/Experimental/Mover/Source/Mover/Public/MoveLibrary/AirMovementUtils.h
+++ b/Engine/Plugins/Experimental/Mover/Source/Mover/Public/MoveLibrary/AirMovementUtils.h
@@ -75,11 +75,11 @@ public:
 	
     // Checks if a hit result represents a walkable location that an actor can land on
     UFUNCTION(BlueprintCallable, Category=Mover)
-    static UE_API bool IsValidLandingSpot(const FMovingComponentSet& MovingComps, const FVector& Location, const FHitResult& Hit, float FloorSweepDistance, float MaxWalkSlopeCosine, UPARAM(Ref) FFloorCheckResult& OutFloorResult);
+    static UE_API bool IsValidLandingSpot(const FMovingComponentSet& MovingComps, const FVector& Location, const FHitResult& Hit, float FloorSweepDistance, float MaxWalkSlopeCosine, bool bUseFlatBaseForFloorChecks, UPARAM(Ref) FFloorCheckResult& OutFloorResult);
     
     /** Attempts to move a component along a surface, while checking for landing on a walkable surface. Intended for use while falling. Returns the percent of time applied, with 0.0 meaning no movement occurred. */
     UFUNCTION(BlueprintCallable, Category=Mover)
-    static UE_API float TryMoveToFallAlongSurface(const FMovingComponentSet& MovingComps, const FVector& Delta, float PctOfDeltaToMove, const FQuat Rotation, const FVector& Normal, UPARAM(Ref) FHitResult& Hit, bool bHandleImpact, float FloorSweepDistance, float MaxWalkSlopeCosine, UPARAM(Ref) FFloorCheckResult& OutFloorResult, UPARAM(Ref) FMovementRecord& MoveRecord);
+    static UE_API float TryMoveToFallAlongSurface(const FMovingComponentSet& MovingComps, const FVector& Delta, float PctOfDeltaToMove, const FQuat Rotation, const FVector& Normal, UPARAM(Ref) FHitResult& Hit, bool bHandleImpact, float FloorSweepDistance, float MaxWalkSlopeCosine, bool bUseFlatBaseForFloorChecks, UPARAM(Ref) FFloorCheckResult& OutFloorResult, UPARAM(Ref) FMovementRecord& MoveRecord);
 
 public:	
 	// Threadsafe movement queries
@@ -91,7 +91,7 @@ public:
 	 * Returns the percent of time applied, with 0.0 meaning no movement would occur. 
 	 */
 	UFUNCTION(BlueprintCallable, Category = Mover)
-	static UE_API float TestFallingMoveAlongHitSurface(const FMovingComponentSet& MovingComps, const FVector& OriginalMoveDelta, const FVector& LocationAtHit, const FQuat& TargetRotation, bool bHandleImpact, float FloorSweepDistance, float MaxWalkSlopeCosine, UPARAM(Ref) FHitResult& InOutHit, UPARAM(Ref) FFloorCheckResult& OutFloorResult, UPARAM(Ref) FMovementRecord& InOutMoveRecord);
+	static UE_API float TestFallingMoveAlongHitSurface(const FMovingComponentSet& MovingComps, const FVector& OriginalMoveDelta, const FVector& LocationAtHit, const FQuat& TargetRotation, bool bHandleImpact, float FloorSweepDistance, float MaxWalkSlopeCosine, bool bUseFlatBaseForFloorChecks, UPARAM(Ref) FHitResult& InOutHit, UPARAM(Ref) FFloorCheckResult& OutFloorResult, UPARAM(Ref) FMovementRecord& InOutMoveRecord);
 
 
 };
diff --git a/Engine/Plugins/Experimental/Mover/Source/Mover/Public/MoveLibrary/FloorQueryUtils.h b/Engine/Plugins/Experimental/Mover/Source/Mover/Public/MoveLibrary/FloorQueryUtils.h
--- a/Engine/Plugins/Experimental/Mover/Source/Mover/Public/MoveLibrary/FloorQueryUtils.h
+++ b/Engine/Plugins/Experimental/Mover/Source/Mover/Public/MoveLibrary/FloorQueryUtils.h
@@ -111,11 +111,11 @@ class UFloorQueryUtils : public UBlueprintFunctionLibrary
 	GENERATED_BODY()
 
 public:
-	static UE_API void FindFloor(const FMovingComponentSet& MovingComps, float FloorSweepDistance, float MaxWalkSlopeCosine, const FVector& Location, FFloorCheckResult& OutFloorResult);
+	static UE_API void FindFloor(const FMovingComponentSet& MovingComps, float FloorSweepDistance, float MaxWalkSlopeCosine, bool bUseFlatBaseForFloorChecks, const FVector& Location, FFloorCheckResult& OutFloorResult);
 
-	static UE_API void ComputeFloorDist(const FMovingComponentSet& MovingComps, float LineTraceDistance, float FloorSweepDistance, float MaxWalkSlopeCosine, const FVector& Location, FFloorCheckResult& OutFloorResult);
+	static UE_API void ComputeFloorDist(const FMovingComponentSet& MovingComps, float LineTraceDistance, float FloorSweepDistance, float MaxWalkSlopeCosine, const FVector& Location, bool bUseFlatBaseForFloorChecks, FFloorCheckResult& OutFloorResult);
 
-	static UE_API bool FloorSweepTest(const FMovingComponentSet& MovingComps, FHitResult& OutHit, const FVector& Start, const FVector& End, ECollisionChannel TraceChannel, const struct FCollisionShape& CollisionShape, const struct FCollisionQueryParams& Params, const struct FCollisionResponseParams& ResponseParam);
+	static UE_API bool FloorSweepTest(const FMovingComponentSet& MovingComps, FHitResult& OutHit, const FVector& Start, const FVector& End, ECollisionChannel TraceChannel, const FCollisionShape& CollisionShape, const struct FCollisionQueryParams& Params, const struct FCollisionResponseParams& ResponseParam, bool bUseFlatBaseForFloorChecks);
 
 	UFUNCTION(BlueprintCallable, Category=Mover)
 	static UE_API bool IsHitSurfaceWalkable(const FHitResult& Hit, const FVector& UpDirection, float MaxWalkSlopeCosine);
diff --git a/Engine/Plugins/Experimental/Mover/Source/Mover/Public/MoveLibrary/GroundMovementUtils.h b/Engine/Plugins/Experimental/Mover/Source/Mover/Public/MoveLibrary/GroundMovementUtils.h
--- a/Engine/Plugins/Experimental/Mover/Source/Mover/Public/MoveLibrary/GroundMovementUtils.h
+++ b/Engine/Plugins/Experimental/Mover/Source/Mover/Public/MoveLibrary/GroundMovementUtils.h
@@ -86,7 +86,7 @@ public:
 	
 	// TODO: Refactor this API for fewer parameters
 	/** Move up steps or slope. Does nothing and returns false if hit surface is invalid for step-up use */
-	static UE_API bool TryMoveToStepUp(const FMovingComponentSet& MovingComps, const FVector& GravDir, float MaxStepHeight, float MaxWalkSlopeCosine, float FloorSweepDistance, const FVector& MoveDelta, const FHitResult& Hit, const FFloorCheckResult& CurrentFloor, bool bIsFalling, FOptionalFloorCheckResult* OutStepDownResult, FMovementRecord& MoveRecord);
+	static UE_API bool TryMoveToStepUp(const FMovingComponentSet& MovingComps, const FVector& GravDir, float MaxStepHeight, float MaxWalkSlopeCosine, bool bUseFlatBaseForFloorChecks, float FloorSweepDistance, const FVector& MoveDelta, const FHitResult& Hit, const FFloorCheckResult& CurrentFloor, bool bIsFalling, FOptionalFloorCheckResult* OutStepDownResult, FMovementRecord& MoveRecord);
 
 	/** Moves vertically to stay within range of the walkable floor. Does nothing and returns false if floor is unwalkable or if already in range. */ 
 	static UE_API bool TryMoveToAdjustHeightAboveFloor(const FMovingComponentSet& MovingComps, FFloorCheckResult& CurrentFloor, float MaxWalkSlopeCosine, FMovementRecord& MoveRecord);
@@ -140,7 +140,7 @@ public:
      * Appends to InOutMoveRecord with any movement substeps
 	 * Returns whether any movement was possible.
 	 */
-	static UE_API bool TestMoveToStepOver(const FMovingComponentSet& MovingComps, const FVector& GravDir, float MaxStepHeight, float MaxWalkSlopeCosine, float FloorSweepDistance, const FVector& MoveDelta, const FQuat& Rotation, const FHitResult& MoveHitResult, const FFloorCheckResult& CurrentFloor, bool bIsFalling, FOptionalFloorCheckResult* OutFloorTestResult, FVector& OutFinalLocation, FMovementRecord& InOutMoveRecord);
+	static UE_API bool TestMoveToStepOver(const FMovingComponentSet& MovingComps, const FVector& GravDir, float MaxStepHeight, float MaxWalkSlopeCosine, bool bUseFlatBaseForFloorChecks, float FloorSweepDistance, const FVector& MoveDelta, const FQuat& Rotation, const FHitResult& MoveHitResult, const FFloorCheckResult& CurrentFloor, bool bIsFalling, FOptionalFloorCheckResult* OutFloorTestResult, FVector& OutFinalLocation, FMovementRecord& InOutMoveRecord);
 
 
 };
