Mover: Properly exposing Movement Mode in-out references in GenerateMove and SimulationTick to BP. Also fixing up couple BP examples of a Movement Mode to use new reference properties.

#virtualized

diff --git a/Engine/Plugins/Experimental/Mover/Source/Mover/Public/MovementMode.h b/Engine/Plugins/Experimental/Mover/Source/Mover/Public/MovementMode.h
--- a/Engine/Plugins/Experimental/Mover/Source/Mover/Public/MovementMode.h
+++ b/Engine/Plugins/Experimental/Mover/Source/Mover/Public/MovementMode.h
@@ -55,11 +55,11 @@ public:
 	UE_API virtual void Activate_External();
 	UE_API virtual void Deactivate_External();
 	
-	UFUNCTION(BlueprintNativeEvent, meta = (DisplayName = "Generate Move"))
-	UE_API void GenerateMove(const FMoverTickStartData& StartState, const FMoverTimeStep& TimeStep, FProposedMove& OutProposedMove) const;
+	UFUNCTION(BlueprintNativeEvent, meta = (DisplayName = "Generate Move", ForceAsFunction))
+	UE_API void GenerateMove(const FMoverTickStartData& StartState, const FMoverTimeStep& TimeStep, UPARAM(ref) FProposedMove& OutProposedMove) const;
 
-	UFUNCTION(BlueprintNativeEvent, meta = (DisplayName = "Simulation Tick"))
-	UE_API void SimulationTick(const FSimulationTickParams& Params, FMoverTickEndData& OutputState);
+	UFUNCTION(BlueprintNativeEvent, meta = (DisplayName = "Simulation Tick", ForceAsFunction))
+	UE_API void SimulationTick(const FSimulationTickParams& Params, UPARAM(ref) FMoverTickEndData& OutputState);
 	
 	/** Gets the MoverComponent that owns this movement mode */
 	UFUNCTION(BlueprintCallable, Category=Mover, meta=(DisplayName="Get Mover Component", ScriptName = GetMoverComponent))
