Mover network fixes:
Fix for trajectory prediction failing inside mover on sim proxies
Handle missing input commands on simple walking modes to stop sim proxies trying to stop

diff --git a/Engine/Plugins/Experimental/Mover/Source/Mover/Private/DefaultMovementSet/Modes/SimpleWalkingMode.cpp b/Engine/Plugins/Experimental/Mover/Source/Mover/Private/DefaultMovementSet/Modes/SimpleWalkingMode.cpp
--- a/Engine/Plugins/Experimental/Mover/Source/Mover/Private/DefaultMovementSet/Modes/SimpleWalkingMode.cpp
+++ b/Engine/Plugins/Experimental/Mover/Source/Mover/Private/DefaultMovementSet/Modes/SimpleWalkingMode.cpp
@@ -36,12 +36,12 @@ void USimpleWalkingMode::GenerateMove_Implementation(const FMoverTickStartData&
 	}
 	else
 	{
-		DesiredVelocity = FVector::ZeroVector;
+		// If no input found, most likely a networked sim proxy
+		// Try to deduce an input from the sync state
+		DesiredVelocity = StartingSyncState->GetIntent_WorldSpace();
 		MoveInputType = EMoveInputType::DirectionalIntent;
-		DesiredFacingDir = FVector::ZeroVector;
+		DesiredFacingDir = StartingSyncState->GetOrientation_WorldSpace().Quaternion().GetForwardVector();
 	}
-	
-	
 	float MaxMoveSpeed = MaxSpeedOverride >= 0.0f ? MaxSpeedOverride : CommonLegacySettings->MaxSpeed;
 
 	// Subtract vertical component but keep same magnitude
diff --git a/Engine/Plugins/Experimental/Mover/Source/Mover/Private/MoverComponent.cpp b/Engine/Plugins/Experimental/Mover/Source/Mover/Private/MoverComponent.cpp
--- a/Engine/Plugins/Experimental/Mover/Source/Mover/Private/MoverComponent.cpp
+++ b/Engine/Plugins/Experimental/Mover/Source/Mover/Private/MoverComponent.cpp
@@ -2261,7 +2261,7 @@ TArray<FTrajectorySampleInfo> UMoverComponent::GetPredictedTrajectory(FMoverPred
 		FutureTimeStep.BaseSimTimeMs = CachedLastSimTickTimeStep.BaseSimTimeMs;
 		FutureTimeStep.ServerFrame = 0;
 
-		if (const UBaseMovementMode* CurrentMovementMode = ModeFSM->GetCurrentMode())
+		if (const UBaseMovementMode* CurrentMovementMode = GetMovementMode())
 		{
 			if (FMoverDefaultSyncState* StepSyncState = StepState.SyncState.SyncStateCollection.FindMutableDataByType<FMoverDefaultSyncState>())
 			{
