Mover: improving montage replication in Independent ticking modes by sharing the most recent animation position.  This was needed because clients using Independent ticking mode do not have a shared timeline.  Includes new cvar for enabling logging of montage start/stop events.

#rnx

diff --git a/Engine/Plugins/Experimental/Mover/Source/Mover/Private/DefaultMovementSet/CharacterMoverComponent.cpp b/Engine/Plugins/Experimental/Mover/Source/Mover/Private/DefaultMovementSet/CharacterMoverComponent.cpp
--- a/Engine/Plugins/Experimental/Mover/Source/Mover/Private/DefaultMovementSet/CharacterMoverComponent.cpp
+++ b/Engine/Plugins/Experimental/Mover/Source/Mover/Private/DefaultMovementSet/CharacterMoverComponent.cpp
@@ -16,6 +16,14 @@
 
 #include UE_INLINE_GENERATED_CPP_BY_NAME(CharacterMoverComponent)
 
+#if !UE_BUILD_SHIPPING
+FAutoConsoleVariable CVarLogSimProxyMontageReplication(
+	TEXT("mover.debug.LogSimProxyMontageReplication"),
+	false,
+	TEXT("Whether to log detailed information about montage replication on a sim proxy using the Character-focused MoverComponent. 0: Disable, 1: Enable"),
+	ECVF_Cheat);
+#endif	// !UE_BUILD_SHIPPING
+
 UCharacterMoverComponent::UCharacterMoverComponent()
 {
 	// Default movement modes
@@ -267,6 +275,11 @@ void UCharacterMoverComponent::UpdateSyncedMontageState(const FMoverTimeStep& Ti
 
 			if (bShouldStopSyncedMontage)
 			{
+				#if !UE_BUILD_SHIPPING
+				UE_CLOG(CVarLogSimProxyMontageReplication->GetBool(), LogMover, Log, TEXT("Mover SP montage repl (SimF %i SimT: %.3f): STOP %s"),
+					TimeStep.ServerFrame, TimeStep.BaseSimTimeMs, *SyncedMontageState.Montage->GetName());
+				#endif // !UE_BUILD_SHIPPING
+
 				if (MeshAnimInstance)
 				{
 					MeshAnimInstance->Montage_Stop(SyncedMontageState.Montage->GetDefaultBlendOutTime(), SyncedMontageState.Montage);
@@ -277,9 +290,14 @@ void UCharacterMoverComponent::UpdateSyncedMontageState(const FMoverTimeStep& Ti
 
 			if (bShouldStartNewMontage && NewMontageState.Montage && MeshAnimInstance)
 			{
-				const float StartPosition = NewMontageState.StartingMontagePosition + (FMath::Max(0.0f, (TimeStep.BaseSimTimeMs - MontageStateProvider->StartSimTimeMs) * 0.001f) * NewMontageState.PlayRate);
+				const float StartPosition = NewMontageState.CurrentPosition;
 				const float PlaySeconds = MeshAnimInstance->Montage_Play(NewMontageState.Montage, NewMontageState.PlayRate, EMontagePlayReturnType::MontageLength, StartPosition);
 
+				#if !UE_BUILD_SHIPPING
+				UE_CLOG(CVarLogSimProxyMontageReplication->GetBool(), LogMover, Log, TEXT("Mover SP montage repl (SimF %i SimT: %.3f): PLAY %s (StartPos: %.3f  Rate: %.3f  PlaySecs: %.3f)"),
+					TimeStep.ServerFrame, TimeStep.BaseSimTimeMs, *NewMontageState.Montage->GetName(), StartPosition, NewMontageState.PlayRate, PlaySeconds);
+				#endif // !UE_BUILD_SHIPPING
+
 				if (PlaySeconds > 0.0f)
 				{
 					SyncedMontageState = NewMontageState;	// only consider us sync'd if the montage actually started
diff --git a/Engine/Plugins/Experimental/Mover/Source/Mover/Private/DefaultMovementSet/LayeredMoves/AnimRootMotionLayeredMove.cpp b/Engine/Plugins/Experimental/Mover/Source/Mover/Private/DefaultMovementSet/LayeredMoves/AnimRootMotionLayeredMove.cpp
--- a/Engine/Plugins/Experimental/Mover/Source/Mover/Private/DefaultMovementSet/LayeredMoves/AnimRootMotionLayeredMove.cpp
+++ b/Engine/Plugins/Experimental/Mover/Source/Mover/Private/DefaultMovementSet/LayeredMoves/AnimRootMotionLayeredMove.cpp
@@ -97,6 +97,8 @@ bool FLayeredMove_AnimRootMotion::GenerateMove(const FMoverTickStartData& SimSta
 		OutProposedMove.AngularVelocityDegrees   = FMath::RadiansToDegrees(WorldSpaceRootMotion.GetRotation().ToRotationVector() / DeltaSeconds);
 	}
 
+	MontageState.CurrentPosition = ExtractionStartPosition;
+
 #if !UE_BUILD_SHIPPING
 	UE_CLOG(CVarLogAnimRootMotionSteps->GetBool(), LogMover, Log, TEXT("AnimRootMotion. SimF %i (dt %.3f) Range [%.3f, %.3f] => LocalT: %s (WST: %s)  Vel: %.3f"),
 	        TimeStep.ServerFrame, DeltaSeconds, ExtractionStartPosition, ExtractionEndPosition, 
diff --git a/Engine/Plugins/Experimental/Mover/Source/Mover/Private/MoveLibrary/PlayMoverMontageCallbackProxy.cpp b/Engine/Plugins/Experimental/Mover/Source/Mover/Private/MoveLibrary/PlayMoverMontageCallbackProxy.cpp
--- a/Engine/Plugins/Experimental/Mover/Source/Mover/Private/MoveLibrary/PlayMoverMontageCallbackProxy.cpp
+++ b/Engine/Plugins/Experimental/Mover/Source/Mover/Private/MoveLibrary/PlayMoverMontageCallbackProxy.cpp
@@ -60,6 +60,7 @@ bool UPlayMoverMontageCallbackProxy::PlayMoverMontage(
 				AnimRootMotionMove->MontageState.Montage = MontageToPlay;
 				AnimRootMotionMove->MontageState.PlayRate = PlayRate;
 				AnimRootMotionMove->MontageState.StartingMontagePosition = StartingMontagePosition;
+				AnimRootMotionMove->MontageState.CurrentPosition = StartingMontagePosition;
 				
 				float RemainingUnscaledMontageSeconds(0.f);
 
diff --git a/Engine/Plugins/Experimental/Mover/Source/Mover/Public/DefaultMovementSet/LayeredMoves/MontageStateProvider.h b/Engine/Plugins/Experimental/Mover/Source/Mover/Public/DefaultMovementSet/LayeredMoves/MontageStateProvider.h
--- a/Engine/Plugins/Experimental/Mover/Source/Mover/Public/DefaultMovementSet/LayeredMoves/MontageStateProvider.h
+++ b/Engine/Plugins/Experimental/Mover/Source/Mover/Public/DefaultMovementSet/LayeredMoves/MontageStateProvider.h
@@ -28,6 +28,9 @@ struct FMoverAnimMontageState
 	UPROPERTY(BlueprintReadWrite, Category = Mover)
 	float PlayRate = 1.0f;
 
+	// Current position (during playback only)
+	UPROPERTY(BlueprintReadOnly, Category = Mover)
+	float CurrentPosition = 0.0f;
 
 	void Reset();
 	void NetSerialize(FArchive& Ar);
