Mover: fixed missing window for blackboard changes that could occur when editing a physics pawn in a running PIE session

#rnx

diff --git a/Engine/Plugins/Experimental/Mover/Source/Mover/Private/Backends/MoverNetworkPhysicsLiaisonBase.cpp b/Engine/Plugins/Experimental/Mover/Source/Mover/Private/Backends/MoverNetworkPhysicsLiaisonBase.cpp
--- a/Engine/Plugins/Experimental/Mover/Source/Mover/Private/Backends/MoverNetworkPhysicsLiaisonBase.cpp
+++ b/Engine/Plugins/Experimental/Mover/Source/Mover/Private/Backends/MoverNetworkPhysicsLiaisonBase.cpp
@@ -293,8 +293,12 @@ void UMoverNetworkPhysicsLiaisonComponentBase::HandleComponentPhysicsStateChange
 	{
 		if (UMoverComponent& MoverComp = GetMoverComponent(); MoverComp.ModeFSM)
 		{
+			MoverComp.RollbackBlackboard_InternalWrapper->BeginSimulationFrame(GetCurrentAsyncMoverTimeStep_External());
+
 			MoverComp.ModeFSM->SetModeImmediately(MoverComp.StartingMovementMode);
 			MoverComp.SimBlackboard->InvalidateAll();
+
+			MoverComp.RollbackBlackboard_InternalWrapper->EndSimulationFrame();
 		}
 
 		InitializeSimOutputData();
