Mover: Passing along starting velocity in Falling modes if we didn't use any time. This fixes some velocity cancelation that sometimes causes a stutter on landing.

diff --git a/Engine/Plugins/Experimental/Mover/Source/Mover/Private/DefaultMovementSet/Modes/AsyncFallingMode.cpp b/Engine/Plugins/Experimental/Mover/Source/Mover/Private/DefaultMovementSet/Modes/AsyncFallingMode.cpp
--- a/Engine/Plugins/Experimental/Mover/Source/Mover/Private/DefaultMovementSet/Modes/AsyncFallingMode.cpp
+++ b/Engine/Plugins/Experimental/Mover/Source/Mover/Private/DefaultMovementSet/Modes/AsyncFallingMode.cpp
@@ -227,7 +227,7 @@ void UAsyncFallingMode::SimulationTick_Implementation(const FSimulationTickParam
 		{
 			// Adjust height to float slightly above walkable floor
 			LocationInProgress = UGroundMovementUtils::TestMoveToKeepMinHeightAboveFloor(Params.MovingComps, LocationInProgress, RotationInProgress, CommonLegacySettings->MaxWalkSlopeCosine, IN OUT LandingFloor, IN OUT MoveRecord);
-			CaptureFinalState(LocationInProgress, RotationInProgress.Rotator(), LandingFloor, DeltaSeconds, DeltaSeconds * PctTimeApplied, ProposedMove.AngularVelocityDegrees, OutputSyncState, OutputState, IN OUT MoveRecord);
+			CaptureFinalState(StartingSyncState, LocationInProgress, RotationInProgress.Rotator(), LandingFloor, DeltaSeconds, DeltaSeconds * PctTimeApplied, ProposedMove.AngularVelocityDegrees, OutputSyncState, OutputState, IN OUT MoveRecord);
 			return;
 		}
 
@@ -264,7 +264,7 @@ void UAsyncFallingMode::SimulationTick_Implementation(const FSimulationTickParam
 		PctTimeApplied = 1.f;
 	}
 	
-	CaptureFinalState(LocationInProgress, RotationInProgress.Rotator(), LandingFloor, DeltaSeconds, DeltaSeconds* PctTimeApplied,ProposedMove.AngularVelocityDegrees, OutputSyncState, OutputState, MoveRecord);
+	CaptureFinalState(StartingSyncState, LocationInProgress, RotationInProgress.Rotator(), LandingFloor, DeltaSeconds, DeltaSeconds* PctTimeApplied,ProposedMove.AngularVelocityDegrees, OutputSyncState, OutputState, MoveRecord);
 }
 
 
@@ -327,7 +327,7 @@ void UAsyncFallingMode::ProcessLanded(const FFloorCheckResult& FloorResult, FVec
 	}
 }
 
-void UAsyncFallingMode::CaptureFinalState(const FVector FinalLocation, const FRotator FinalRotation, const FFloorCheckResult& FloorResult, float DeltaSeconds, float DeltaSecondsUsed, const FVector& AngularVelocityDegrees, FMoverDefaultSyncState& OutputSyncState, FMoverTickEndData& TickEndData, FMovementRecord& Record) const
+void UAsyncFallingMode::CaptureFinalState(const FMoverDefaultSyncState* StartSyncState, const FVector FinalLocation, const FRotator FinalRotation, const FFloorCheckResult& FloorResult, float DeltaSeconds, float DeltaSecondsUsed, const FVector& AngularVelocityDegrees, FMoverDefaultSyncState& OutputSyncState, FMoverTickEndData& TickEndData, FMovementRecord& Record) const
 {
 	UMoverBlackboard* SimBlackboard = GetMoverComponent()->GetSimBlackboard_Mutable();
 
@@ -346,7 +346,8 @@ void UAsyncFallingMode::CaptureFinalState(const FVector FinalLocation, const FRo
 	
 	Record.SetDeltaSeconds( DeltaSecondsUsed );
 	
-	FVector EffectiveVelocity = Record.GetRelevantVelocity();
+	// If we didn't use any time lets just pass along velocity so we don't lose it when we go into the next mode with refunded time
+	FVector EffectiveVelocity = DeltaSecondsUsed <= UE_SMALL_NUMBER ? StartSyncState->GetVelocity_WorldSpace() : Record.GetRelevantVelocity();
 
 	FRelativeBaseInfo MovementBaseInfo;
 	ProcessLanded(FloorResult, EffectiveVelocity, MovementBaseInfo, TickEndData);
diff --git a/Engine/Plugins/Experimental/Mover/Source/Mover/Private/DefaultMovementSet/Modes/FallingMode.cpp b/Engine/Plugins/Experimental/Mover/Source/Mover/Private/DefaultMovementSet/Modes/FallingMode.cpp
--- a/Engine/Plugins/Experimental/Mover/Source/Mover/Private/DefaultMovementSet/Modes/FallingMode.cpp
+++ b/Engine/Plugins/Experimental/Mover/Source/Mover/Private/DefaultMovementSet/Modes/FallingMode.cpp
@@ -337,7 +337,8 @@ void UFallingMode::CaptureFinalState(USceneComponent* UpdatedComponent, const FM
 	
 	Record.SetDeltaSeconds( DeltaSecondsUsed );
 	
-	FVector EffectiveVelocity = Record.GetRelevantVelocity();
+	// If we didn't use any time lets just pass along velocity so we don't lose it when we go into the next mode with refunded time
+	FVector EffectiveVelocity = DeltaSecondsUsed <= UE_SMALL_NUMBER ? StartSyncState.GetVelocity_WorldSpace() : Record.GetRelevantVelocity();
 	// TODO: Update Main/large movement record with substeps from our local record
 
 	FRelativeBaseInfo MovementBaseInfo;
diff --git a/Engine/Plugins/Experimental/Mover/Source/Mover/Public/DefaultMovementSet/Modes/AsyncFallingMode.h b/Engine/Plugins/Experimental/Mover/Source/Mover/Public/DefaultMovementSet/Modes/AsyncFallingMode.h
--- a/Engine/Plugins/Experimental/Mover/Source/Mover/Public/DefaultMovementSet/Modes/AsyncFallingMode.h
+++ b/Engine/Plugins/Experimental/Mover/Source/Mover/Public/DefaultMovementSet/Modes/AsyncFallingMode.h
@@ -91,7 +91,7 @@ protected:
 	UFUNCTION(BlueprintCallable, Category=Mover)
 	UE_API virtual void ProcessLanded(const FFloorCheckResult& FloorResult, FVector& Velocity, FRelativeBaseInfo& BaseInfo, FMoverTickEndData& TickEndData) const;
 
-	UE_API void CaptureFinalState(const FVector FinalLocation, const FRotator FinalRotation, const FFloorCheckResult& FloorResult, float DeltaSeconds, float DeltaSecondsUsed, const FVector& AngularVelocityDegrees, FMoverDefaultSyncState& OutputSyncState, FMoverTickEndData& TickEndData, FMovementRecord& Record) const;
+	UE_API void CaptureFinalState(const FMoverDefaultSyncState* StartSyncState, const FVector FinalLocation, const FRotator FinalRotation, const FFloorCheckResult& FloorResult, float DeltaSeconds, float DeltaSecondsUsed, const FVector& AngularVelocityDegrees, FMoverDefaultSyncState& OutputSyncState, FMoverTickEndData& TickEndData, FMovementRecord& Record) const;
 
 	TWeakObjectPtr<const UCommonLegacyMovementSettings> CommonLegacySettings;
 };
