Mover: adding rollback blackboard support to Chaos Mover simulations, along with mode change records as the first entry type. Also added prev mode name to record type.

diff --git a/Engine/Plugins/Experimental/ChaosMover/Source/ChaosMover/Private/Backends/ChaosMoverBackend.cpp b/Engine/Plugins/Experimental/ChaosMover/Source/ChaosMover/Private/Backends/ChaosMoverBackend.cpp
--- a/Engine/Plugins/Experimental/ChaosMover/Source/ChaosMover/Private/Backends/ChaosMoverBackend.cpp
+++ b/Engine/Plugins/Experimental/ChaosMover/Source/ChaosMover/Private/Backends/ChaosMoverBackend.cpp
@@ -68,6 +68,8 @@ void UChaosMoverBackendComponent::InitializeComponent()
 		ImmediateModeTransition = NewObject<UImmediateMovementModeTransition>(&GetMoverComponent(), TEXT("ImmediateModeTransition"), RF_Transient);
 		TransformOnInit = GetMoverComponent().GetUpdatedComponentTransform();
 
+		Simulation->SetRollbackBlackboard(GetMoverComponent().GetRollbackBlackboard_Internal());
+
 		// Create NetworkPhysicsComponent
 		if ((World->GetNetMode() != ENetMode::NM_Standalone) && Chaos::FPhysicsSolverBase::IsNetworkPhysicsPredictionEnabled())
 		{
diff --git a/Engine/Plugins/Experimental/ChaosMover/Source/ChaosMover/Private/ChaosMoverSimulation.cpp b/Engine/Plugins/Experimental/ChaosMover/Source/ChaosMover/Private/ChaosMoverSimulation.cpp
--- a/Engine/Plugins/Experimental/ChaosMover/Source/ChaosMover/Private/ChaosMoverSimulation.cpp
+++ b/Engine/Plugins/Experimental/ChaosMover/Source/ChaosMover/Private/ChaosMoverSimulation.cpp
@@ -23,6 +23,7 @@
 #include "MoveLibrary/FloorQueryUtils.h"
 #include "MoveLibrary/MovementMixer.h"
 #include "MoveLibrary/MoverBlackboard.h"
+#include "MoveLibrary/RollbackBlackboard.h"
 #include "MoveLibrary/WaterMovementUtils.h"
 #include "MovementModeStateMachine.h"
 #include "MoverSimulationTypes.h"
@@ -217,9 +218,13 @@ void UChaosMoverSimulation::SimulationTick(const FMoverTimeStep& TimeStep, const
 
 	if (bInitialized)
 	{
+		RollbackBlackboard->BeginSimulationFrame(TimeStep);
+
 		OnPreSimulationTick(TimeStep, InputData);
 		OnSimulationTick(TimeStep, InputData, OutputData);
 		OnPostSimulationTick(TimeStep, OutputData);
+
+		RollbackBlackboard->EndSimulationFrame();
 	}
 }
 
@@ -373,6 +378,7 @@ void UChaosMoverSimulation::OnSimulationRollback(const FMoverTimeStep& NewTimeSt
 {
 	// Rollback blackboard on the first frame of resimulation
 	Blackboard->Invalidate(EInvalidationReason::Rollback);
+	RollbackBlackboard->BeginRollback(NewTimeStep);
 
 	// Remove any events that have been stored for this simulation tick
 	Events.Empty();
@@ -384,6 +390,8 @@ void UChaosMoverSimulation::OnSimulationRollback(const FMoverTimeStep& NewTimeSt
 	}
 
 	StateMachine.OnSimulationRollback(NewTimeStep, PrevSyncState, CurrentSyncState);
+	
+	RollbackBlackboard->EndRollback();
 }
 
 void UChaosMoverSimulation::OnPreSimulationTick(const FMoverTimeStep& TimeStep, const UE::ChaosMover::FSimulationInputData& InputData)
diff --git a/Engine/Plugins/Experimental/ChaosMover/Source/ChaosMover/Private/ChaosMoverStateMachine.cpp b/Engine/Plugins/Experimental/ChaosMover/Source/ChaosMover/Private/ChaosMoverStateMachine.cpp
--- a/Engine/Plugins/Experimental/ChaosMover/Source/ChaosMover/Private/ChaosMoverStateMachine.cpp
+++ b/Engine/Plugins/Experimental/ChaosMover/Source/ChaosMover/Private/ChaosMoverStateMachine.cpp
@@ -9,6 +9,7 @@
 #include "ChaosMover/ChaosMoverSimulation.h"
 #include "MoveLibrary/MovementMixer.h"
 #include "MoveLibrary/MovementUtils.h"
+#include "MoveLibrary/RollbackBlackboard.h"
 #include "MovementModeTransition.h"
 #include "MoverComponent.h"
 #include "MoverDeveloperSettings.h"
@@ -791,6 +792,18 @@ namespace UE::ChaosMover
 							Modes[PreviousModeName]->Deactivate();
 						}
 
+						// Track last mode change in the blackboard
+						URollbackBlackboard_InternalWrapper* RollbackBlackboard = Simulation->GetRollbackBlackboard_Mutable();
+
+						FMovementModeChangeRecord ModeChangeRecord;
+						ModeChangeRecord.ModeName = CurrentModeName;
+						ModeChangeRecord.PrevModeName = PreviousModeName;
+						ModeChangeRecord.Frame = InternalServerFrame;
+						ModeChangeRecord.SimTimeMs = InternalSimTimeMs;
+
+						RollbackBlackboard->TrySet(CommonBlackboard::LastModeChangeRecord, ModeChangeRecord);
+
+
 						Modes[CurrentModeName]->Activate();
 
 						// Notify the simulation of a mode change so it can react accordingly
diff --git a/Engine/Plugins/Experimental/Mover/Source/Mover/Private/MoveLibrary/RollbackBlackboard.cpp b/Engine/Plugins/Experimental/Mover/Source/Mover/Private/MoveLibrary/RollbackBlackboard.cpp
--- a/Engine/Plugins/Experimental/Mover/Source/Mover/Private/MoveLibrary/RollbackBlackboard.cpp
+++ b/Engine/Plugins/Experimental/Mover/Source/Mover/Private/MoveLibrary/RollbackBlackboard.cpp
@@ -95,6 +95,7 @@ void URollbackBlackboard::BeginSimulationFrame(const FMoverTimeStep& PendingTime
 	check(!bIsSimulationInProgress && !bIsRollbackInProgress);
 	InProgressSimFrameThreadId = FPlatformTLS::GetCurrentThreadId();
 	bIsSimulationInProgress = true;
+	bIsResimulating = PendingTimeStep.bIsResimulating;
 
 	InProgressSimTimeStamp.TimeMs = PendingTimeStep.BaseSimTimeMs;
 	InProgressSimTimeStamp.Frame = PendingTimeStep.ServerFrame;
@@ -103,7 +104,7 @@ void URollbackBlackboard::BeginSimulationFrame(const FMoverTimeStep& PendingTime
 void URollbackBlackboard::EndSimulationFrame()
 {
 	check(bIsSimulationInProgress && (InProgressSimFrameThreadId == FPlatformTLS::GetCurrentThreadId()));
-	bIsSimulationInProgress = false;
+	bIsSimulationInProgress = bIsResimulating = false;
 
 	// TODO: need some kind of lock so we wait on any in-progress operations before advancing the CurrentSimTimeStamp
 	CurrentSimTimeStamp = InProgressSimTimeStamp;
@@ -118,17 +119,17 @@ void URollbackBlackboard::EndSimulationFrame()
 }
 
 
-void URollbackBlackboard::BeginRollback()
+void URollbackBlackboard::BeginRollback(const FMoverTimeStep& NewBaseTimeStep)
 {
 	check(!bIsSimulationInProgress && !bIsRollbackInProgress);
 	InRollbackThreadId = FPlatformTLS::GetCurrentThreadId();
 	bIsRollbackInProgress = true;
-}
 
-void URollbackBlackboard::RollBack(const FMoverTimeStep& NewBaseTimeStep)
-{
-	check(!bIsSimulationInProgress && bIsRollbackInProgress);
-	
+	UE_LOG(LogMover, Verbose, TEXT("Blackboard begin rollback. From Sim F %i / T %.3f -> F %i / T %.3f"), 
+		CurrentSimTimeStamp.Frame, CurrentSimTimeStamp.TimeMs, NewBaseTimeStep.ServerFrame, NewBaseTimeStep.BaseSimTimeMs);
+
+	// TODO: Need locking mechanism
+
 	const EntryTimeStamp NewBaseTimeStamp((double)NewBaseTimeStep.BaseSimTimeMs, NewBaseTimeStep.ServerFrame);
 
 	for (const TPair<FName, TUniquePtr<BlackboardEntryBase>>& KVP : EntryMap)
diff --git a/Engine/Plugins/Experimental/Mover/Source/Mover/Private/MovementModeStateMachine.cpp b/Engine/Plugins/Experimental/Mover/Source/Mover/Private/MovementModeStateMachine.cpp
--- a/Engine/Plugins/Experimental/Mover/Source/Mover/Private/MovementModeStateMachine.cpp
+++ b/Engine/Plugins/Experimental/Mover/Source/Mover/Private/MovementModeStateMachine.cpp
@@ -697,8 +697,6 @@ void UMovementModeStateMachine::AdvanceToNextMode()
 			Modes[PreviousModeName]->Deactivate();
 		}
 
-		Modes[CurrentModeName]->Activate();
-
 		// signal movement mode change event
 		const UMoverComponent* MoverComp = CastChecked<UMoverComponent>(GetOuter());
 
@@ -706,11 +704,14 @@ void UMovementModeStateMachine::AdvanceToNextMode()
 
 		FMovementModeChangeRecord ModeChangeRecord;
 		ModeChangeRecord.ModeName = CurrentModeName;
+		ModeChangeRecord.PrevModeName = PreviousModeName;
 		ModeChangeRecord.Frame = CurrentBaseTimeStep.ServerFrame;
 		ModeChangeRecord.SimTimeMs = CurrentBaseTimeStep.BaseSimTimeMs;
 
 		RollbackBlackboard->TrySet(CommonBlackboard::LastModeChangeRecord, ModeChangeRecord);
 
+		Modes[CurrentModeName]->Activate();
+
 		MoverComp->OnMovementModeChanged.Broadcast(PreviousModeName, NextModeName);
 	}
 }
diff --git a/Engine/Plugins/Experimental/Mover/Source/Mover/Private/MoverComponent.cpp b/Engine/Plugins/Experimental/Mover/Source/Mover/Private/MoverComponent.cpp
--- a/Engine/Plugins/Experimental/Mover/Source/Mover/Private/MoverComponent.cpp
+++ b/Engine/Plugins/Experimental/Mover/Source/Mover/Private/MoverComponent.cpp
@@ -99,13 +99,13 @@ void UMoverComponent::InitializeComponent()
 		RollbackBlackboard_InternalWrapper->Init(*RollbackBlackboard);
 
 		// create any internal entries
-		static URollbackBlackboard::EntrySettings ModeActivationSettings;
-		ModeActivationSettings.SizingPolicy = EBlackboardSizingPolicy::FixedDeclaredSize;
-		ModeActivationSettings.FixedSize = 4;
-		ModeActivationSettings.PersistencePolicy = EBlackboardPersistencePolicy::Forever;
-		ModeActivationSettings.RollbackPolicy = EBlackboardRollbackPolicy::InvalidatedOnRollback;
+		static URollbackBlackboard::EntrySettings ModeChangeRecordSettings;
+		ModeChangeRecordSettings.SizingPolicy = EBlackboardSizingPolicy::FixedDeclaredSize;
+		ModeChangeRecordSettings.FixedSize = 4;
+		ModeChangeRecordSettings.PersistencePolicy = EBlackboardPersistencePolicy::Forever;
+		ModeChangeRecordSettings.RollbackPolicy = EBlackboardRollbackPolicy::InvalidatedOnRollback;
 
-		RollbackBlackboard->CreateEntry<FMovementModeChangeRecord>(CommonBlackboard::LastModeChangeRecord, ModeActivationSettings);
+		RollbackBlackboard->CreateEntry<FMovementModeChangeRecord>(CommonBlackboard::LastModeChangeRecord, ModeChangeRecordSettings);
 
 
 		FindDefaultUpdatedComponent();
@@ -1079,9 +1079,7 @@ void UMoverComponent::OnSimulationRollback(const FMoverSyncState* SyncState, con
 {
 	SimBlackboard->Invalidate(EInvalidationReason::Rollback);
 
-	RollbackBlackboard_InternalWrapper->BeginRollback();
-
-	RollbackBlackboard_InternalWrapper->RollBack(NewBaseTimeStep);
+	RollbackBlackboard_InternalWrapper->BeginRollback(NewBaseTimeStep);
 
 	ModeFSM->OnSimulationRollback(SyncState, AuxState, NewBaseTimeStep);
 
diff --git a/Engine/Plugins/Experimental/Mover/Source/Mover/Private/MoverSimulation.cpp b/Engine/Plugins/Experimental/Mover/Source/Mover/Private/MoverSimulation.cpp
--- a/Engine/Plugins/Experimental/Mover/Source/Mover/Private/MoverSimulation.cpp
+++ b/Engine/Plugins/Experimental/Mover/Source/Mover/Private/MoverSimulation.cpp
@@ -3,6 +3,7 @@
 #include "MoverSimulation.h"
 
 #include "MoveLibrary/MoverBlackboard.h"
+#include "MoveLibrary/RollbackBlackboard.h"
 
 #include UE_INLINE_GENERATED_CPP_BY_NAME(MoverSimulation)
 
@@ -19,4 +20,19 @@ const UMoverBlackboard* UMoverSimulation::GetBlackboard() const
 UMoverBlackboard* UMoverSimulation::GetBlackboard_Mutable()
 {
 	return Blackboard;
+}
+
+const URollbackBlackboard_InternalWrapper* UMoverSimulation::GetRollbackBlackboard() const
+{
+	return RollbackBlackboard;
+}
+
+URollbackBlackboard_InternalWrapper* UMoverSimulation::GetRollbackBlackboard_Mutable()
+{
+	return RollbackBlackboard;
+}
+
+void UMoverSimulation::SetRollbackBlackboard(URollbackBlackboard_InternalWrapper* RollbackSimBlackboard)
+{
+	RollbackBlackboard = RollbackSimBlackboard;
 }
\ No newline at end of file
diff --git a/Engine/Plugins/Experimental/Mover/Source/Mover/Public/MoveLibrary/RollbackBlackboard.h b/Engine/Plugins/Experimental/Mover/Source/Mover/Public/MoveLibrary/RollbackBlackboard.h
--- a/Engine/Plugins/Experimental/Mover/Source/Mover/Public/MoveLibrary/RollbackBlackboard.h
+++ b/Engine/Plugins/Experimental/Mover/Source/Mover/Public/MoveLibrary/RollbackBlackboard.h
@@ -387,8 +387,7 @@ private:
 	void BeginSimulationFrame(const FMoverTimeStep& PendingTimeStep);
 	void EndSimulationFrame();
 
-	void BeginRollback();
-	void RollBack(const FMoverTimeStep& NewBaseTimeStep);	// NewBaseTimeStep represents the current time and the PENDING frame that's next to be resimulated
+	void BeginRollback(const FMoverTimeStep& NewBaseTimeStep);	// NewBaseTimeStep represents the current time and the PENDING frame that's next to be resimulated
 	void EndRollback();
 
 
@@ -398,6 +397,7 @@ private:
 
 
 	bool bIsSimulationInProgress = false;
+	bool bIsResimulating = false;
 	uint32 InProgressSimFrameThreadId;
 
 	bool bIsRollbackInProgress = false;
@@ -457,27 +457,22 @@ public:
 	}
 
 
-	void BeginSimulationFrame(const FMoverTimeStep& PendingTimeStep)
+	UE_API void BeginSimulationFrame(const FMoverTimeStep& PendingTimeStep)
 	{
 		Blackboard->BeginSimulationFrame(PendingTimeStep);
 	}
 
-	void EndSimulationFrame()
+	UE_API void EndSimulationFrame()
 	{
 		Blackboard->EndSimulationFrame();
 	}
 
-	void BeginRollback()
+	UE_API void BeginRollback(const FMoverTimeStep& NewBaseTimeStep)
 	{
-		Blackboard->BeginRollback();
+		Blackboard->BeginRollback(NewBaseTimeStep);
 	}
 
-	void RollBack(const FMoverTimeStep& NewBaseTimeStep)
-	{
-		Blackboard->RollBack(NewBaseTimeStep);
-	}
-
-	void EndRollback()
+	UE_API void EndRollback()
 	{
 		Blackboard->EndRollback();
 	}
diff --git a/Engine/Plugins/Experimental/Mover/Source/Mover/Public/MoverSimulation.h b/Engine/Plugins/Experimental/Mover/Source/Mover/Public/MoverSimulation.h
--- a/Engine/Plugins/Experimental/Mover/Source/Mover/Public/MoverSimulation.h
+++ b/Engine/Plugins/Experimental/Mover/Source/Mover/Public/MoverSimulation.h
@@ -10,6 +10,7 @@
 
 class UMoverBlackboard;
 struct FMoverSyncState;
+class URollbackBlackboard_InternalWrapper;
 
 /**
 * WIP Base class for a Mover simulation.
@@ -24,12 +25,21 @@ class UMoverSimulation : public UObject
 public:
 	UE_API UMoverSimulation();
 
+	// Warning: the regular blackboard will be fully replaced by the rollback blackboard in the future
 	UFUNCTION(BlueprintPure, Category = Mover)
 	UE_API const UMoverBlackboard* GetBlackboard() const;
 
+	// Warning: the regular blackboard will be fully replaced by the rollback blackboard in the future
 	UFUNCTION(BlueprintPure, Category = Mover)
 	UE_API UMoverBlackboard* GetBlackboard_Mutable();
 
+	UFUNCTION(BlueprintPure, Category = Mover)
+	UE_API const URollbackBlackboard_InternalWrapper* GetRollbackBlackboard() const;
+
+	UFUNCTION(BlueprintPure, Category = Mover)
+	UE_API URollbackBlackboard_InternalWrapper* GetRollbackBlackboard_Mutable();
+
+
 	/**
 	* Attempt to teleport to TargetTransform. The teleport is not guaranteed to happen. This function is meant to be called by an instant movement effect as part of its effect application.
 	* If it succeeds a FTeleportSucceededEventData will be emitted, if it fails a FTeleportFailedEventData will be sent.
@@ -41,9 +51,17 @@ public:
 	UFUNCTION(BlueprintCallable, Category = Mover)
 	virtual void AttemptTeleport(const FMoverTimeStep& TimeStep, const FTransform& TargetTransform, bool bUseActorRotation, FMoverSyncState& OutputState) {}
 
+
+	// Used during initialization only
+	UE_API void SetRollbackBlackboard(URollbackBlackboard_InternalWrapper* RollbackSimBlackboard);
+
 protected:
 	UPROPERTY(Transient)
 	TObjectPtr<UMoverBlackboard> Blackboard = nullptr;
+
+	UPROPERTY(Transient)
+	TObjectPtr<URollbackBlackboard_InternalWrapper> RollbackBlackboard = nullptr;
+
 };
 
 #undef UE_API
diff --git a/Engine/Plugins/Experimental/Mover/Source/Mover/Public/MoverTypes.h b/Engine/Plugins/Experimental/Mover/Source/Mover/Public/MoverTypes.h
--- a/Engine/Plugins/Experimental/Mover/Source/Mover/Public/MoverTypes.h
+++ b/Engine/Plugins/Experimental/Mover/Source/Mover/Public/MoverTypes.h
@@ -374,6 +374,7 @@ struct FMoverDataPersistence
 struct FMovementModeChangeRecord
 {
 	FName ModeName;
+	FName PrevModeName;
 	uint32 Frame;
 	double SimTimeMs;
 };
