Fix missing symbols due to missing some exports

#rnx

diff --git a/Engine/Plugins/Experimental/Mover/Source/Mover/Public/MoveLibrary/RollbackBlackboard.h b/Engine/Plugins/Experimental/Mover/Source/Mover/Public/MoveLibrary/RollbackBlackboard.h
--- a/Engine/Plugins/Experimental/Mover/Source/Mover/Public/MoveLibrary/RollbackBlackboard.h
+++ b/Engine/Plugins/Experimental/Mover/Source/Mover/Public/MoveLibrary/RollbackBlackboard.h
@@ -384,11 +384,11 @@ private:
 
 
 private:
-	void BeginSimulationFrame(const FMoverTimeStep& PendingTimeStep);
-	void EndSimulationFrame();
+	UE_API void BeginSimulationFrame(const FMoverTimeStep& PendingTimeStep);
+	UE_API void EndSimulationFrame();
 
-	void BeginRollback(const FMoverTimeStep& NewBaseTimeStep);	// NewBaseTimeStep represents the current time and the PENDING frame that's next to be resimulated
-	void EndRollback();
+	UE_API void BeginRollback(const FMoverTimeStep& NewBaseTimeStep);	// NewBaseTimeStep represents the current time and the PENDING frame that's next to be resimulated
+	UE_API void EndRollback();
 
 
 private:
