Moving exported function definitions outside of header to fix missing symbols issue in builds

#rnx

diff --git a/Engine/Plugins/Experimental/Mover/Source/Mover/Private/MoveLibrary/RollbackBlackboard.cpp b/Engine/Plugins/Experimental/Mover/Source/Mover/Private/MoveLibrary/RollbackBlackboard.cpp
--- a/Engine/Plugins/Experimental/Mover/Source/Mover/Private/MoveLibrary/RollbackBlackboard.cpp
+++ b/Engine/Plugins/Experimental/Mover/Source/Mover/Private/MoveLibrary/RollbackBlackboard.cpp
@@ -146,3 +146,24 @@ void URollbackBlackboard::EndRollback()
 	check (bIsRollbackInProgress && (InRollbackThreadId == FPlatformTLS::GetCurrentThreadId()));
 	bIsRollbackInProgress = false;
 }
+
+
+void URollbackBlackboard_InternalWrapper::BeginSimulationFrame(const FMoverTimeStep& PendingTimeStep)
+{
+	Blackboard->BeginSimulationFrame(PendingTimeStep);
+}
+
+void URollbackBlackboard_InternalWrapper::EndSimulationFrame()
+{
+	Blackboard->EndSimulationFrame();
+}
+
+void URollbackBlackboard_InternalWrapper::BeginRollback(const FMoverTimeStep& NewBaseTimeStep)
+{
+	Blackboard->BeginRollback(NewBaseTimeStep);
+}
+
+void URollbackBlackboard_InternalWrapper::EndRollback()
+{
+	Blackboard->EndRollback();
+}
diff --git a/Engine/Plugins/Experimental/Mover/Source/Mover/Public/MoveLibrary/RollbackBlackboard.h b/Engine/Plugins/Experimental/Mover/Source/Mover/Public/MoveLibrary/RollbackBlackboard.h
--- a/Engine/Plugins/Experimental/Mover/Source/Mover/Public/MoveLibrary/RollbackBlackboard.h
+++ b/Engine/Plugins/Experimental/Mover/Source/Mover/Public/MoveLibrary/RollbackBlackboard.h
@@ -456,26 +456,11 @@ public:
 		return Blackboard->TryGetRef_Internal<EntryT>(ObjName);
 	}
 
+	UE_API void BeginSimulationFrame(const FMoverTimeStep& PendingTimeStep);
+	UE_API void EndSimulationFrame();
 
-	UE_API void BeginSimulationFrame(const FMoverTimeStep& PendingTimeStep)
-	{
-		Blackboard->BeginSimulationFrame(PendingTimeStep);
-	}
-
-	UE_API void EndSimulationFrame()
-	{
-		Blackboard->EndSimulationFrame();
-	}
-
-	UE_API void BeginRollback(const FMoverTimeStep& NewBaseTimeStep)
-	{
-		Blackboard->BeginRollback(NewBaseTimeStep);
-	}
-
-	UE_API void EndRollback()
-	{
-		Blackboard->EndRollback();
-	}
+	UE_API void BeginRollback(const FMoverTimeStep& NewBaseTimeStep);
+	UE_API void EndRollback();
 
 private: 
 	UPROPERTY()
