Avoid possible null refs when using blackboard

#rnx

diff --git a/Engine/Plugins/Experimental/ChaosMover/Source/ChaosMover/Private/ChaosMoverStateMachine.cpp b/Engine/Plugins/Experimental/ChaosMover/Source/ChaosMover/Private/ChaosMoverStateMachine.cpp
--- a/Engine/Plugins/Experimental/ChaosMover/Source/ChaosMover/Private/ChaosMoverStateMachine.cpp
+++ b/Engine/Plugins/Experimental/ChaosMover/Source/ChaosMover/Private/ChaosMoverStateMachine.cpp
@@ -793,16 +793,16 @@ namespace UE::ChaosMover
 						}
 
 						// Track last mode change in the blackboard
-						URollbackBlackboard_InternalWrapper* RollbackBlackboard = Simulation->GetRollbackBlackboard_Mutable();
-
-						FMovementModeChangeRecord ModeChangeRecord;
-						ModeChangeRecord.ModeName = CurrentModeName;
-						ModeChangeRecord.PrevModeName = PreviousModeName;
-						ModeChangeRecord.Frame = InternalServerFrame;
-						ModeChangeRecord.SimTimeMs = InternalSimTimeMs;
-
-						RollbackBlackboard->TrySet(CommonBlackboard::LastModeChangeRecord, ModeChangeRecord);
+						if (URollbackBlackboard_InternalWrapper* RollbackBlackboard = Simulation->GetRollbackBlackboard_Mutable())
+						{
+							FMovementModeChangeRecord ModeChangeRecord;
+							ModeChangeRecord.ModeName = CurrentModeName;
+							ModeChangeRecord.PrevModeName = PreviousModeName;
+							ModeChangeRecord.Frame = InternalServerFrame;
+							ModeChangeRecord.SimTimeMs = InternalSimTimeMs;
 
+							RollbackBlackboard->TrySet(CommonBlackboard::LastModeChangeRecord, ModeChangeRecord);
+						}
 
 						Modes[CurrentModeName]->Activate();
 
diff --git a/Engine/Plugins/Experimental/Mover/Source/Mover/Private/MovementModeStateMachine.cpp b/Engine/Plugins/Experimental/Mover/Source/Mover/Private/MovementModeStateMachine.cpp
--- a/Engine/Plugins/Experimental/Mover/Source/Mover/Private/MovementModeStateMachine.cpp
+++ b/Engine/Plugins/Experimental/Mover/Source/Mover/Private/MovementModeStateMachine.cpp
@@ -700,15 +700,16 @@ void UMovementModeStateMachine::AdvanceToNextMode()
 		// signal movement mode change event
 		const UMoverComponent* MoverComp = CastChecked<UMoverComponent>(GetOuter());
 
-		URollbackBlackboard_InternalWrapper* RollbackBlackboard = MoverComp->GetRollbackBlackboard_Internal();
-
-		FMovementModeChangeRecord ModeChangeRecord;
-		ModeChangeRecord.ModeName = CurrentModeName;
-		ModeChangeRecord.PrevModeName = PreviousModeName;
-		ModeChangeRecord.Frame = CurrentBaseTimeStep.ServerFrame;
-		ModeChangeRecord.SimTimeMs = CurrentBaseTimeStep.BaseSimTimeMs;
+		if (URollbackBlackboard_InternalWrapper* RollbackBlackboard = MoverComp->GetRollbackBlackboard_Internal())
+		{
+			FMovementModeChangeRecord ModeChangeRecord;
+			ModeChangeRecord.ModeName = CurrentModeName;
+			ModeChangeRecord.PrevModeName = PreviousModeName;
+			ModeChangeRecord.Frame = CurrentBaseTimeStep.ServerFrame;
+			ModeChangeRecord.SimTimeMs = CurrentBaseTimeStep.BaseSimTimeMs;
 
-		RollbackBlackboard->TrySet(CommonBlackboard::LastModeChangeRecord, ModeChangeRecord);
+			RollbackBlackboard->TrySet(CommonBlackboard::LastModeChangeRecord, ModeChangeRecord);
+		}
 
 		Modes[CurrentModeName]->Activate();
 
