Mover: Searching modifiers by handle with an invalid handle will return nullptr rather than active modifiers with invalid handles that haven't been fixed up yet.

diff --git a/Engine/Plugins/Experimental/Mover/Source/Mover/Private/MoverComponent.cpp b/Engine/Plugins/Experimental/Mover/Source/Mover/Private/MoverComponent.cpp
--- a/Engine/Plugins/Experimental/Mover/Source/Mover/Private/MoverComponent.cpp
+++ b/Engine/Plugins/Experimental/Mover/Source/Mover/Private/MoverComponent.cpp
@@ -724,6 +724,11 @@ bool UMoverComponent::IsModifierActiveOrQueued(const FMovementModifierHandle& Mo
 
 const FMovementModifierBase* UMoverComponent::FindMovementModifier(const FMovementModifierHandle& ModifierHandle) const
 {
+	if (!ModifierHandle.IsValid())
+	{
+		return nullptr;
+	}
+	
 	const FMoverSyncState& CachedSyncState = MoverSyncStateDoubleBuffer.GetReadable();
 	
 	// Check active modifiers for modifier handle
