Mover: making input cmd available on interpolated sim proxies even though they don't actually simulate, if available. Primarily applies to NPP Independent ticking mode.

diff --git a/Engine/Plugins/Experimental/Mover/Source/Mover/Private/Backends/MoverNetworkPredictionLiaison.cpp b/Engine/Plugins/Experimental/Mover/Source/Mover/Private/Backends/MoverNetworkPredictionLiaison.cpp
--- a/Engine/Plugins/Experimental/Mover/Source/Mover/Private/Backends/MoverNetworkPredictionLiaison.cpp
+++ b/Engine/Plugins/Experimental/Mover/Source/Mover/Private/Backends/MoverNetworkPredictionLiaison.cpp
@@ -101,6 +101,12 @@ void UMoverNetworkPredictionLiaisonComponent::FinalizeFrame(const FMoverSyncStat
 	if (MoverComp->GetOwnerRole() == ROLE_SimulatedProxy && NetworkPredictionSettings.SimulatedProxyNetworkLOD == ENetworkLOD::Interpolated)
 	{
 		FMoverInputCmdContext InputCmd;
+
+		if (const FMoverInputCmdContext* PendingInputCmd = NetworkPredictionProxy.ReadInputCmd<FMoverInputCmdContext>())
+		{
+			InputCmd = *PendingInputCmd;
+		}
+
 		MoverComp->TickInterpolatedSimProxy(MoverComp->GetLastTimeStep(), InputCmd, MoverComp, MoverComp->GetSyncState(), *SyncState, *AuxState);
 	}
 	
diff --git a/Engine/Plugins/Experimental/Mover/Source/Mover/Private/MoverComponent.cpp b/Engine/Plugins/Experimental/Mover/Source/Mover/Private/MoverComponent.cpp
--- a/Engine/Plugins/Experimental/Mover/Source/Mover/Private/MoverComponent.cpp
+++ b/Engine/Plugins/Experimental/Mover/Source/Mover/Private/MoverComponent.cpp
@@ -421,6 +421,8 @@ void UMoverComponent::FinalizeSmoothingFrame(const FMoverSyncState* SyncState, c
 
 void UMoverComponent::TickInterpolatedSimProxy(const FMoverTimeStep& TimeStep, const FMoverInputCmdContext& InputCmd, UMoverComponent* MoverComp, const FMoverSyncState& CachedSyncState, const FMoverSyncState& SyncState, const FMoverAuxStateContext& AuxState)
 {
+	CachedLastUsedInputCmd = InputCmd;
+
 	TArray<TSharedPtr<FMovementModifierBase>> ModifiersToStart;
 	TArray<TSharedPtr<FMovementModifierBase>> ModifiersToEnd;
 
