Mover: Making MoveInput and MoveInputType BlueprintReadOnly, so setting this input is done through specific functions only. This allows us to properly apply some precision to input to make sure all clients/server have same value after NetSerialization

#virtualized

diff --git a/Engine/Plugins/Experimental/Mover/Source/Mover/Private/MoverDataModelTypes.cpp b/Engine/Plugins/Experimental/Mover/Source/Mover/Private/MoverDataModelTypes.cpp
--- a/Engine/Plugins/Experimental/Mover/Source/Mover/Private/MoverDataModelTypes.cpp
+++ b/Engine/Plugins/Experimental/Mover/Source/Mover/Private/MoverDataModelTypes.cpp
@@ -482,9 +482,14 @@ FVector FMoverDefaultSyncState::GetAngularVelocityDegrees_BaseSpace() const
 
 // UMoverDataModelBlueprintLibrary ///////////////////////////////////////////////////
 
-void UMoverDataModelBlueprintLibrary::SetMoveIntent(FCharacterDefaultInputs& Inputs, const FVector& WorldDirectionIntent)
+void UMoverDataModelBlueprintLibrary::SetDirectionalInput(FCharacterDefaultInputs& Inputs, const FVector& DirectionInput)
 {
-	Inputs.SetMoveInput(EMoveInputType::DirectionalIntent, WorldDirectionIntent);
+	Inputs.SetMoveInput(EMoveInputType::DirectionalIntent, DirectionInput);
+}
+
+void UMoverDataModelBlueprintLibrary::SetVelocityInput(FCharacterDefaultInputs& Inputs, const FVector& VelocityInput)
+{
+	Inputs.SetMoveInput(EMoveInputType::Velocity, VelocityInput);
 }
 
 FVector UMoverDataModelBlueprintLibrary::GetMoveDirectionIntentFromInputs(const FCharacterDefaultInputs& Inputs)
diff --git a/Engine/Plugins/Experimental/Mover/Source/Mover/Public/MoverDataModelTypes.h b/Engine/Plugins/Experimental/Mover/Source/Mover/Public/MoverDataModelTypes.h
--- a/Engine/Plugins/Experimental/Mover/Source/Mover/Public/MoverDataModelTypes.h
+++ b/Engine/Plugins/Experimental/Mover/Source/Mover/Public/MoverDataModelTypes.h
@@ -42,11 +42,14 @@ struct FCharacterDefaultInputs : public FMoverDataStructBase
 	EMoveInputType GetMoveInputType() const { return MoveInputType; }
 
 protected:
-	UPROPERTY(VisibleAnywhere, BlueprintReadWrite, Category = Mover)
+	UPROPERTY(VisibleAnywhere, BlueprintReadOnly, Category = Mover)
 	EMoveInputType MoveInputType;
 
-	// Representing the directional move input for this frame. Must be interpreted according to MoveInputType. Relative to MovementBase if set, world space otherwise. Will be truncated to match network serialization precision.
-	UPROPERTY(VisibleAnywhere, BlueprintReadWrite, Category = Mover)
+	/**
+	 * Representing the directional move input for this frame. Must be interpreted according to MoveInputType. Relative to MovementBase if set, world space otherwise. Will be truncated to match network serialization precision.
+	 * Note: Use SetDirectionalInput or SetVelocityInput to set MoveInput and MoveInputType
+	 */
+	UPROPERTY(VisibleAnywhere, BlueprintReadOnly, Category = Mover)
 	FVector MoveInput;
 
 public:
@@ -271,9 +274,13 @@ class UMoverDataModelBlueprintLibrary : public UBlueprintFunctionLibrary
 
 public:	// FCharacterDefaultInputs
 
-	/** Sets move inputs from worldspace intent, as a per-axis magnitude in the range [-1,1] Zero vector indicates intent to stop. */
+	/** Sets move input from a unit length vector representing directional intent */
+	UFUNCTION(BlueprintCallable, Category = Mover)
+	static UE_API void SetDirectionalInput(UPARAM(Ref) FCharacterDefaultInputs& Inputs, const FVector& DirectionInput);
+	
+	/** Sets move input from a vector representing desired velocity c/m^2 */
 	UFUNCTION(BlueprintCallable, Category = Mover)
-	static UE_API void SetMoveIntent(UPARAM(Ref) FCharacterDefaultInputs& Inputs, const FVector& WorldDirectionIntent);
+	static UE_API void SetVelocityInput(UPARAM(Ref) FCharacterDefaultInputs& Inputs, const FVector& VelocityInput = FVector::ZeroVector);
 
 	/** Returns the move direction intent, if any, in world space */
 	UFUNCTION(BlueprintCallable, Category = Mover)
