Rename variable in Mover's SimpleWalkingMode

#virtualized

diff --git a/Engine/Plugins/Experimental/Mover/Source/Mover/Public/DefaultMovementSet/Modes/SimpleWalkingMode.h b/Engine/Plugins/Experimental/Mover/Source/Mover/Public/DefaultMovementSet/Modes/SimpleWalkingMode.h
--- a/Engine/Plugins/Experimental/Mover/Source/Mover/Public/DefaultMovementSet/Modes/SimpleWalkingMode.h
+++ b/Engine/Plugins/Experimental/Mover/Source/Mover/Public/DefaultMovementSet/Modes/SimpleWalkingMode.h
@@ -23,7 +23,7 @@ public:
 	// Override this to make a simple walking mode
 	UFUNCTION(BlueprintNativeEvent, meta = (DisplayName = "Generate Simple Walk Move"))
 	UE_API void GenerateWalkMove(UPARAM(ref) FMoverTickStartData& StartState, float DeltaSeconds, const FVector& DesiredVelocity,
-	                                     const FQuat& DesiredFacing, const FQuat& CurrentFacing, UPARAM(ref) FVector& InAngularVelocityDegrees, UPARAM(ref) FVector& InOutVelocity);
+	                                     const FQuat& DesiredFacing, const FQuat& CurrentFacing, UPARAM(ref) FVector& InOutAngularVelocityDegrees, UPARAM(ref) FVector& InOutVelocity);
 
 	// If this value is greater or equal to 0, this will override the max speed read from the common legacy shared walk settings
 	UPROPERTY(EditAnywhere, BlueprintReadWrite, Category = "Mover|Walking Settings", meta = (ForceUnits = "cm/s"))
