PhysicsMover - Add warnings to old physics classes that they will be removed

#rnx

diff --git a/Engine/Plugins/Experimental/Mover/Source/Mover/Public/Backends/MoverNetworkPhysicsLiaison.h b/Engine/Plugins/Experimental/Mover/Source/Mover/Public/Backends/MoverNetworkPhysicsLiaison.h
--- a/Engine/Plugins/Experimental/Mover/Source/Mover/Public/Backends/MoverNetworkPhysicsLiaison.h
+++ b/Engine/Plugins/Experimental/Mover/Source/Mover/Public/Backends/MoverNetworkPhysicsLiaison.h
@@ -17,10 +17,12 @@ class UCommonLegacyMovementSettings;
 // UMoverNetworkPhysicsLiaisonComponent
 
 /**
+ * WARNING - This class will be removed. Please use UChaosMoverBackend instead
+ *
  * MoverNetworkPhysicsLiaisonComponent: This component acts as a middleman between an actor's Mover component and the chaos physics network prediction system.
  * This class is set on a Mover component as the "back end".
  */
-UCLASS(MinimalAPI)
+UCLASS(MinimalAPI, DisplayName = "DEPRECATED Mover Network Physics Liaison Component")
 class UMoverNetworkPhysicsLiaisonComponent : public UMoverNetworkPhysicsLiaisonComponentBase
 {
 	GENERATED_BODY()
diff --git a/Engine/Plugins/Experimental/Mover/Source/Mover/Public/Backends/MoverNetworkPhysicsLiaisonBase.h b/Engine/Plugins/Experimental/Mover/Source/Mover/Public/Backends/MoverNetworkPhysicsLiaisonBase.h
--- a/Engine/Plugins/Experimental/Mover/Source/Mover/Public/Backends/MoverNetworkPhysicsLiaisonBase.h
+++ b/Engine/Plugins/Experimental/Mover/Source/Mover/Public/Backends/MoverNetworkPhysicsLiaisonBase.h
@@ -176,6 +176,8 @@ public:
 // UMoverNetworkPhysicsLiaisonComponentBase
 
 /**
+ * WARNING - This class will be removed. Please use UChaosMoverBackend instead
+ *
  * Base class for liaison components to act as a middleman between an actor's Mover component and the Chaos physics network prediction system to move
  * the Mover's updated component on the physics thread (PT).
  * 
diff --git a/Engine/Plugins/Experimental/Mover/Source/Mover/Public/Backends/MoverPathedPhysicsLiaison.h b/Engine/Plugins/Experimental/Mover/Source/Mover/Public/Backends/MoverPathedPhysicsLiaison.h
--- a/Engine/Plugins/Experimental/Mover/Source/Mover/Public/Backends/MoverPathedPhysicsLiaison.h
+++ b/Engine/Plugins/Experimental/Mover/Source/Mover/Public/Backends/MoverPathedPhysicsLiaison.h
@@ -14,12 +14,15 @@
 class UPathedPhysicsMovementMode;
 class UPathedPhysicsMoverComponent;
 /**
+ * WARNING - This class will be removed. Please use UChaosMoverBackend along with UChaosPathedMovementControllerComponent
+ * to simulate moving platforms
+ * 
  * Liaison that works with path-following movement modes.
  * Establishes a joint constraint between the actual component to update and wherever it should be along the path.
  * This allows the platform's velocity and rotation to be affected by physics as it moves along the path, and the joint pulls it back where it should be (think spring).
  * Only compatible with UPathedPhysicsMovementModes, which are responsible for dictating how "loose" the joint between the platform and its ideal position is.
  */
-UCLASS(MinimalAPI, Within = PathedPhysicsMoverComponent)
+UCLASS(MinimalAPI, Within = PathedPhysicsMoverComponent, DisplayName = "DEPRECATED Mover Pathed Physics Liaison Component")
 class UMoverPathedPhysicsLiaisonComponent : public UMoverNetworkPhysicsLiaisonComponentBase
 {
 	GENERATED_BODY()
diff --git a/Engine/Plugins/Experimental/Mover/Source/Mover/Public/PhysicsMover/Modes/PhysicsDrivenFallingMode.h b/Engine/Plugins/Experimental/Mover/Source/Mover/Public/PhysicsMover/Modes/PhysicsDrivenFallingMode.h
--- a/Engine/Plugins/Experimental/Mover/Source/Mover/Public/PhysicsMover/Modes/PhysicsDrivenFallingMode.h
+++ b/Engine/Plugins/Experimental/Mover/Source/Mover/Public/PhysicsMover/Modes/PhysicsDrivenFallingMode.h
@@ -14,9 +14,11 @@
 namespace Chaos { class FCharacterGroundConstraint; }
 
 /**
+ * WARNING - This class will be removed. Please use UChaosFallingMode instead
+ *
  * PhysicsDrivenFallingMode: Override base kinematic falling mode for physics based motion.
  */
-UCLASS(MinimalAPI, Blueprintable, BlueprintType)
+UCLASS(MinimalAPI, Blueprintable, BlueprintType, DisplayName = "DEPRECATED Physics Driven Falling Mode")
 class UPhysicsDrivenFallingMode : public UFallingMode, public IPhysicsCharacterMovementModeInterface
 {
 	GENERATED_UCLASS_BODY()
diff --git a/Engine/Plugins/Experimental/Mover/Source/Mover/Public/PhysicsMover/Modes/PhysicsDrivenFlyingMode.h b/Engine/Plugins/Experimental/Mover/Source/Mover/Public/PhysicsMover/Modes/PhysicsDrivenFlyingMode.h
--- a/Engine/Plugins/Experimental/Mover/Source/Mover/Public/PhysicsMover/Modes/PhysicsDrivenFlyingMode.h
+++ b/Engine/Plugins/Experimental/Mover/Source/Mover/Public/PhysicsMover/Modes/PhysicsDrivenFlyingMode.h
@@ -14,9 +14,11 @@
 namespace Chaos { class FCharacterGroundConstraint; }
 
 /**
+ * WARNING - This class will be removed. Please use UChaosFlyingMode instead
+ *
  * PhysicsDrivenFlyingMode: Override base kinematic flying mode for physics based motion.
  */
-UCLASS(MinimalAPI, Blueprintable, BlueprintType)
+UCLASS(MinimalAPI, Blueprintable, BlueprintType, DisplayName = "DEPRECATED Physics Driven Flying Mode")
 class UPhysicsDrivenFlyingMode : public UFlyingMode, public IPhysicsCharacterMovementModeInterface
 {
 	GENERATED_UCLASS_BODY()
diff --git a/Engine/Plugins/Experimental/Mover/Source/Mover/Public/PhysicsMover/Modes/PhysicsDrivenSwimmingMode.h b/Engine/Plugins/Experimental/Mover/Source/Mover/Public/PhysicsMover/Modes/PhysicsDrivenSwimmingMode.h
--- a/Engine/Plugins/Experimental/Mover/Source/Mover/Public/PhysicsMover/Modes/PhysicsDrivenSwimmingMode.h
+++ b/Engine/Plugins/Experimental/Mover/Source/Mover/Public/PhysicsMover/Modes/PhysicsDrivenSwimmingMode.h
@@ -14,9 +14,11 @@
 namespace Chaos { class FCharacterGroundConstraint; }
 
 /**
+ * WARNING - This class will be removed. Please use UChaosSwimmingMode instead
+ *
  * PhysicsDrivenSwimmingMode: Override base kinematic Swimming mode for physics based motion.
  */
-UCLASS(MinimalAPI, Blueprintable, BlueprintType)
+UCLASS(MinimalAPI, Blueprintable, BlueprintType, DisplayName = "DEPRECATED Physics Driven Swimming Mode")
 class UPhysicsDrivenSwimmingMode : public USwimmingMode, public IPhysicsCharacterMovementModeInterface
 {
 	GENERATED_UCLASS_BODY()
diff --git a/Engine/Plugins/Experimental/Mover/Source/Mover/Public/PhysicsMover/Modes/PhysicsDrivenWalkingMode.h b/Engine/Plugins/Experimental/Mover/Source/Mover/Public/PhysicsMover/Modes/PhysicsDrivenWalkingMode.h
--- a/Engine/Plugins/Experimental/Mover/Source/Mover/Public/PhysicsMover/Modes/PhysicsDrivenWalkingMode.h
+++ b/Engine/Plugins/Experimental/Mover/Source/Mover/Public/PhysicsMover/Modes/PhysicsDrivenWalkingMode.h
@@ -15,9 +15,11 @@ namespace Chaos { class FCharacterGroundConstraint; }
 struct FWaterCheckResult;
 
 /**
+ * WARNING - This class will be removed. Please use UChaosWalkingMode instead
+ *
  * PhysicsDrivenWalkingMode: Override base kinematic walking mode for physics based motion.
  */
-UCLASS(MinimalAPI, Blueprintable, BlueprintType)
+UCLASS(MinimalAPI, Blueprintable, BlueprintType, DisplayName="DEPRECATED Physics Driven Walking Mode")
 class UPhysicsDrivenWalkingMode : public UWalkingMode, public IPhysicsCharacterMovementModeInterface
 {
 	GENERATED_UCLASS_BODY()
diff --git a/Engine/Plugins/Experimental/Mover/Source/Mover/Public/PhysicsMover/PhysicsCharacterMoverComponent.h b/Engine/Plugins/Experimental/Mover/Source/Mover/Public/PhysicsMover/PhysicsCharacterMoverComponent.h
--- a/Engine/Plugins/Experimental/Mover/Source/Mover/Public/PhysicsMover/PhysicsCharacterMoverComponent.h
+++ b/Engine/Plugins/Experimental/Mover/Source/Mover/Public/PhysicsMover/PhysicsCharacterMoverComponent.h
@@ -7,7 +7,10 @@
 
 #define UE_API MOVER_API
 
-
+/*
+* WARNING - This class will be removed.Please use UChaosCharacterMoverComponent instead
+*
+*/
 UCLASS(MinimalAPI, BlueprintType, Blueprintable, meta = (BlueprintSpawnableComponent))
 class UPhysicsCharacterMoverComponent : public UCharacterMoverComponent
 {
diff --git a/Engine/Plugins/Experimental/Mover/Source/Mover/Public/PhysicsMover/Transitions/PhysicsJumpCheck.h b/Engine/Plugins/Experimental/Mover/Source/Mover/Public/PhysicsMover/Transitions/PhysicsJumpCheck.h
--- a/Engine/Plugins/Experimental/Mover/Source/Mover/Public/PhysicsMover/Transitions/PhysicsJumpCheck.h
+++ b/Engine/Plugins/Experimental/Mover/Source/Mover/Public/PhysicsMover/Transitions/PhysicsJumpCheck.h
@@ -8,9 +8,11 @@
 #define UE_API MOVER_API
 
 /**
+ * WARNING - This class will be removed. Please use UChaosCharacterJumpCheck instead
+ *
  * Transition that handles jumping based on input for a physics-based character
  */
-UCLASS(MinimalAPI, Blueprintable, BlueprintType)
+UCLASS(MinimalAPI, Blueprintable, BlueprintType, DisplayName = "DEPRECATED Physics Jump Check")
 class UPhysicsJumpCheck : public UBaseMovementModeTransition
 {
 	GENERATED_UCLASS_BODY()
diff --git a/Engine/Plugins/Experimental/Mover/Source/Mover/Public/PhysicsMover/Transitions/PhysicsLaunchCheck.h b/Engine/Plugins/Experimental/Mover/Source/Mover/Public/PhysicsMover/Transitions/PhysicsLaunchCheck.h
--- a/Engine/Plugins/Experimental/Mover/Source/Mover/Public/PhysicsMover/Transitions/PhysicsLaunchCheck.h
+++ b/Engine/Plugins/Experimental/Mover/Source/Mover/Public/PhysicsMover/Transitions/PhysicsLaunchCheck.h
@@ -8,9 +8,11 @@
 #define UE_API MOVER_API
 
 /**
+ * WARNING - This class will be removed. Please use UChaosCharacterLandingCheck instead
+ *
  * Transition that handles launching based on input for a physics-based character
  */
-UCLASS(MinimalAPI, Blueprintable, BlueprintType)
+UCLASS(MinimalAPI, Blueprintable, BlueprintType, DisplayName = "DEPRECATED Physics Launch Check")
 class UPhysicsLaunchCheck : public UBaseMovementModeTransition
 {
 	GENERATED_BODY()
