Mover: Adding some tolerance settings to SmoothWalkingState reconcile logic to avoid unneeded corrections

diff --git a/Engine/Plugins/Experimental/Mover/Source/Mover/Private/DefaultMovementSet/Modes/SmoothWalkingState.cpp b/Engine/Plugins/Experimental/Mover/Source/Mover/Private/DefaultMovementSet/Modes/SmoothWalkingState.cpp
--- a/Engine/Plugins/Experimental/Mover/Source/Mover/Private/DefaultMovementSet/Modes/SmoothWalkingState.cpp
+++ b/Engine/Plugins/Experimental/Mover/Source/Mover/Private/DefaultMovementSet/Modes/SmoothWalkingState.cpp
@@ -4,6 +4,14 @@
 
 #include UE_INLINE_GENERATED_CPP_BY_NAME(SmoothWalkingState)
 
+namespace SmoothWalkingStateErrorTolerance
+{
+	constexpr float VelocityErrorTolerance = 10.f;
+	constexpr float AngularVelocityErrorTolerance = 10.f;
+	constexpr float AccelerationErrorTolerance = 50.f;
+	constexpr float FacingDegreeErrorTolerance = 10.0f;
+}
+
 UScriptStruct* FSmoothWalkingState::GetScriptStruct() const
 { 
 	return StaticStruct(); 
@@ -44,11 +52,11 @@ bool FSmoothWalkingState::ShouldReconcile(const FMoverDataStructBase& AuthorityS
 {
 	const FSmoothWalkingState* AuthoritySpringState = static_cast<const FSmoothWalkingState*>(&AuthorityState);
 
-	return (!(SpringVelocity - AuthoritySpringState->SpringVelocity).IsNearlyZero() ||
-			!(SpringAcceleration - AuthoritySpringState->SpringAcceleration).IsNearlyZero() ||
-			!(IntermediateVelocity - AuthoritySpringState->IntermediateVelocity).IsNearlyZero() ||
-			 (IntermediateFacing.AngularDistance(AuthoritySpringState->IntermediateFacing) > UE_SMALL_NUMBER ||
-			!(IntermediateAngularVelocity - AuthoritySpringState->IntermediateAngularVelocity).IsNearlyZero()));
+	return (!(SpringVelocity - AuthoritySpringState->SpringVelocity).IsNearlyZero(SmoothWalkingStateErrorTolerance::VelocityErrorTolerance) ||
+			!(SpringAcceleration - AuthoritySpringState->SpringAcceleration).IsNearlyZero(SmoothWalkingStateErrorTolerance::AccelerationErrorTolerance) ||
+			!(IntermediateVelocity - AuthoritySpringState->IntermediateVelocity).IsNearlyZero(SmoothWalkingStateErrorTolerance::VelocityErrorTolerance) ||
+			 (IntermediateFacing.AngularDistance(AuthoritySpringState->IntermediateFacing) > SmoothWalkingStateErrorTolerance::FacingDegreeErrorTolerance ||
+			!(IntermediateAngularVelocity - AuthoritySpringState->IntermediateAngularVelocity).IsNearlyZero(SmoothWalkingStateErrorTolerance::AngularVelocityErrorTolerance)));
 }
 
 
