One line crash fix for potential nullptr dereference when trying to get floor from mover in blueprint before sim blackboard has been created

diff --git a/Engine/Plugins/Experimental/Mover/Source/Mover/Private/MoverComponent.cpp b/Engine/Plugins/Experimental/Mover/Source/Mover/Private/MoverComponent.cpp
--- a/Engine/Plugins/Experimental/Mover/Source/Mover/Private/MoverComponent.cpp
+++ b/Engine/Plugins/Experimental/Mover/Source/Mover/Private/MoverComponent.cpp
@@ -2397,7 +2397,7 @@ const FMoverSyncState& UMoverComponent::GetSyncState() const
 bool UMoverComponent::TryGetFloorCheckHitResult(FHitResult& OutHitResult) const
 {
 	FFloorCheckResult FloorCheck;
-	if (SimBlackboard->TryGet(CommonBlackboard::LastFloorResult, FloorCheck))
+	if (SimBlackboard != nullptr && SimBlackboard->TryGet(CommonBlackboard::LastFloorResult, FloorCheck))
 	{
 		OutHitResult = FloorCheck.HitResult;
 		return true;
