Mover: adding experimental option to append inputs into sync state, for cases where you want to share inputs with sim proxies that are not simulated and may not have those inputs available. This is temporary until all networking configurations support an equivalent option.

diff --git a/Engine/Plugins/Experimental/Mover/Source/Mover/Private/BlueprintDataExt.cpp b/Engine/Plugins/Experimental/Mover/Source/Mover/Private/BlueprintDataExt.cpp
--- a/Engine/Plugins/Experimental/Mover/Source/Mover/Private/BlueprintDataExt.cpp
+++ b/Engine/Plugins/Experimental/Mover/Source/Mover/Private/BlueprintDataExt.cpp
@@ -254,3 +254,10 @@ void FMoverDictionaryData::Merge(const FMoverDataStructBase& From)
 	}
 	
 }
+
+
+bool FMoverDictionaryData::ShouldReconcile(const FMoverDataStructBase& AuthorityState) const
+{
+	// This type isn't set up to cause reconciliation.
+	return false;
+}
diff --git a/Engine/Plugins/Experimental/Mover/Source/Mover/Private/InputContainerStruct.cpp b/Engine/Plugins/Experimental/Mover/Source/Mover/Private/InputContainerStruct.cpp
new file mode 100644
--- /dev/null
+++ b/Engine/Plugins/Experimental/Mover/Source/Mover/Private/InputContainerStruct.cpp
@@ -0,0 +1,47 @@
+// Copyright Epic Games, Inc. All Rights Reserved.
+
+
+#include "InputContainerStruct.h"
+#include "MoverLog.h"
+
+#include UE_INLINE_GENERATED_CPP_BY_NAME(InputContainerStruct)
+
+#define LOCTEXT_NAMESPACE "MoverInputContainerStruct"
+
+
+
+void FMoverInputContainerDataStruct::Interpolate(const FMoverDataStructBase& From, const FMoverDataStructBase& To, float LerpFactor)
+{
+	const FMoverInputContainerDataStruct* FromContainer = static_cast<const FMoverInputContainerDataStruct*>(&From);
+	const FMoverInputContainerDataStruct* ToContainer = static_cast<const FMoverInputContainerDataStruct*>(&To);
+
+	InputCollection.Interpolate(FromContainer->InputCollection, ToContainer->InputCollection, LerpFactor);
+}
+
+
+FMoverDataStructBase* FMoverInputContainerDataStruct::Clone() const
+{
+	FMoverInputContainerDataStruct* CopyPtr = new FMoverInputContainerDataStruct(*this);
+	return CopyPtr;
+}
+
+bool FMoverInputContainerDataStruct::NetSerialize(FArchive& Ar, UPackageMap* Map, bool& bOutSuccess)
+{
+	if (!Super::NetSerialize(Ar, Map, bOutSuccess))
+	{
+		bOutSuccess = false;
+		return false;
+	}
+
+	if (!InputCollection.NetSerialize(Ar, Map, bOutSuccess))
+	{
+		bOutSuccess = false;
+		return false;
+	}
+
+	bOutSuccess = true;
+	return true;
+}
+
+
+#undef LOCTEXT_NAMESPACE
diff --git a/Engine/Plugins/Experimental/Mover/Source/Mover/Private/MoverComponent.cpp b/Engine/Plugins/Experimental/Mover/Source/Mover/Private/MoverComponent.cpp
--- a/Engine/Plugins/Experimental/Mover/Source/Mover/Private/MoverComponent.cpp
+++ b/Engine/Plugins/Experimental/Mover/Source/Mover/Private/MoverComponent.cpp
@@ -12,6 +12,7 @@
 #include "MoveLibrary/MovementUtils.h"
 #include "MoveLibrary/FloorQueryUtils.h"
 #include "MoveLibrary/RollbackBlackboard.h"
+#include "InputContainerStruct.h"
 #include "MoverLog.h"
 #include "InstantMovementEffect.h"
 #include "Backends/MoverNetworkPredictionLiaison.h"
@@ -421,6 +422,23 @@ void UMoverComponent::FinalizeSmoothingFrame(const FMoverSyncState* SyncState, c
 
 void UMoverComponent::TickInterpolatedSimProxy(const FMoverTimeStep& TimeStep, const FMoverInputCmdContext& InputCmd, UMoverComponent* MoverComp, const FMoverSyncState& CachedSyncState, const FMoverSyncState& SyncState, const FMoverAuxStateContext& AuxState)
 {
+	if (bSyncInputsForSimProxy)
+	{
+		CachedLastUsedInputCmd = InputCmd;
+
+		// Copy any structs that may be inputs from sync state to input cmd - note the use of the special container class that lets the inputs avoid causing rollbacks
+		if (FMoverInputContainerDataStruct* InputContainer = static_cast<FMoverInputContainerDataStruct*>(SyncState.SyncStateCollection.FindDataByType(FMoverInputContainerDataStruct::StaticStruct())))
+		{
+			for (auto InputStructIt = InputContainer->InputCollection.GetCollectionDataIterator(); InputStructIt; ++InputStructIt)
+			{
+				if (const FMoverDataStructBase* InputDataStruct = InputStructIt->Get())
+				{
+					CachedLastUsedInputCmd.InputCollection.AddDataByCopy(InputDataStruct);
+				}
+			}
+		}
+	}
+
 	TArray<TSharedPtr<FMovementModifierBase>> ModifiersToStart;
 	TArray<TSharedPtr<FMovementModifierBase>> ModifiersToEnd;
 
@@ -647,6 +665,22 @@ void UMoverComponent::SimulationTick(const FMoverTimeStep& InTimeStep, const FMo
 	{
 		CachedNewestSimTickTimeStep = MoverTimeStep;
 	}	
+
+	if (bSyncInputsForSimProxy)
+	{
+		// stow all inputs away in a special container struct that avoids causing potential rollbacks 
+		// so they can be available to other clients even if they're only interpolated sim proxies
+		if (FMoverInputContainerDataStruct* InputContainer = static_cast<FMoverInputContainerDataStruct*>(SimOutput.SyncState.SyncStateCollection.FindOrAddDataByType(FMoverInputContainerDataStruct::StaticStruct())))
+		{	
+			for (auto InputCmdIt = SimInput.InputCmd.InputCollection.GetCollectionDataIterator(); InputCmdIt; ++InputCmdIt)
+			{
+				if (InputCmdIt->Get())
+				{
+					InputContainer->InputCollection.AddDataByCopy(InputCmdIt->Get());
+				}
+			}
+		}
+	}
 }
 
 UBaseMovementMode* UMoverComponent::FindMovementMode(TSubclassOf<UBaseMovementMode> MovementMode) const
diff --git a/Engine/Plugins/Experimental/Mover/Source/Mover/Public/BlueprintDataExt.h b/Engine/Plugins/Experimental/Mover/Source/Mover/Public/BlueprintDataExt.h
--- a/Engine/Plugins/Experimental/Mover/Source/Mover/Public/BlueprintDataExt.h
+++ b/Engine/Plugins/Experimental/Mover/Source/Mover/Public/BlueprintDataExt.h
@@ -11,7 +11,8 @@
 
 /** Data block containing mappings between names and commonly-used types, so that Blueprint-only devs can include custom data
  * in their project's sync state or input cmds without needing native code changes.
- * EXPERIMENTAL: this will be replaced with generic user-defined struct support
+ * EXPERIMENTAL: this will be removed in favor of generic user-defined struct support. If this is for Blueprint usage,
+ * consider using FMoverUserDefinedDataStruct instead. If native, consider deriving your own FMoverDataStructBase type.
  */ 
 USTRUCT(BlueprintType, Experimental)
 struct FMoverDictionaryData : public FMoverDataStructBase
@@ -53,6 +54,8 @@ public:
 	virtual void ToString(FAnsiStringBuilderBase& Out) const override;
 	virtual void Interpolate(const FMoverDataStructBase& From, const FMoverDataStructBase& To, float Alpha) override;
 	virtual void Merge(const FMoverDataStructBase& From) override;
+	virtual bool ShouldReconcile(const FMoverDataStructBase& AuthorityState) const override;
+
 };
 
 
diff --git a/Engine/Plugins/Experimental/Mover/Source/Mover/Public/InputContainerStruct.h b/Engine/Plugins/Experimental/Mover/Source/Mover/Public/InputContainerStruct.h
new file mode 100644
--- /dev/null
+++ b/Engine/Plugins/Experimental/Mover/Source/Mover/Public/InputContainerStruct.h
@@ -0,0 +1,39 @@
+// Copyright Epic Games, Inc. All Rights Reserved.
+
+#pragma once
+
+#include "MoverTypes.h"
+#include "InputContainerStruct.generated.h"
+
+
+
+#define UE_API MOVER_API
+
+/** 
+ * Wrapper class that's used to include input structs in the sync state without them causing reconciliation.
+ * This is intended only for internal use.
+ */
+
+ USTRUCT()
+ struct FMoverInputContainerDataStruct : public FMoverDataStructBase
+ {
+	GENERATED_BODY()
+
+public:
+	// All input data in this struct
+	FMoverDataCollection InputCollection;
+
+	// Implementation of FMoverDataStructBase
+
+	// This struct never triggers reconciliation
+	virtual bool ShouldReconcile(const FMoverDataStructBase& AuthorityState) const override { return false; }
+	virtual void Interpolate(const FMoverDataStructBase& From, const FMoverDataStructBase& To, float LerpFactor) override;
+	virtual FMoverDataStructBase* Clone() const override;
+	virtual bool NetSerialize(FArchive& Ar, UPackageMap* Map, bool& bOutSuccess) override;
+
+	virtual UScriptStruct* GetScriptStruct() const override { return StaticStruct(); }
+
+};
+
+
+#undef UE_API
diff --git a/Engine/Plugins/Experimental/Mover/Source/Mover/Public/MoverComponent.h b/Engine/Plugins/Experimental/Mover/Source/Mover/Public/MoverComponent.h
--- a/Engine/Plugins/Experimental/Mover/Source/Mover/Public/MoverComponent.h
+++ b/Engine/Plugins/Experimental/Mover/Source/Mover/Public/MoverComponent.h
@@ -835,6 +835,10 @@ public:
 	UPROPERTY(EditAnywhere, BlueprintReadWrite, Category = "Mover", AdvancedDisplay)
 	uint8 bAcceptExternalMovement : 1 = 0;
 
+	// If enabled, we'll send inputs along with to sim proxy via the sync state, and they'll be available via GetLastInputCmd. This may be useful for cases where input is used to hint at object state, such as an anim graph. This option is intended to be temporary until all networking backends allow this.
+	UPROPERTY(EditAnywhere, BlueprintReadWrite, Category = "Mover", AdvancedDisplay, Experimental)
+	uint8 bSyncInputsForSimProxy : 1 = 0;
+
 	MOVER_API void SetSimulationOutput(const FMoverTimeStep& TimeStep, const UE::Mover::FSimulationOutputData& OutputData);
 
 	// Dispatch a simulation event. It will be processed immediately.
