Mover: removed ensure that was saying the updated component couldn't also be the primary visual component. Made is so visual offset smoothing mode quietly does nothing if this is the case.

#rnx

diff --git a/Engine/Plugins/Experimental/Mover/Source/Mover/Private/MoverComponent.cpp b/Engine/Plugins/Experimental/Mover/Source/Mover/Private/MoverComponent.cpp
--- a/Engine/Plugins/Experimental/Mover/Source/Mover/Private/MoverComponent.cpp
+++ b/Engine/Plugins/Experimental/Mover/Source/Mover/Private/MoverComponent.cpp
@@ -283,7 +283,7 @@ void UMoverComponent::BeginPlay()
 			}
 		}
 
-		ensureMsgf(UpdatedComponent && (PrimaryVisualComponent != UpdatedComponent), TEXT("A Mover actor (%s) must have an UpdatedComponent and cannot have a PrimaryVisualComponent that is the same as UpdatedComponent"), *GetNameSafe(MyActor));
+		ensureMsgf(UpdatedComponent, TEXT("A Mover actor (%s) must have an UpdatedComponent"), *GetNameSafe(MyActor));
 
 		// Optional motion warping support
 		if (UMotionWarpingComponent* WarpingComp = MyActor->FindComponentByClass<UMotionWarpingComponent>())
@@ -408,7 +408,7 @@ void UMoverComponent::FinalizeSmoothingFrame(const FMoverSyncState* SyncState, c
 {
 	if (PrimaryVisualComponent)
 	{
-		if (SmoothingMode == EMoverSmoothingMode::VisualComponentOffset)
+		if (SmoothingMode == EMoverSmoothingMode::VisualComponentOffset && (PrimaryVisualComponent != UpdatedComponent))
 		{
 			// Offset the visual component so it aligns with the smoothed state transform, while leaving the actual root component in place
 			if (const FMoverDefaultSyncState* MoverState = SyncState->SyncStateCollection.FindDataByType<FMoverDefaultSyncState>())
