Mover: rollback blackboard API tweaks to prevent settings-related crash, plus support for plain types / r-values, const-correctness, and adding proper module exports

#rnx

diff --git a/Engine/Plugins/Experimental/Mover/Source/Mover/Private/MoveLibrary/RollbackBlackboard.cpp b/Engine/Plugins/Experimental/Mover/Source/Mover/Private/MoveLibrary/RollbackBlackboard.cpp
--- a/Engine/Plugins/Experimental/Mover/Source/Mover/Private/MoveLibrary/RollbackBlackboard.cpp
+++ b/Engine/Plugins/Experimental/Mover/Source/Mover/Private/MoveLibrary/RollbackBlackboard.cpp
@@ -65,7 +65,7 @@ bool URollbackBlackboard::BlackboardEntryBase::CanReadEntryAt(const EntryTimeSta
 		return false;	// entry isn't initialized yet or never set
 	}
 
-	switch (Settings->PersistencePolicy)
+	switch (Settings.PersistencePolicy)
 	{
 		case EBlackboardPersistencePolicy::NextFrameOnly:
 		{
diff --git a/Engine/Plugins/Experimental/Mover/Source/Mover/Public/MoveLibrary/RollbackBlackboard.h b/Engine/Plugins/Experimental/Mover/Source/Mover/Public/MoveLibrary/RollbackBlackboard.h
--- a/Engine/Plugins/Experimental/Mover/Source/Mover/Public/MoveLibrary/RollbackBlackboard.h
+++ b/Engine/Plugins/Experimental/Mover/Source/Mover/Public/MoveLibrary/RollbackBlackboard.h
@@ -28,6 +28,7 @@ struct FMoverTimeStep;
  * Notes:
  *  - This is implemented under-the-hood using circular buffers, alleviating the need for mem alloc/frees during use. It's important to create entries with policy settings appropriate for their use pattern.
  *  - We use an "InternalWrapper" class to allow in-simulation objects like movement modes to use the same API, without needing to choose between internal- or external-facing functions. 
+ *  - The "InternalWrapper" class will be replaced by a different pattern to govern access levels
  */
 
  // TODO: 
@@ -154,15 +155,15 @@ private:
 		}
 
 		// Note that NewPendingFrame means the frame that will now be re-simulated. So any existing entries with a Frame >= NewPendingFrame can be invalidated.
-		void RollBack(uint32 NewPendingFrame);
+		UE_API void RollBack(uint32 NewPendingFrame);
 
 
 	protected:
-		bool CanReadEntryAt(const EntryTimeStamp& ReaderTimeStamp, EEntryIndexType IndexType) const;
+		UE_API bool CanReadEntryAt(const EntryTimeStamp& ReaderTimeStamp, EEntryIndexType IndexType) const;
 
-		uint32 ComputeBufferSize(EBlackboardSizingPolicy InSizingPolicy, uint32 FixedBufferSize = -1);
+		UE_API uint32 ComputeBufferSize(EBlackboardSizingPolicy InSizingPolicy, uint32 FixedBufferSize = -1);
 
-		const EntrySettings* Settings = nullptr;
+		EntrySettings Settings;
 
 		TCircularBuffer<EntryTimeStamp> Timestamps = TCircularBuffer<EntryTimeStamp>(0);
 
@@ -182,9 +183,9 @@ private:
 		// TODO: consider what options need to be specified
 		BlackboardEntry(const EntrySettings& InSettings, const EntryT* OptionalInitialObj=nullptr)
 		{ 
-			Settings = &InSettings;
+			Settings = InSettings;
 
-			int32 BufferSize = ComputeBufferSize(Settings->SizingPolicy, Settings->FixedSize);
+			int32 BufferSize = ComputeBufferSize(Settings.SizingPolicy, Settings.FixedSize);
 
 			Timestamps = TCircularBuffer<EntryTimeStamp>(BufferSize, EntryTimeStamp());
 
@@ -328,7 +329,7 @@ public:
 private:
 	/** Store object by a named key, overwriting any existing object. Should only be called from within the simulation. */
 	template<typename EntryT>
-	bool TrySet_Internal(FName ObjName, EntryT& Obj)
+	bool TrySet_Internal(FName ObjName, const EntryT& Obj)
 	{
 		check((bIsSimulationInProgress && (InProgressSimFrameThreadId == FPlatformTLS::GetCurrentThreadId())) || 
 			  (bIsRollbackInProgress && (InRollbackThreadId == FPlatformTLS::GetCurrentThreadId()))) ;
@@ -439,7 +440,7 @@ public:
 	}
 
 	template<typename EntryT>
-	bool TrySet(FName ObjName, EntryT& Obj)
+	bool TrySet(FName ObjName, const EntryT& Obj)
 	{
 		return Blackboard->TrySet_Internal<EntryT>(ObjName, Obj);
 	}
