[Mover] Fix erratic rotation when moving down negative X in smooth walking mode

diff --git a/Engine/Plugins/Experimental/Mover/Source/Mover/Private/DefaultMovementSet/Modes/SmoothWalkingMode.cpp b/Engine/Plugins/Experimental/Mover/Source/Mover/Private/DefaultMovementSet/Modes/SmoothWalkingMode.cpp
--- a/Engine/Plugins/Experimental/Mover/Source/Mover/Private/DefaultMovementSet/Modes/SmoothWalkingMode.cpp
+++ b/Engine/Plugins/Experimental/Mover/Source/Mover/Private/DefaultMovementSet/Modes/SmoothWalkingMode.cpp
@@ -168,8 +168,8 @@ void USmoothWalkingMode::GenerateWalkMove_Implementation(FMoverTickStartData& St
 		// Ensure the output angular velocity will snap perfectly to the target
 		// Note we don't do this normally because its better to have a consistent angular velocity
 		// If we output the angular velocity based on the updated facing every frame it can cause errors at low dt due to inaccuracy in the inverse
-		// exponential approximation inside thne spring damper
-		CurrentAngularVelocityRadians = DeltaSeconds > 0.0f ? (CurrentFacing.Inverse() * UpdatedFacing).ToRotationVector() / DeltaSeconds : FVector::ZeroVector;
+		// exponential approximation inside the spring damper
+		CurrentAngularVelocityRadians = DeltaSeconds > 0.0f ? ((CurrentFacing.Inverse() * UpdatedFacing).GetShortestArcWith(FQuat::Identity)).ToRotationVector() / DeltaSeconds : FVector::ZeroVector;
 		SpringState.IntermediateFacing = DesiredFacing;
 
 		// If we've reached our target then also snap the angular velocity to zero once it is close enough
