Allow sampling of mover trajectory at a different (higher frequency) rate than we output to the trajectory, to avoid errors from moves which are not timestep independent.
Change use of PoseSearchTrajectoryPredictor so that it is not expected to also set the current transform, it is only expected to predict into the future. This makes the meaning of 'NumPredictionSamples' more consistent throughout the API

diff --git a/Engine/Plugins/Animation/PoseSearch/Source/Runtime/Private/PoseSearchTrajectoryLibrary.cpp b/Engine/Plugins/Animation/PoseSearch/Source/Runtime/Private/PoseSearchTrajectoryLibrary.cpp
--- a/Engine/Plugins/Animation/PoseSearch/Source/Runtime/Private/PoseSearchTrajectoryLibrary.cpp
+++ b/Engine/Plugins/Animation/PoseSearch/Source/Runtime/Private/PoseSearchTrajectoryLibrary.cpp
@@ -631,9 +631,16 @@ void UPoseSearchTrajectoryLibrary::PoseSearchGenerateTransformTrajectoryWithPred
 	InitTrajectorySamples(InOutTrajectory, CurrentPosition, CurrentFacing, TrajectoryDataSampling, InDeltaTime);
 	UpdateHistory_TransformHistory(InOutTrajectory, CurrentPosition, CurrentVelocity, TrajectoryDataSampling, InDeltaTime);
 
+	// Set the current position at i == NumHistoryCount at t == delta time
+	// Remember: t == 0 is the previous position, t == delta time is the current frame's position
+	// assuming we call this method after the movement component has updated to a new position
+	InOutTrajectory.Samples[InTrajectoryHistoryCount].TimeInSeconds = InDeltaTime;
+	InOutTrajectory.Samples[InTrajectoryHistoryCount].Position = CurrentPosition;
+	InOutTrajectory.Samples[InTrajectoryHistoryCount].Facing = CurrentFacing;
+
 	if (InPredictor != nullptr)
 	{
-		InPredictor->Predict(InOutTrajectory, InTrajectoryPredictionCount + 1, InPredictionSamplingInterval, InTrajectoryHistoryCount);
+		InPredictor->Predict(InOutTrajectory, InTrajectoryPredictionCount, InPredictionSamplingInterval, InTrajectoryHistoryCount);
 	}
 
 	//InOutDesiredControllerYawLastUpdate = TrajectoryDataState.DesiredControllerYawLastUpdate;
diff --git a/Engine/Plugins/Experimental/Mover/Source/Mover/Private/MoverPoseSearchTrajectoryPredictor.cpp b/Engine/Plugins/Experimental/Mover/Source/Mover/Private/MoverPoseSearchTrajectoryPredictor.cpp
--- a/Engine/Plugins/Experimental/Mover/Source/Mover/Private/MoverPoseSearchTrajectoryPredictor.cpp
+++ b/Engine/Plugins/Experimental/Mover/Source/Mover/Private/MoverPoseSearchTrajectoryPredictor.cpp
@@ -16,34 +16,65 @@ void UMoverTrajectoryPredictor::Predict(FTransformTrajectory& InOutTrajectory,
 		return;
 	}
 
-	Predict(*MoverComponent, InOutTrajectory, NumPredictionSamples, SecondsPerPredictionSample, NumHistorySamples);
+	Predict(*MoverComponent, InOutTrajectory, NumPredictionSamples, SecondsPerPredictionSample, NumHistorySamples, static_cast<float>(MoverSamplingFrameRate.AsInterval()));
 }
 
-void UMoverTrajectoryPredictor::Predict(UMoverComponent& MoverComponent, FTransformTrajectory& InOutTrajectory, int32 NumPredictionSamples, float SecondsPerPredictionSample, int32 NumHistorySamples)
+void UMoverTrajectoryPredictor::Predict(UMoverComponent& MoverComponent, FTransformTrajectory& InOutTrajectory, int32 NumPredictionSamples, float SecondsPerPredictionSample, int32 NumHistorySamples, float MoverSamplingInterval)
 {
 	FMoverPredictTrajectoryParams PredictParams;
-	PredictParams.NumPredictionSamples = NumPredictionSamples;
-	PredictParams.SecondsPerSample = SecondsPerPredictionSample;
+
+	// Important: the sampling frequency of the prediction does not necessarily match the output frequency on the trajectory
+	float LookAheadTime = NumPredictionSamples * SecondsPerPredictionSample;
+	int NumMoverPredictionSamplesRequired =  FMath::FloorToInt32(LookAheadTime / MoverSamplingInterval) + 2;
+	
+	PredictParams.NumPredictionSamples = NumMoverPredictionSamplesRequired;
+	PredictParams.SecondsPerSample = MoverSamplingInterval;
 	PredictParams.bUseVisualComponentRoot = true;
 	PredictParams.bDisableGravity = true;
 
+	// IMPORTANT! The first sample returned is actually the current state
 	TArray<FTrajectorySampleInfo> MoverPredictionSamples = MoverComponent.GetPredictedTrajectory(PredictParams);
 
-	if (InOutTrajectory.Samples.Num() < (NumHistorySamples + MoverPredictionSamples.Num()))
+	if (InOutTrajectory.Samples.Num() < (NumHistorySamples + NumPredictionSamples))
 	{
 		UE_LOG(LogMover, Warning, TEXT("InOutTrajectory Samples array does not have enough space for %i predicted samples"), MoverPredictionSamples.Num());
 	}
 	else
 	{
-		float AccumulatedSeconds = 0.0;
+		// Subsample or supersample the mover prediction to the trajectory prediction
+
+		int MoverPredictionSampleIndex = 1; // we start at index 1, since index 0 is actually the current state
+		
+		// Start at Current position
+		float CurrentTime = InOutTrajectory.Samples[NumHistorySamples].TimeInSeconds;	// t == 0 is the last frame, so we need to account for the starting mover state being offset
+		FTransform MoverPredictLower = MoverPredictionSamples[0].Transform;
+		FTransform MoverPredictUpper = MoverPredictionSamples[1].Transform;
+		float TimeInSecondsLower = CurrentTime;
+		float TimeInSecondsUpper = CurrentTime + MoverSamplingInterval;
+		float AccumulatedSeconds = CurrentTime + SecondsPerPredictionSample; // first prediction sample should be at time for the first prediction
 
 		for (int32 i = 0; i < NumPredictionSamples; ++i)
 		{
-			const int PoseSampleIdx = i + NumHistorySamples;
-			const int PredictedSampleIdx = i;
+			// Progress target if necessary
+			while (AccumulatedSeconds > TimeInSecondsUpper && MoverPredictionSampleIndex < NumMoverPredictionSamplesRequired - 1)
+			{
+				MoverPredictionSampleIndex++;
+
+				// Update to next mover prediction value
+				MoverPredictLower = MoverPredictionSamples[MoverPredictionSampleIndex - 1].Transform;
+				MoverPredictUpper = MoverPredictionSamples[MoverPredictionSampleIndex].Transform;
+
+				TimeInSecondsLower = ((MoverPredictionSampleIndex - 1) * MoverSamplingInterval) + CurrentTime;
+				TimeInSecondsUpper = (MoverPredictionSampleIndex * MoverSamplingInterval) + CurrentTime;
+			}
+
+			const int PoseSampleIdx = i + NumHistorySamples + 1;
+
+			float T = (AccumulatedSeconds - TimeInSecondsLower) / (TimeInSecondsUpper - TimeInSecondsLower);
+			T = FMath::Clamp(T, 0.0f, 1.0f);
 
-			InOutTrajectory.Samples[PoseSampleIdx].Position = MoverPredictionSamples[PredictedSampleIdx].Transform.GetLocation();
-			InOutTrajectory.Samples[PoseSampleIdx].Facing = MoverPredictionSamples[PredictedSampleIdx].Transform.GetRotation();
+			InOutTrajectory.Samples[PoseSampleIdx].Position = FMath::Lerp(MoverPredictLower.GetLocation(), MoverPredictUpper.GetLocation(), T);
+			InOutTrajectory.Samples[PoseSampleIdx].Facing = FQuat::Slerp(MoverPredictLower.GetRotation(), MoverPredictUpper.GetRotation(), T);
 			InOutTrajectory.Samples[PoseSampleIdx].TimeInSeconds = AccumulatedSeconds;
 
 			AccumulatedSeconds += SecondsPerPredictionSample;
diff --git a/Engine/Plugins/Experimental/Mover/Source/Mover/Public/MoverComponent.h b/Engine/Plugins/Experimental/Mover/Source/Mover/Public/MoverComponent.h
--- a/Engine/Plugins/Experimental/Mover/Source/Mover/Public/MoverComponent.h
+++ b/Engine/Plugins/Experimental/Mover/Source/Mover/Public/MoverComponent.h
@@ -482,7 +482,8 @@ public:	// Queries
 	UFUNCTION(BlueprintCallable, BlueprintPure = false, Category = Mover)
 	MOVER_API TArray<FTrajectorySampleInfo> GetFutureTrajectory(float FutureSeconds, float SamplesPerSecond);
 
-	/** Get a sampling of where the actor is projected to be in the future, based on a current state. Note that this is projecting ideal movement without doing full simulation and collision. */
+	/** Get a sampling of where the actor is projected to be in the future, based on a current state. Note that this is projecting ideal movement without doing full simulation and collision.
+	 * The first sample info of the returned array corresponds to the current state of the mover. */
 	UFUNCTION(BlueprintCallable, BlueprintPure = false, Category = Mover)
 	MOVER_API TArray<FTrajectorySampleInfo> GetPredictedTrajectory(FMoverPredictTrajectoryParams PredictionParams);	
 
diff --git a/Engine/Plugins/Experimental/Mover/Source/Mover/Public/MoverPoseSearchTrajectoryPredictor.h b/Engine/Plugins/Experimental/Mover/Source/Mover/Public/MoverPoseSearchTrajectoryPredictor.h
--- a/Engine/Plugins/Experimental/Mover/Source/Mover/Public/MoverPoseSearchTrajectoryPredictor.h
+++ b/Engine/Plugins/Experimental/Mover/Source/Mover/Public/MoverPoseSearchTrajectoryPredictor.h
@@ -2,6 +2,7 @@
 
 #pragma once
 
+#include "Misc/FrameRate.h"
 #include "PoseSearch/PoseSearchTrajectoryPredictor.h"
 #include "MoverPoseSearchTrajectoryPredictor.generated.h"
 
@@ -33,13 +34,17 @@ public:
 	UE_API virtual void GetVelocity(FVector& OutVelocity) override;
 
 	// Share code in static helper functions so we can call it from AnimNext, which doesn't support interfaces
-	static UE_API void Predict(UMoverComponent& Mover, FTransformTrajectory& InOutTrajectory, int32 NumPredictionSamples, float SecondsPerPredictionSample, int32 NumHistorySamples);
+	static UE_API void Predict(UMoverComponent& Mover, FTransformTrajectory& InOutTrajectory, int32 NumPredictionSamples, float SecondsPerPredictionSample, int32 NumHistorySamples, float MoverSamplingInterval);
 	static UE_API void GetCurrentState(UMoverComponent& Mover, FVector& OutPosition, FQuat& OutFacing, FVector& OutVelocity);
 
 
 protected:
 	TObjectPtr<UMoverComponent> MoverComponent;
 
+	// Sampling frame rate to query the mover, doesn't necessarily match the SecondsPerPredictionSample which is used to write to the trajectory
+	// This allows sampling of mover at a high frequency for accuracy but then we can downsample the results to a more coarse grained trajectory
+	UPROPERTY(EditAnywhere, Category = "Animation|PoseSearch|Experimental", BlueprintReadWrite)
+	FFrameRate MoverSamplingFrameRate = FFrameRate(60, 1);
 };
 
 #undef UE_API
diff --git a/Engine/Plugins/Experimental/MoverAnimNext/Source/MoverAnimNext/Private/Graph/RigUnit_GenerateMoverTrajectory.cpp b/Engine/Plugins/Experimental/MoverAnimNext/Source/MoverAnimNext/Private/Graph/RigUnit_GenerateMoverTrajectory.cpp
--- a/Engine/Plugins/Experimental/MoverAnimNext/Source/MoverAnimNext/Private/Graph/RigUnit_GenerateMoverTrajectory.cpp
+++ b/Engine/Plugins/Experimental/MoverAnimNext/Source/MoverAnimNext/Private/Graph/RigUnit_GenerateMoverTrajectory.cpp
@@ -38,5 +38,5 @@ FRigUnit_GenerateMoverTrajectory_Execute()
 	UPoseSearchTrajectoryLibrary::InitTrajectorySamples(InOutTrajectory, CurrentPosition, CurrentFacing, TrajectoryDataSampling, DeltaTime);
 	UPoseSearchTrajectoryLibrary::UpdateHistory_TransformHistory(InOutTrajectory, CurrentPosition, CurrentVelocity, TrajectoryDataSampling, DeltaTime);
 	UMoverTrajectoryPredictor::Predict(*MoverComponent, InOutTrajectory,
-		NumPredictionSamples + 1, PredictionSamplingInterval, NumHistorySamples);
+		NumPredictionSamples + 1, PredictionSamplingInterval, NumHistorySamples, static_cast<float>(MoverSamplingFrameRate.AsInterval()));
 }
diff --git a/Engine/Plugins/Experimental/MoverAnimNext/Source/MoverAnimNext/Private/Graph/RigUnit_GenerateMoverTrajectory.h b/Engine/Plugins/Experimental/MoverAnimNext/Source/MoverAnimNext/Private/Graph/RigUnit_GenerateMoverTrajectory.h
--- a/Engine/Plugins/Experimental/MoverAnimNext/Source/MoverAnimNext/Private/Graph/RigUnit_GenerateMoverTrajectory.h
+++ b/Engine/Plugins/Experimental/MoverAnimNext/Source/MoverAnimNext/Private/Graph/RigUnit_GenerateMoverTrajectory.h
@@ -5,6 +5,7 @@
 #include "AnimNextExecuteContext.h"
 #include "Graph/RigUnit_AnimNextBase.h"
 #include "Animation/TrajectoryTypes.h"
+#include "Misc/FrameRate.h"
 
 #include "RigUnit_GenerateMoverTrajectory.generated.h"
 
@@ -31,12 +32,17 @@ struct FRigUnit_GenerateMoverTrajectory : public FRigUnit_AnimNextBase
 	UPROPERTY(EditAnywhere, Category = "Trajectory", meta = (Input))
 	int32 NumHistorySamples = 30;
 
+	// How far apart (in seconds) the prediction steps should be 
 	UPROPERTY(EditAnywhere, Category = "Trajectory", meta = (Input))
 	float PredictionSamplingInterval = 0.1f;
 
 	UPROPERTY(EditAnywhere, Category = "Trajectory", meta = (Input))
 	int32 NumPredictionSamples = 15;
 
+	// Sampling frame rate to query the mover, doesn't necessarily match the PredictionSamplingInterval which is used to write to the trajectory
+	UPROPERTY(EditAnywhere, Category = "Trajectory", meta = (Input))
+	FFrameRate MoverSamplingFrameRate = FFrameRate(60, 1);
+
 	UPROPERTY(EditAnywhere, Category = "Trajectory", meta = (Input, Output))
 	FTransformTrajectory InOutTrajectory;
 	
