Mover: Adding virtual destructor to BlackboardEntryBase struct to fix a memory leak.

diff --git a/Engine/Plugins/Experimental/Mover/Source/Mover/Public/MoveLibrary/RollbackBlackboard.h b/Engine/Plugins/Experimental/Mover/Source/Mover/Public/MoveLibrary/RollbackBlackboard.h
--- a/Engine/Plugins/Experimental/Mover/Source/Mover/Public/MoveLibrary/RollbackBlackboard.h
+++ b/Engine/Plugins/Experimental/Mover/Source/Mover/Public/MoveLibrary/RollbackBlackboard.h
@@ -148,6 +148,11 @@ private:
 	struct BlackboardEntryBase
 	{
 	public:
+
+		virtual ~BlackboardEntryBase()
+		{
+		}
+		
 		void OnSimulationFrameEnd()
 		{
 			//Advances the external index to match
