Mover: Fixing debug editor crash when initializing a CircularBuffer with a capacity of 0 on MoverComponent initialization by reworking constructors of RollbackBlackboardEntry structs to remove the need to default construct CircularBuffers.

diff --git a/Engine/Plugins/Experimental/Mover/Source/Mover/Private/MoveLibrary/RollbackBlackboard.cpp b/Engine/Plugins/Experimental/Mover/Source/Mover/Private/MoveLibrary/RollbackBlackboard.cpp
--- a/Engine/Plugins/Experimental/Mover/Source/Mover/Private/MoveLibrary/RollbackBlackboard.cpp
+++ b/Engine/Plugins/Experimental/Mover/Source/Mover/Private/MoveLibrary/RollbackBlackboard.cpp
@@ -30,6 +30,14 @@ uint32 URollbackBlackboard::BlackboardEntryBase::ComputeBufferSize(EBlackboardSi
 	return BufferSize;
 }
 
+URollbackBlackboard::BlackboardEntryBase::BlackboardEntryBase(const EntrySettings& InSettings, const int32& BufferSize)
+	: Settings(InSettings)
+	, Timestamps(BufferSize, EntryTimeStamp())
+	, ExternalIdx(0)
+	, InternalIdx(0)
+{
+}
+
 void URollbackBlackboard::BlackboardEntryBase::RollBack(uint32 NewPendingFrame)
 {
 	// Goal: adjust entry to point at the value from the prior frame. May make the entry invalidated, if there were no values that old.
diff --git a/Engine/Plugins/Experimental/Mover/Source/Mover/Public/MoveLibrary/RollbackBlackboard.h b/Engine/Plugins/Experimental/Mover/Source/Mover/Public/MoveLibrary/RollbackBlackboard.h
--- a/Engine/Plugins/Experimental/Mover/Source/Mover/Public/MoveLibrary/RollbackBlackboard.h
+++ b/Engine/Plugins/Experimental/Mover/Source/Mover/Public/MoveLibrary/RollbackBlackboard.h
@@ -148,11 +148,13 @@ private:
 	struct BlackboardEntryBase
 	{
 	public:
-
+		BlackboardEntryBase() = delete;
+		BlackboardEntryBase(const EntrySettings& InSettings, const int32& BufferSize);
+		
 		virtual ~BlackboardEntryBase()
 		{
 		}
-		
+
 		void OnSimulationFrameEnd()
 		{
 			//Advances the external index to match
@@ -162,15 +164,14 @@ private:
 		// Note that NewPendingFrame means the frame that will now be re-simulated. So any existing entries with a Frame >= NewPendingFrame can be invalidated.
 		UE_API void RollBack(uint32 NewPendingFrame);
 
+		UE_API static uint32 ComputeBufferSize(EBlackboardSizingPolicy InSizingPolicy, uint32 FixedBufferSize = -1);
 
 	protected:
 		UE_API bool CanReadEntryAt(const EntryTimeStamp& ReaderTimeStamp, EEntryIndexType IndexType) const;
 
-		UE_API uint32 ComputeBufferSize(EBlackboardSizingPolicy InSizingPolicy, uint32 FixedBufferSize = -1);
-
 		EntrySettings Settings;
 
-		TCircularBuffer<EntryTimeStamp> Timestamps = TCircularBuffer<EntryTimeStamp>(0);
+		TCircularBuffer<EntryTimeStamp> Timestamps;
 
 		uint32 ExternalIdx;	// Indexes to the last entry on a committed frame
 		uint32 InternalIdx;	// Indexes to where in-progress simulations should read/write (matches ExternalIdx when not mid-simulation, and advances if written during simulation)
@@ -186,25 +187,18 @@ private:
 	public:
 
 		// TODO: consider what options need to be specified
-		BlackboardEntry(const EntrySettings& InSettings, const EntryT* OptionalInitialObj=nullptr)
-		{ 
-			Settings = InSettings;
-
-			int32 BufferSize = ComputeBufferSize(Settings.SizingPolicy, Settings.FixedSize);
-
-			Timestamps = TCircularBuffer<EntryTimeStamp>(BufferSize, EntryTimeStamp());
-
+		BlackboardEntry(const EntrySettings& InSettings, const int32& InBufferSize, const EntryT* OptionalInitialObj=nullptr)
+			: BlackboardEntryBase(InSettings, InBufferSize)
+			, EntryBuffer(InBufferSize)
+		{
 			// JAH TODO: Consider making the initial object required, so we can also initialize where in the buffer we point to
 			if (OptionalInitialObj)
 			{
-				EntryBuffer = TCircularBuffer<EntryT>(BufferSize, *OptionalInitialObj);
+				for (uint32 i = 0; i < EntryBuffer.Capacity(); i++)
+				{
+					EntryBuffer[i] = *OptionalInitialObj;
+				}
 			}
-			else
-			{
-				EntryBuffer = TCircularBuffer<EntryT>(BufferSize);
-			}
-
-			ExternalIdx = InternalIdx = 0;
 		}
 
 
@@ -262,7 +256,7 @@ private:
 	
 		// TODO: Consider whether these need to be TCircularBuffers, or whether a TArray would be just fine.  TCircularBuffers round up to the next power of 2 in capacity, so may be wasting significant memory for larger structures.
 		// Note that these EntryBuffer is always the same capacity as Timestamps
-		TCircularBuffer<EntryT> EntryBuffer = TCircularBuffer<EntryT>(0);
+		TCircularBuffer<EntryT> EntryBuffer;
 
 
 	};	// end BlackboardEntry
@@ -279,7 +273,8 @@ public:
 			return false;
 		}
 
-		EntryMap.Emplace(EntryName, MakeUnique<BlackboardEntry<EntryT>>(InSettings, nullptr));
+		const uint32 BufferSize = BlackboardEntryBase::ComputeBufferSize(InSettings.SizingPolicy, InSettings.FixedSize);
+		EntryMap.Emplace(EntryName, MakeUnique<BlackboardEntry<EntryT>>(InSettings, BufferSize, nullptr));
 		return true;
 	}
 
