Mover: fix issue where montages with a non-zero start time would be played from the wrong position on sim proxies

diff --git a/Engine/Plugins/Experimental/Mover/Source/Mover/Private/DefaultMovementSet/LayeredMoves/MontageStateProvider.cpp b/Engine/Plugins/Experimental/Mover/Source/Mover/Private/DefaultMovementSet/LayeredMoves/MontageStateProvider.cpp
--- a/Engine/Plugins/Experimental/Mover/Source/Mover/Private/DefaultMovementSet/LayeredMoves/MontageStateProvider.cpp
+++ b/Engine/Plugins/Experimental/Mover/Source/Mover/Private/DefaultMovementSet/LayeredMoves/MontageStateProvider.cpp
@@ -43,4 +43,6 @@ void FMoverAnimMontageState::NetSerialize(FArchive& Ar)
 	{
 		PlayRate = 1.0f;
 	}
+
+	Ar << CurrentPosition;
 }
