Add RigVM nodes to enable / disable root motion on mover.
Add Mass ST Task to toggle root motion.
Add Mover Methods to query if a tag exists.
Add Mover FSM nullptr checks.
Fix uncessary / slow weakptr pins.

#tests PIE, Editor

diff --git a/Engine/Plugins/Experimental/Mover/Source/Mover/Private/MovementModeStateMachine.cpp b/Engine/Plugins/Experimental/Mover/Source/Mover/Private/MovementModeStateMachine.cpp
--- a/Engine/Plugins/Experimental/Mover/Source/Mover/Private/MovementModeStateMachine.cpp
+++ b/Engine/Plugins/Experimental/Mover/Source/Mover/Private/MovementModeStateMachine.cpp
@@ -664,6 +664,43 @@ void UMovementModeStateMachine::CancelFeaturesWithTag(FGameplayTag TagToCancel,
 
 }
 
+bool UMovementModeStateMachine::HasFeaturesWithTag(FGameplayTag TagToMatch, bool bRequireExactMatch) const
+{
+	bool bResult = false;
+
+	// Compare all matching queued movement features
+	{
+		FRWScopeLock QueueLock(ModifiersQueueLock, SLT_ReadOnly);
+		const TSharedPtr<FMovementModifierBase>* FindResult = QueuedMovementModifiers.FindByPredicate([TagToMatch, bRequireExactMatch] (const TSharedPtr<FMovementModifierBase>& Modifier)
+			{
+				return (Modifier.IsValid() && Modifier->HasGameplayTag(TagToMatch, bRequireExactMatch));
+			});
+		bResult = FindResult != nullptr;
+	}
+
+	if (!bResult)
+	{
+		FRWScopeLock QueueLock(LayeredMoveQueueLock, SLT_ReadOnly);
+		const TSharedPtr<FLayeredMoveBase>* FindResult = QueuedLayeredMoves.FindByPredicate([TagToMatch, bRequireExactMatch](const TSharedPtr<FLayeredMoveBase>& LayeredMove)
+			{
+				return (LayeredMove.IsValid() && LayeredMove->HasGameplayTag(TagToMatch, bRequireExactMatch));
+			});
+		bResult = FindResult != nullptr;
+	}
+
+	if (!bResult)
+	{
+		FRWScopeLock QueueLock(LayeredMoveQueueLock, SLT_ReadOnly);
+		const TSharedPtr<FLayeredMoveInstance>* FindResult = QueuedLayeredMoveInstances.FindByPredicate([TagToMatch, bRequireExactMatch](const TSharedPtr<FLayeredMoveInstance>& LayeredMoveInstance)
+			{
+				return (LayeredMoveInstance.IsValid() && LayeredMoveInstance->HasGameplayTag(TagToMatch, bRequireExactMatch));
+			});
+		bResult = FindResult != nullptr;
+	}
+
+	return bResult;
+}
+
 
 void UMovementModeStateMachine::ConstructDefaultModes()
 {
diff --git a/Engine/Plugins/Experimental/Mover/Source/Mover/Private/MoverComponent.cpp b/Engine/Plugins/Experimental/Mover/Source/Mover/Private/MoverComponent.cpp
--- a/Engine/Plugins/Experimental/Mover/Source/Mover/Private/MoverComponent.cpp
+++ b/Engine/Plugins/Experimental/Mover/Source/Mover/Private/MoverComponent.cpp
@@ -724,6 +724,11 @@ bool UMoverComponent::IsModifierActiveOrQueued(const FMovementModifierHandle& Mo
 
 const FMovementModifierBase* UMoverComponent::FindMovementModifier(const FMovementModifierHandle& ModifierHandle) const
 {
+	if (!ModeFSM)
+	{
+		return nullptr;
+	}
+
 	const FMoverSyncState& CachedSyncState = MoverSyncStateDoubleBuffer.GetReadable();
 	
 	// Check active modifiers for modifier handle
@@ -753,6 +758,11 @@ const FMovementModifierBase* UMoverComponent::FindMovementModifier(const FMoveme
 
 const FMovementModifierBase* UMoverComponent::FindMovementModifierByType(const UScriptStruct* DataStructType) const
 {
+	if (!ModeFSM)
+	{
+		return nullptr;
+	}
+
 	const FMoverSyncState& CachedSyncState = MoverSyncStateDoubleBuffer.GetReadable();
 	
 	// Check active modifiers for modifier handle
@@ -1609,7 +1619,10 @@ DEFINE_FUNCTION(UMoverComponent::execK2_QueueLayeredMove)
 
 void UMoverComponent::QueueLayeredMove(TSharedPtr<FLayeredMoveBase> LayeredMove)
 {	
-	ModeFSM->QueueLayeredMove(LayeredMove);
+	if (ModeFSM)
+	{
+		ModeFSM->QueueLayeredMove(LayeredMove);
+	}
 }
 
 FMovementModifierHandle UMoverComponent::K2_QueueMovementModifier(const int32& MoveAsRawData)
@@ -1647,20 +1660,34 @@ DEFINE_FUNCTION(UMoverComponent::execK2_QueueMovementModifier)
 
 FMovementModifierHandle UMoverComponent::QueueMovementModifier(TSharedPtr<FMovementModifierBase> Modifier)
 {
-	return ModeFSM->QueueMovementModifier(Modifier);
+	return ModeFSM ? ModeFSM->QueueMovementModifier(Modifier) : FMovementModifierHandle(MODIFIER_INVALID_HANDLE);
 }
 
 void UMoverComponent::CancelModifierFromHandle(FMovementModifierHandle ModifierHandle)
 {
-	ModeFSM->CancelModifierFromHandle(ModifierHandle);
+	if (ModeFSM)
+	{
+		ModeFSM->CancelModifierFromHandle(ModifierHandle);
+	}
 }
 
-
 void UMoverComponent::CancelFeaturesWithTag(FGameplayTag TagToCancel, bool bRequireExactMatch)
 {
-	ModeFSM->CancelFeaturesWithTag(TagToCancel, bRequireExactMatch);
+	if (ModeFSM)
+	{
+		ModeFSM->CancelFeaturesWithTag(TagToCancel, bRequireExactMatch);
+	}
 }
 
+bool UMoverComponent::HasFeaturesWithTag(FGameplayTag TagToMatch, bool bRequireExactMatch) const
+{
+	if (!ModeFSM)
+	{
+		return false;
+	}
+
+	return ModeFSM->HasFeaturesWithTag(TagToMatch, bRequireExactMatch);
+}
 
 void UMoverComponent::K2_QueueInstantMovementEffect(const int32& EffectAsRawData)
 {
@@ -1757,6 +1784,11 @@ void UMoverComponent::QueueInstantMovementEffect(TSharedPtr<FInstantMovementEffe
 
 void UMoverComponent::QueueInstantMovementEffect(const FScheduledInstantMovementEffect& InstantMovementEffect)
 {
+	if (!ModeFSM)
+	{
+		return;
+	}
+
 	// TODO Move QueueInstantMovementEffect to UMoverSimulation and implement differently in sync or async mode
 	if (IsInGameThread())
 	{
@@ -1868,7 +1900,10 @@ void UMoverComponent::QueueNextMode(FName DesiredModeName, bool bShouldReenter)
 
 void UMoverComponent::DoQueueNextMode(FName DesiredModeName, bool bShouldReenter)
 {
-	ModeFSM->QueueNextMode(DesiredModeName, bShouldReenter);
+	if (ModeFSM)
+	{
+		ModeFSM->QueueNextMode(DesiredModeName, bShouldReenter);
+	}
 }
 
 UBaseMovementMode* UMoverComponent::AddMovementModeFromClass(FName ModeName, TSubclassOf<UBaseMovementMode> MovementMode)
@@ -1890,6 +1925,12 @@ UBaseMovementMode* UMoverComponent::AddMovementModeFromClass(FName ModeName, TSu
 
 bool UMoverComponent::AddMovementModeFromObject(FName ModeName, UBaseMovementMode* MovementMode)
 {
+	if (!ModeFSM)
+	{
+		UE_LOG(LogMover, Warning, TEXT("Attempted to add %s movement mode before Movement Mode State Machine created %s. AddMovementModeFromObject did nothing. Check if mover component is initialized? (%s)"), *GetNameSafe(MovementMode), *GetNameSafe(this), *GetNameSafe(GetOwner()));
+		return false;
+	}
+
 	if (MovementMode)
 	{
 		if (MovementMode->GetClass()->HasAnyClassFlags(CLASS_Abstract))
@@ -1928,6 +1969,12 @@ bool UMoverComponent::AddMovementModeFromObject(FName ModeName, UBaseMovementMod
 
 bool UMoverComponent::RemoveMovementMode(FName ModeName)
 {
+	if (!ModeFSM)
+	{
+		UE_LOG(LogMover, Warning, TEXT("Attempted to remove %s movement mode before Movement Mode State Machine created %s. RemoveMovementMode did nothing. Check if mover component is initialized? (%s)"), *ModeName.ToString(), *GetNameSafe(this), *GetNameSafe(GetOwner()));
+		return false;
+	}
+
 	if (ModeFSM->GetCurrentModeName() == ModeName)
 	{
 		UE_LOG(LogMover, Warning, TEXT("The mode being removed (%s Movement Mode) is the mode this actor (%s) is currently in. It was removed but may cause issues. Consider waiting to remove the mode or queueing a different valid mode to avoid issues."), *ModeName.ToString(), *GetNameSafe(GetOwner()));
diff --git a/Engine/Plugins/Experimental/Mover/Source/Mover/Public/MovementModeStateMachine.h b/Engine/Plugins/Experimental/Mover/Source/Mover/Public/MovementModeStateMachine.h
--- a/Engine/Plugins/Experimental/Mover/Source/Mover/Public/MovementModeStateMachine.h
+++ b/Engine/Plugins/Experimental/Mover/Source/Mover/Public/MovementModeStateMachine.h
@@ -65,6 +65,7 @@ public:
  	UE_API const FMovementModifierBase* FindQueuedModifierByType(const UScriptStruct* ModifierType) const;
 
 	UE_API void CancelFeaturesWithTag(FGameplayTag TagToCancel, bool bRequireExactMatch=false);
+	UE_API bool HasFeaturesWithTag(FGameplayTag TagToMatch, bool bRequireExactMatch = false) const;
 
 	// This function is meant to be used only in async mode on the physics thread, not on the game thread
 	void QueueInstantMovementEffect_Internal(const FScheduledInstantMovementEffect& ScheduledInstantMovementEffect);
diff --git a/Engine/Plugins/Experimental/Mover/Source/Mover/Public/MoverComponent.h b/Engine/Plugins/Experimental/Mover/Source/Mover/Public/MoverComponent.h
--- a/Engine/Plugins/Experimental/Mover/Source/Mover/Public/MoverComponent.h
+++ b/Engine/Plugins/Experimental/Mover/Source/Mover/Public/MoverComponent.h
@@ -311,6 +311,12 @@ public:
 	UFUNCTION(BlueprintCallable, Category = Mover)
 	MOVER_API void CancelFeaturesWithTag(FGameplayTag TagToCancel, bool bRequireExactMatch=false);
 
+	/**
+	 * Check any active or queued movement features (layered moves, modifiers, etc.) for a matching gameplay tag.
+	 */
+	UFUNCTION(BlueprintCallable, Category = Mover)
+	MOVER_API bool HasFeaturesWithTag(FGameplayTag TagToCancel, bool bRequireExactMatch = false) const;
+
 	/**
 	 * Queue an Instant Movement Effect to start at the end of this frame or start of the next subtick - whichever happens first. This will clone whatever move you pass in, so you'll need to fully set it up before queuing.
 	 * @param InstantMovementEffect			The effect to queue, which must be a FInstantMovementEffect sub-type. 
diff --git a/Engine/Plugins/Experimental/MoverAnimNext/Source/MoverAnimNext/Private/Graph/RigUnit_MoverToggleRootMotion.cpp b/Engine/Plugins/Experimental/MoverAnimNext/Source/MoverAnimNext/Private/Graph/RigUnit_MoverToggleRootMotion.cpp
new file mode 100644
--- /dev/null
+++ b/Engine/Plugins/Experimental/MoverAnimNext/Source/MoverAnimNext/Private/Graph/RigUnit_MoverToggleRootMotion.cpp
@@ -0,0 +1,43 @@
+// Copyright Epic Games, Inc. All Rights Reserved.
+
+#include "RigUnit_MoverToggleRootMotion.h"
+
+#include "DefaultMovementSet/LayeredMoves/RootMotionAttributeLayeredMove.h"
+#include "MoverComponent.h"
+#include "Module/AnimNextModuleInstance.h"
+
+FRigUnit_MoverToggleRootMotion_Execute()
+{
+	TRACE_CPUPROFILER_EVENT_SCOPE(FRigUnit_MoverToggleRootMotion_Execute);
+	using namespace UE::UAF;
+
+	if (!MoverComponent || !MoverComponent->HasBeenInitialized())
+	{
+		return;
+	}
+
+	TWeakObjectPtr<UMoverComponent> WeakMoverComponent = MoverComponent;
+	auto MoverComponentTask = [bRootMotionEnabled, WeakMoverComponent]()
+	{
+		if (UMoverComponent* MoverComponent = WeakMoverComponent.Get())
+		{
+			if (bRootMotionEnabled)
+			{
+				MoverComponent->QueueLayeredMove(MakeShared<FLayeredMove_RootMotionAttribute>());
+			}
+			else
+			{
+				MoverComponent->CancelFeaturesWithTag(Mover_AnimRootMotion_MeshAttribute);
+			}
+		}
+	};
+
+	if (bToggleOnGameThread)
+	{
+		FAnimNextModuleInstance::RunTaskOnGameThread(MoverComponentTask);
+	}
+	else
+	{
+		MoverComponentTask();
+	}
+}
diff --git a/Engine/Plugins/Experimental/MoverAnimNext/Source/MoverAnimNext/Private/Graph/RigUnit_MoverToggleRootMotion.h b/Engine/Plugins/Experimental/MoverAnimNext/Source/MoverAnimNext/Private/Graph/RigUnit_MoverToggleRootMotion.h
new file mode 100644
--- /dev/null
+++ b/Engine/Plugins/Experimental/MoverAnimNext/Source/MoverAnimNext/Private/Graph/RigUnit_MoverToggleRootMotion.h
@@ -0,0 +1,45 @@
+// Copyright Epic Games, Inc. All Rights Reserved.
+
+#pragma once
+
+#include "AnimNextExecuteContext.h"
+#include "Graph/RigUnit_AnimNextBase.h"
+
+#include "RigUnit_MoverToggleRootMotion.generated.h"
+
+class UMoverComponent;
+
+/**
+ * RigVM Node that can be used to toggle mover root motion. The toggle will occur next frame.
+ * You should also update your animation tick dependencies to so that animation runs before mover when this is toggled.
+ * 
+ * For example. Within UAF during PrePhysics: 
+ * If (bRootMotionEnabled) -> AddDependency : MoverComponent, Before, SimulateMovment -> RemoveDependency : MoverComponent, After, SimulateMovment
+ * If (!bRootMotionEnabled) -> RemoveDependency : MoverComponent, Before, SimulateMovment -> AddDependency : MoverComponent, After, SimulateMovment
+ * 
+ * Where: If / AddDependency / RemoveDependency are nodes in your RigVM UAF PrePhysics graph.
+ */
+USTRUCT(meta = (DisplayName = "Mover Toggle Root Motion", Category = "Mover", NodeColor = "0, 1, 1", Keywords = "Mover, Root Motion, Enable, Disable"))
+struct FRigUnit_MoverToggleRootMotion : public FRigUnit_AnimNextBase
+{
+	GENERATED_BODY()
+
+	RIGVM_METHOD()
+	void Execute();
+
+	/** Mover to toggle root motion on */
+	UPROPERTY(EditAnywhere, Category = "Mover", meta = (Input))
+	TObjectPtr<UMoverComponent> MoverComponent = nullptr;
+
+	/** True if we want to enable root motion on the mover */
+	UPROPERTY(EditAnywhere, Category = "Mover", meta = (Input))
+	bool bRootMotionEnabled = false;
+
+	/** True if we want to delay the toggle to next frame on game thread. This allows for tick dependencies to change for next frame */
+	UPROPERTY(EditAnywhere, Category = "Mover", AdvancedDisplay, meta = (Input))
+	bool bToggleOnGameThread = true;
+
+	// Execution input / output
+	UPROPERTY(EditAnywhere, DisplayName = "Execute", Category = "Mover", meta = (Input, Output))
+	FAnimNextExecuteContext ExecuteContext;
+};
