[Backout] - CL46016661
#rnx
Original CL Desc
-----------------------------------------------------------------
Add RigVM nodes to enable / disable root motion on mover.
Add Mass ST Task to toggle root motion.
Add Mover Methods to query if a tag exists.
Add Mover FSM nullptr checks.
Fix uncessary / slow weakptr pins.

#tests PIE, Editor

diff --git a/Engine/Plugins/Experimental/Mover/Source/Mover/Private/MovementModeStateMachine.cpp b/Engine/Plugins/Experimental/Mover/Source/Mover/Private/MovementModeStateMachine.cpp
--- a/Engine/Plugins/Experimental/Mover/Source/Mover/Private/MovementModeStateMachine.cpp
+++ b/Engine/Plugins/Experimental/Mover/Source/Mover/Private/MovementModeStateMachine.cpp
@@ -664,43 +664,6 @@ void UMovementModeStateMachine::CancelFeaturesWithTag(FGameplayTag TagToCancel,
 
 }
 
-bool UMovementModeStateMachine::HasFeaturesWithTag(FGameplayTag TagToMatch, bool bRequireExactMatch) const
-{
-	bool bResult = false;
-
-	// Compare all matching queued movement features
-	{
-		FRWScopeLock QueueLock(ModifiersQueueLock, SLT_ReadOnly);
-		const TSharedPtr<FMovementModifierBase>* FindResult = QueuedMovementModifiers.FindByPredicate([TagToMatch, bRequireExactMatch] (const TSharedPtr<FMovementModifierBase>& Modifier)
-			{
-				return (Modifier.IsValid() && Modifier->HasGameplayTag(TagToMatch, bRequireExactMatch));
-			});
-		bResult = FindResult != nullptr;
-	}
-
-	if (!bResult)
-	{
-		FRWScopeLock QueueLock(LayeredMoveQueueLock, SLT_ReadOnly);
-		const TSharedPtr<FLayeredMoveBase>* FindResult = QueuedLayeredMoves.FindByPredicate([TagToMatch, bRequireExactMatch](const TSharedPtr<FLayeredMoveBase>& LayeredMove)
-			{
-				return (LayeredMove.IsValid() && LayeredMove->HasGameplayTag(TagToMatch, bRequireExactMatch));
-			});
-		bResult = FindResult != nullptr;
-	}
-
-	if (!bResult)
-	{
-		FRWScopeLock QueueLock(LayeredMoveQueueLock, SLT_ReadOnly);
-		const TSharedPtr<FLayeredMoveInstance>* FindResult = QueuedLayeredMoveInstances.FindByPredicate([TagToMatch, bRequireExactMatch](const TSharedPtr<FLayeredMoveInstance>& LayeredMoveInstance)
-			{
-				return (LayeredMoveInstance.IsValid() && LayeredMoveInstance->HasGameplayTag(TagToMatch, bRequireExactMatch));
-			});
-		bResult = FindResult != nullptr;
-	}
-
-	return bResult;
-}
-
 
 void UMovementModeStateMachine::ConstructDefaultModes()
 {
diff --git a/Engine/Plugins/Experimental/Mover/Source/Mover/Private/MoverComponent.cpp b/Engine/Plugins/Experimental/Mover/Source/Mover/Private/MoverComponent.cpp
--- a/Engine/Plugins/Experimental/Mover/Source/Mover/Private/MoverComponent.cpp
+++ b/Engine/Plugins/Experimental/Mover/Source/Mover/Private/MoverComponent.cpp
@@ -724,11 +724,6 @@ bool UMoverComponent::IsModifierActiveOrQueued(const FMovementModifierHandle& Mo
 
 const FMovementModifierBase* UMoverComponent::FindMovementModifier(const FMovementModifierHandle& ModifierHandle) const
 {
-	if (!ModeFSM)
-	{
-		return nullptr;
-	}
-
 	const FMoverSyncState& CachedSyncState = MoverSyncStateDoubleBuffer.GetReadable();
 	
 	// Check active modifiers for modifier handle
@@ -758,11 +753,6 @@ const FMovementModifierBase* UMoverComponent::FindMovementModifier(const FMoveme
 
 const FMovementModifierBase* UMoverComponent::FindMovementModifierByType(const UScriptStruct* DataStructType) const
 {
-	if (!ModeFSM)
-	{
-		return nullptr;
-	}
-
 	const FMoverSyncState& CachedSyncState = MoverSyncStateDoubleBuffer.GetReadable();
 	
 	// Check active modifiers for modifier handle
@@ -1619,10 +1609,7 @@ DEFINE_FUNCTION(UMoverComponent::execK2_QueueLayeredMove)
 
 void UMoverComponent::QueueLayeredMove(TSharedPtr<FLayeredMoveBase> LayeredMove)
 {	
-	if (ModeFSM)
-	{
-		ModeFSM->QueueLayeredMove(LayeredMove);
-	}
+	ModeFSM->QueueLayeredMove(LayeredMove);
 }
 
 FMovementModifierHandle UMoverComponent::K2_QueueMovementModifier(const int32& MoveAsRawData)
@@ -1660,34 +1647,20 @@ DEFINE_FUNCTION(UMoverComponent::execK2_QueueMovementModifier)
 
 FMovementModifierHandle UMoverComponent::QueueMovementModifier(TSharedPtr<FMovementModifierBase> Modifier)
 {
-	return ModeFSM ? ModeFSM->QueueMovementModifier(Modifier) : FMovementModifierHandle(MODIFIER_INVALID_HANDLE);
+	return ModeFSM->QueueMovementModifier(Modifier);
 }
 
 void UMoverComponent::CancelModifierFromHandle(FMovementModifierHandle ModifierHandle)
 {
-	if (ModeFSM)
-	{
-		ModeFSM->CancelModifierFromHandle(ModifierHandle);
-	}
+	ModeFSM->CancelModifierFromHandle(ModifierHandle);
 }
 
+
 void UMoverComponent::CancelFeaturesWithTag(FGameplayTag TagToCancel, bool bRequireExactMatch)
 {
-	if (ModeFSM)
-	{
-		ModeFSM->CancelFeaturesWithTag(TagToCancel, bRequireExactMatch);
-	}
+	ModeFSM->CancelFeaturesWithTag(TagToCancel, bRequireExactMatch);
 }
 
-bool UMoverComponent::HasFeaturesWithTag(FGameplayTag TagToMatch, bool bRequireExactMatch) const
-{
-	if (!ModeFSM)
-	{
-		return false;
-	}
-
-	return ModeFSM->HasFeaturesWithTag(TagToMatch, bRequireExactMatch);
-}
 
 void UMoverComponent::K2_QueueInstantMovementEffect(const int32& EffectAsRawData)
 {
@@ -1784,11 +1757,6 @@ void UMoverComponent::QueueInstantMovementEffect(TSharedPtr<FInstantMovementEffe
 
 void UMoverComponent::QueueInstantMovementEffect(const FScheduledInstantMovementEffect& InstantMovementEffect)
 {
-	if (!ModeFSM)
-	{
-		return;
-	}
-
 	// TODO Move QueueInstantMovementEffect to UMoverSimulation and implement differently in sync or async mode
 	if (IsInGameThread())
 	{
@@ -1900,10 +1868,7 @@ void UMoverComponent::QueueNextMode(FName DesiredModeName, bool bShouldReenter)
 
 void UMoverComponent::DoQueueNextMode(FName DesiredModeName, bool bShouldReenter)
 {
-	if (ModeFSM)
-	{
-		ModeFSM->QueueNextMode(DesiredModeName, bShouldReenter);
-	}
+	ModeFSM->QueueNextMode(DesiredModeName, bShouldReenter);
 }
 
 UBaseMovementMode* UMoverComponent::AddMovementModeFromClass(FName ModeName, TSubclassOf<UBaseMovementMode> MovementMode)
@@ -1925,12 +1890,6 @@ UBaseMovementMode* UMoverComponent::AddMovementModeFromClass(FName ModeName, TSu
 
 bool UMoverComponent::AddMovementModeFromObject(FName ModeName, UBaseMovementMode* MovementMode)
 {
-	if (!ModeFSM)
-	{
-		UE_LOG(LogMover, Warning, TEXT("Attempted to add %s movement mode before Movement Mode State Machine created %s. AddMovementModeFromObject did nothing. Check if mover component is initialized? (%s)"), *GetNameSafe(MovementMode), *GetNameSafe(this), *GetNameSafe(GetOwner()));
-		return false;
-	}
-
 	if (MovementMode)
 	{
 		if (MovementMode->GetClass()->HasAnyClassFlags(CLASS_Abstract))
@@ -1969,12 +1928,6 @@ bool UMoverComponent::AddMovementModeFromObject(FName ModeName, UBaseMovementMod
 
 bool UMoverComponent::RemoveMovementMode(FName ModeName)
 {
-	if (!ModeFSM)
-	{
-		UE_LOG(LogMover, Warning, TEXT("Attempted to remove %s movement mode before Movement Mode State Machine created %s. RemoveMovementMode did nothing. Check if mover component is initialized? (%s)"), *ModeName.ToString(), *GetNameSafe(this), *GetNameSafe(GetOwner()));
-		return false;
-	}
-
 	if (ModeFSM->GetCurrentModeName() == ModeName)
 	{
 		UE_LOG(LogMover, Warning, TEXT("The mode being removed (%s Movement Mode) is the mode this actor (%s) is currently in. It was removed but may cause issues. Consider waiting to remove the mode or queueing a different valid mode to avoid issues."), *ModeName.ToString(), *GetNameSafe(GetOwner()));
diff --git a/Engine/Plugins/Experimental/Mover/Source/Mover/Public/MovementModeStateMachine.h b/Engine/Plugins/Experimental/Mover/Source/Mover/Public/MovementModeStateMachine.h
--- a/Engine/Plugins/Experimental/Mover/Source/Mover/Public/MovementModeStateMachine.h
+++ b/Engine/Plugins/Experimental/Mover/Source/Mover/Public/MovementModeStateMachine.h
@@ -65,7 +65,6 @@ public:
  	UE_API const FMovementModifierBase* FindQueuedModifierByType(const UScriptStruct* ModifierType) const;
 
 	UE_API void CancelFeaturesWithTag(FGameplayTag TagToCancel, bool bRequireExactMatch=false);
-	UE_API bool HasFeaturesWithTag(FGameplayTag TagToMatch, bool bRequireExactMatch = false) const;
 
 	// This function is meant to be used only in async mode on the physics thread, not on the game thread
 	void QueueInstantMovementEffect_Internal(const FScheduledInstantMovementEffect& ScheduledInstantMovementEffect);
diff --git a/Engine/Plugins/Experimental/Mover/Source/Mover/Public/MoverComponent.h b/Engine/Plugins/Experimental/Mover/Source/Mover/Public/MoverComponent.h
--- a/Engine/Plugins/Experimental/Mover/Source/Mover/Public/MoverComponent.h
+++ b/Engine/Plugins/Experimental/Mover/Source/Mover/Public/MoverComponent.h
@@ -311,12 +311,6 @@ public:
 	UFUNCTION(BlueprintCallable, Category = Mover)
 	MOVER_API void CancelFeaturesWithTag(FGameplayTag TagToCancel, bool bRequireExactMatch=false);
 
-	/**
-	 * Check any active or queued movement features (layered moves, modifiers, etc.) for a matching gameplay tag.
-	 */
-	UFUNCTION(BlueprintCallable, Category = Mover)
-	MOVER_API bool HasFeaturesWithTag(FGameplayTag TagToCancel, bool bRequireExactMatch = false) const;
-
 	/**
 	 * Queue an Instant Movement Effect to start at the end of this frame or start of the next subtick - whichever happens first. This will clone whatever move you pass in, so you'll need to fully set it up before queuing.
 	 * @param InstantMovementEffect			The effect to queue, which must be a FInstantMovementEffect sub-type. 
diff --git a/Engine/Plugins/Experimental/MoverAnimNext/Source/MoverAnimNext/Private/Graph/RigUnit_MoverToggleRootMotion.cpp b/Engine/Plugins/Experimental/MoverAnimNext/Source/MoverAnimNext/Private/Graph/RigUnit_MoverToggleRootMotion.cpp
deleted file mode 100644
diff --git a/Engine/Plugins/Experimental/MoverAnimNext/Source/MoverAnimNext/Private/Graph/RigUnit_MoverToggleRootMotion.h b/Engine/Plugins/Experimental/MoverAnimNext/Source/MoverAnimNext/Private/Graph/RigUnit_MoverToggleRootMotion.h
deleted file mode 100644
