MoverSimulationTypes: Fix FName definitions in header

diff --git a/Engine/Plugins/Experimental/Mover/Source/Mover/Private/MoverSimulationTypes.cpp b/Engine/Plugins/Experimental/Mover/Source/Mover/Private/MoverSimulationTypes.cpp
--- a/Engine/Plugins/Experimental/Mover/Source/Mover/Private/MoverSimulationTypes.cpp
+++ b/Engine/Plugins/Experimental/Mover/Source/Mover/Private/MoverSimulationTypes.cpp
@@ -8,6 +8,25 @@
 
 #include UE_INLINE_GENERATED_CPP_BY_NAME(MoverSimulationTypes)
 
+namespace DefaultModeNames
+{
+	const FName Walking = TEXT("Walking");
+	const FName Falling = TEXT("Falling");
+	const FName Flying  = TEXT("Flying");
+	const FName Swimming  = TEXT("Swimming");
+}
+
+namespace CommonBlackboard
+{
+	const FName LastFloorResult = TEXT("LastFloor");
+	const FName LastWaterResult = TEXT("LastWater");
+	const FName LastFoundDynamicMovementBase = TEXT("LastFoundDynamicMovementBase");
+	const FName LastAppliedDynamicMovementBase = TEXT("LastAppliedDynamicMovementBase");
+	const FName TimeSinceSupported = TEXT("TimeSinceSupported");
+
+	const FName LastModeChangeRecord = TEXT("LastModeChangeRecord");
+}
+
 UScriptStruct* FMoverSimulationEventData::GetScriptStruct() const
 {
 	checkf(false, TEXT("%hs is being called erroneously. This must be overridden in derived types!"), __FUNCTION__);
diff --git a/Engine/Plugins/Experimental/Mover/Source/Mover/Public/MoverSimulationTypes.h b/Engine/Plugins/Experimental/Mover/Source/Mover/Public/MoverSimulationTypes.h
--- a/Engine/Plugins/Experimental/Mover/Source/Mover/Public/MoverSimulationTypes.h
+++ b/Engine/Plugins/Experimental/Mover/Source/Mover/Public/MoverSimulationTypes.h
@@ -20,22 +20,22 @@
 // Names for our default modes
 namespace DefaultModeNames
 {
-	const FName Walking = TEXT("Walking");
-	const FName Falling = TEXT("Falling");
-	const FName Flying  = TEXT("Flying");
-	const FName Swimming  = TEXT("Swimming");
+	extern MOVER_API const FName Walking;
+	extern MOVER_API const FName Falling;
+	extern MOVER_API const FName Flying;
+	extern MOVER_API const FName Swimming;
 }
 
 // Commonly-used blackboard object keys
 namespace CommonBlackboard
 {
-	const FName LastFloorResult = TEXT("LastFloor");
-	const FName LastWaterResult = TEXT("LastWater");
-	const FName LastFoundDynamicMovementBase = TEXT("LastFoundDynamicMovementBase");
-	const FName LastAppliedDynamicMovementBase = TEXT("LastAppliedDynamicMovementBase");
-	const FName TimeSinceSupported = TEXT("TimeSinceSupported");
+	extern MOVER_API const FName LastFloorResult;
+	extern MOVER_API const FName LastWaterResult;
+	extern MOVER_API const FName LastFoundDynamicMovementBase;
+	extern MOVER_API const FName LastAppliedDynamicMovementBase;
+	extern MOVER_API const FName TimeSinceSupported;
 
-	const FName LastModeChangeRecord = TEXT("LastModeChangeRecord");
+	extern MOVER_API const FName LastModeChangeRecord;
 }
 
 
