Mover: Adding a comment

diff --git a/Engine/Plugins/Experimental/Mover/Source/Mover/Public/InstantMovementEffect.h b/Engine/Plugins/Experimental/Mover/Source/Mover/Public/InstantMovementEffect.h
--- a/Engine/Plugins/Experimental/Mover/Source/Mover/Public/InstantMovementEffect.h
+++ b/Engine/Plugins/Experimental/Mover/Source/Mover/Public/InstantMovementEffect.h
@@ -67,6 +67,7 @@ struct FInstantMovementEffect
 
 	virtual void AddReferencedObjects(class FReferenceCollector& Collector) {}
 
+	// Return false is the effect didn't apply at all, true if anything was modified, especially the output state's movement mode
 	virtual bool ApplyMovementEffect(FApplyMovementEffectParams& ApplyEffectParams, FMoverSyncState& OutputState) { return false; }
 	virtual bool ApplyMovementEffect_Async(FApplyMovementEffectParams_Async& ApplyEffectParams, FMoverSyncState& OutputState) { return false; }
 };
