Mover: fixes bug where events with close timestamps could be processed out of order. Changed method of pruning event arrays so that order is preserved.

diff --git a/Engine/Plugins/Experimental/Mover/Source/Mover/Private/MoverSimulationTypes.cpp b/Engine/Plugins/Experimental/Mover/Source/Mover/Private/MoverSimulationTypes.cpp
--- a/Engine/Plugins/Experimental/Mover/Source/Mover/Private/MoverSimulationTypes.cpp
+++ b/Engine/Plugins/Experimental/Mover/Source/Mover/Private/MoverSimulationTypes.cpp
@@ -99,7 +99,7 @@ namespace UE::Mover
 			}
 
 			// If resimulating find any events stored after the time step and remove
-			Events.RemoveAllSwap([&InTimeStep](const TSharedPtr<FMoverSimulationEventData> Event) {
+			Events.RemoveAll([&InTimeStep](const TSharedPtr<FMoverSimulationEventData> Event) {
 				return !Event.IsValid() || (Event->EventTimeMs >= InTimeStep.BaseSimTimeMs);
 				});
 		}
@@ -150,7 +150,7 @@ namespace UE::Mover
 			}
 		}
 
-		Events.RemoveAllSwap([](const TSharedPtr<FMoverSimulationEventData> Event) {
+		Events.RemoveAll([](const TSharedPtr<FMoverSimulationEventData> Event) {
 			return !Event.IsValid();
 			});
 	}
