Mover: fix case where we were using the wrong timestep copy that did not have all the most up-to-date settings

#rnx

diff --git a/Engine/Plugins/Experimental/Mover/Source/Mover/Private/MoverComponent.cpp b/Engine/Plugins/Experimental/Mover/Source/Mover/Private/MoverComponent.cpp
--- a/Engine/Plugins/Experimental/Mover/Source/Mover/Private/MoverComponent.cpp
+++ b/Engine/Plugins/Experimental/Mover/Source/Mover/Private/MoverComponent.cpp
@@ -538,7 +538,7 @@ void UMoverComponent::SimulationTick(const FMoverTimeStep& InTimeStep, const FMo
 	if (bHasJustRolledBack)
 	{
 		MoverTimeStep.bIsFirstResimFrame = MoverTimeStep.bIsResimulating;
-		ProcessFirstSimTickAfterRollback(InTimeStep);
+		ProcessFirstSimTickAfterRollback(MoverTimeStep);
 	}
 
 	PreSimulationTick(MoverTimeStep, SimInput.InputCmd);
