Mover: converted flags that could only be authored by the internal simulation code to be read-only in blueprint

#rnx

diff --git a/Engine/Plugins/Experimental/Mover/Source/Mover/Public/MoverTypes.h b/Engine/Plugins/Experimental/Mover/Source/Mover/Public/MoverTypes.h
--- a/Engine/Plugins/Experimental/Mover/Source/Mover/Public/MoverTypes.h
+++ b/Engine/Plugins/Experimental/Mover/Source/Mover/Public/MoverTypes.h
@@ -115,11 +115,11 @@ struct FMoverTimeStep
 	float StepMs=1.0;
 
 	// Indicates whether this time step is re-simulating based on prior inputs, such as during a correction
-	UPROPERTY(EditAnywhere, BlueprintReadWrite, Category=Mover)
+	UPROPERTY(EditAnywhere, BlueprintReadOnly, Category=Mover)
 	uint8 bIsResimulating : 1 = 0;
 
 	// If resimulating, is this the first frame after the rollback?
-	UPROPERTY(EditAnywhere, BlueprintReadWrite, Category=Mover)
+	UPROPERTY(EditAnywhere, BlueprintReadOnly, Category=Mover)
 	uint8 bIsFirstResimFrame : 1 = 0;
 };
 
