Mover: fixed issue where a stale movement base bone name could cause continuous corrections. This occurs when transitioning from a skeletal mesh base with a valid bone to a static mesh base with no bone.

diff --git a/Engine/Plugins/Experimental/Mover/Source/Mover/Private/MoverDataModelTypes.cpp b/Engine/Plugins/Experimental/Mover/Source/Mover/Private/MoverDataModelTypes.cpp
--- a/Engine/Plugins/Experimental/Mover/Source/Mover/Private/MoverDataModelTypes.cpp
+++ b/Engine/Plugins/Experimental/Mover/Source/Mover/Private/MoverDataModelTypes.cpp
@@ -192,6 +192,7 @@ bool FMoverDefaultSyncState::NetSerialize(FArchive& Ar, UPackageMap* Map, bool&
 	else if (Ar.IsLoading())
 	{
 		MovementBase = nullptr;
+		MovementBaseBoneName = NAME_None;
 	}
 
 	bOutSuccess = true;
