Mover: Add a util function for scheduling an instant movement effect for the very next frame. This is useful when a mover component wants to queue an instant movement effect on another mover component, to happen as soon as possible, but without suffering from tick order dependency. Without that, it can happen on a different frame between client and server because the tick order of the movers is not the same.

#Synced-CL 48072262

diff --git a/Engine/Plugins/Experimental/Mover/Source/Mover/Private/MoverSimulationTypes.cpp b/Engine/Plugins/Experimental/Mover/Source/Mover/Private/MoverSimulationTypes.cpp
--- a/Engine/Plugins/Experimental/Mover/Source/Mover/Private/MoverSimulationTypes.cpp
+++ b/Engine/Plugins/Experimental/Mover/Source/Mover/Private/MoverSimulationTypes.cpp
@@ -59,6 +59,27 @@ FScheduledInstantMovementEffect FScheduledInstantMovementEffect::ScheduleEffect(
 	}
 }
 
+FScheduledInstantMovementEffect FScheduledInstantMovementEffect::ScheduleEffectForNextFrame(UWorld* World, const FMoverTimeStep& TimeStep, TSharedPtr<FInstantMovementEffect> InstantMovementEffect)
+{
+	FPhysScene* Scene = World ? World->GetPhysicsScene() : nullptr;
+	Chaos::FPhysicsSolver* Solver = Scene ? Scene->GetSolver() : nullptr;
+	bool bIsFixedDt = Solver ? Solver->IsUsingFixedDt() : false;
+	if (bIsFixedDt)
+	{
+		int32 ServerFrame = TimeStep.ServerFrame + 1;
+
+		return FScheduledInstantMovementEffect(/* ExecutionServerFrame = */ ServerFrame, /* ExecutionServerTime = */ 0.0, /* bIsFixedDt = */ true, InstantMovementEffect);
+	}
+	else
+	{
+		AGameStateBase* GameState = World ? World->GetGameState() : nullptr;
+		double ServerTime = GameState ? GameState->GetServerWorldTimeSeconds() : 0.0;
+		ServerTime += UE_KINDA_SMALL_NUMBER;
+
+		return FScheduledInstantMovementEffect(/* ExecutionServerFrame = */ INDEX_NONE, /* ExecutionServerTime = */ ServerTime, /* bIsFixedDt = */ false, InstantMovementEffect);
+	}
+}
+
 namespace UE::Mover
 {
 	void FSimulationOutputData::Reset()
diff --git a/Engine/Plugins/Experimental/Mover/Source/Mover/Public/MoverSimulationTypes.h b/Engine/Plugins/Experimental/Mover/Source/Mover/Public/MoverSimulationTypes.h
--- a/Engine/Plugins/Experimental/Mover/Source/Mover/Public/MoverSimulationTypes.h
+++ b/Engine/Plugins/Experimental/Mover/Source/Mover/Public/MoverSimulationTypes.h
@@ -85,7 +85,9 @@ struct FScheduledInstantMovementEffect
 	*   @param InstantMovementEffect the effect to schedule
 	*   @param SchedulingDelaySeconds Scheduling delay to ensure it applies on all end points on the same frame (this is only perfectly accurate when simulation dt is fixed)
 	*/
-	static FScheduledInstantMovementEffect ScheduleEffect(UWorld* World, const FMoverTimeStep& TimeStep, TSharedPtr<FInstantMovementEffect> InstantMovementEffect, float SchedulingDelaySeconds = 0.0f);
+	MOVER_API static FScheduledInstantMovementEffect ScheduleEffect(UWorld* World, const FMoverTimeStep& TimeStep, TSharedPtr<FInstantMovementEffect> InstantMovementEffect, float SchedulingDelaySeconds = 0.0f);
+	// Same as above but schedules the effect to run on the next frame. This is best used for fixed tick mode, it will add a small SchedulingDelaySeconds in the variable tick case
+	MOVER_API static FScheduledInstantMovementEffect ScheduleEffectForNextFrame(UWorld* World, const FMoverTimeStep& TimeStep, TSharedPtr<FInstantMovementEffect> InstantMovementEffect);
 
 	bool ShouldExecuteAtFrame(int32 CurrentServerFrame) const
 	{
