[Mover] Add validity check to simulation output data. This allows us to tell if an output has been correctly generated before recording it

#Synced-CL 48348108

diff --git a/Engine/Plugins/Experimental/ChaosMover/Source/ChaosMover/Private/Backends/ChaosMoverBackend.cpp b/Engine/Plugins/Experimental/ChaosMover/Source/ChaosMover/Private/Backends/ChaosMoverBackend.cpp
--- a/Engine/Plugins/Experimental/ChaosMover/Source/ChaosMover/Private/Backends/ChaosMoverBackend.cpp
+++ b/Engine/Plugins/Experimental/ChaosMover/Source/ChaosMover/Private/Backends/ChaosMoverBackend.cpp
@@ -614,7 +614,10 @@ void UChaosMoverBackendComponent::ConsumeOutputData(const FMoverTimeStep& TimeSt
 		return;
 	}
 
-	SimOutputRecord.Add(TimeStep, OutputData);
+	if (OutputData.IsValid())
+	{
+		SimOutputRecord.Add(TimeStep, OutputData);
+	}
 }
 
 void UChaosMoverBackendComponent::FinalizeFrame(double ResultsTimeInMs)
diff --git a/Engine/Plugins/Experimental/ChaosMover/Source/ChaosMover/Private/ChaosMoverSimulation.cpp b/Engine/Plugins/Experimental/ChaosMover/Source/ChaosMover/Private/ChaosMoverSimulation.cpp
--- a/Engine/Plugins/Experimental/ChaosMover/Source/ChaosMover/Private/ChaosMoverSimulation.cpp
+++ b/Engine/Plugins/Experimental/ChaosMover/Source/ChaosMover/Private/ChaosMoverSimulation.cpp
@@ -840,6 +840,9 @@ void UChaosMoverSimulation::OnPostSimulationTick(const FMoverTimeStep& TimeStep,
 
 	// Extract the events into the output data
 	UpdateEvents(TimeStep, OutputData);
+
+	// Set the output data to be valid if we've reached here
+	OutputData.bIsValid = true;
 }
 
 void UChaosMoverSimulation::UpdateEvents(const FMoverTimeStep& TimeStep, UE::ChaosMover::FSimulationOutputData& OutputData)
diff --git a/Engine/Plugins/Experimental/Mover/Source/Mover/Private/MoverSimulationTypes.cpp b/Engine/Plugins/Experimental/Mover/Source/Mover/Private/MoverSimulationTypes.cpp
--- a/Engine/Plugins/Experimental/Mover/Source/Mover/Private/MoverSimulationTypes.cpp
+++ b/Engine/Plugins/Experimental/Mover/Source/Mover/Private/MoverSimulationTypes.cpp
@@ -88,6 +88,7 @@ namespace UE::Mover
 		LastUsedInputCmd.Reset();
 		AdditionalOutputData.Empty();
 		Events.Empty();
+		bIsValid = false;
 	}
 
 	void FSimulationOutputData::Interpolate(const FSimulationOutputData& From, const FSimulationOutputData& To, float Alpha, double SimTimeMs)
diff --git a/Engine/Plugins/Experimental/Mover/Source/Mover/Public/MoverSimulationTypes.h b/Engine/Plugins/Experimental/Mover/Source/Mover/Public/MoverSimulationTypes.h
--- a/Engine/Plugins/Experimental/Mover/Source/Mover/Public/MoverSimulationTypes.h
+++ b/Engine/Plugins/Experimental/Mover/Source/Mover/Public/MoverSimulationTypes.h
@@ -699,11 +699,16 @@ namespace UE::Mover
 	{
 		MOVER_API void Reset();
 		MOVER_API void Interpolate(const FSimulationOutputData& From, const FSimulationOutputData& To, float Alpha, double SimTimeMs);
+		MOVER_API bool IsValid() const
+		{
+			return bIsValid && (SyncState.SyncStateCollection.GetDataArray().Num() > 0);
+		}
 
 		FMoverSyncState SyncState;
 		FMoverInputCmdContext LastUsedInputCmd;
 		FMoverDataCollection AdditionalOutputData;
 		TArray<TSharedPtr<FMoverSimulationEventData>> Events;
+		bool bIsValid = false;
 	};
 
 	class FSimulationOutputRecord
