[Chaos Mover] Fix unresolved external symbol

#rnx
#Synced-CL 48348108

diff --git a/Engine/Plugins/Experimental/Mover/Source/Mover/Public/MoverSimulationTypes.h b/Engine/Plugins/Experimental/Mover/Source/Mover/Public/MoverSimulationTypes.h
--- a/Engine/Plugins/Experimental/Mover/Source/Mover/Public/MoverSimulationTypes.h
+++ b/Engine/Plugins/Experimental/Mover/Source/Mover/Public/MoverSimulationTypes.h
@@ -699,7 +699,7 @@ namespace UE::Mover
 	{
 		MOVER_API void Reset();
 		MOVER_API void Interpolate(const FSimulationOutputData& From, const FSimulationOutputData& To, float Alpha, double SimTimeMs);
-		MOVER_API bool IsValid() const
+		bool IsValid() const
 		{
 			return bIsValid && (SyncState.SyncStateCollection.GetDataArray().Num() > 0);
 		}
