Mover: fix bug where a tag cancellation request would kill a layered move queued AFTER the request was made

#Synced-CL 48101458

diff --git a/Engine/Plugins/Experimental/Mover/Source/Mover/Private/LayeredMove.cpp b/Engine/Plugins/Experimental/Mover/Source/Mover/Private/LayeredMove.cpp
--- a/Engine/Plugins/Experimental/Mover/Source/Mover/Private/LayeredMove.cpp
+++ b/Engine/Plugins/Experimental/Mover/Source/Mover/Private/LayeredMove.cpp
@@ -137,10 +137,14 @@ void FLayeredMoveGroup::QueueLayeredMove(TSharedPtr<FLayeredMoveBase> Move)
 
 void FLayeredMoveGroup::CancelMovesByTag(FGameplayTag Tag, bool bRequireExactMatch)
 {
+	// Queued moves can be removed now before they become active
+	QueuedLayeredMoves.RemoveAll([Tag, bRequireExactMatch](const TSharedPtr<FLayeredMoveBase>& Move)
+		{
+			return (Move.IsValid() && Move->HasGameplayTag(Tag, bRequireExactMatch));
+		});
 
-	// Schedule a tag cancellation request, to be handled during simulation
+	// Schedule a tag cancellation request, to catch any active moves during simulation
 	TagCancellationRequests.Add(TPair<FGameplayTag, bool>(Tag, bRequireExactMatch));
-
 }
 
 TArray<TSharedPtr<FLayeredMoveBase>> FLayeredMoveGroup::GenerateActiveMoves(const FMoverTimeStep& TimeStep, const UMoverComponent* MoverComp, UMoverBlackboard* SimBlackboard)
@@ -392,11 +396,7 @@ void FLayeredMoveGroup::FlushMoveArrays(const UMoverComponent* MoverComp, UMover
 			const FGameplayTag TagToMatch = CancelRequest.Key;
 			const bool bRequireExactMatch = CancelRequest.Value;
 
-			QueuedLayeredMoves.RemoveAll([TagToMatch, bRequireExactMatch](const TSharedPtr<FLayeredMoveBase>& Move)
-				{
-					return (Move.IsValid() && Move->HasGameplayTag(TagToMatch, bRequireExactMatch));
-				});
-
+			// Only removing active ones because queued ones were cancelled at the time the request was added
 			ActiveLayeredMoves.RemoveAll([MoverComp, SimBlackboard, CurrentSimTimeMs, &bResidualVelocityOverridden, &bClampVelocityOverridden, bIsAsync, TagToMatch, bRequireExactMatch, this] (const TSharedPtr<FLayeredMoveBase>& Move)
 				{
 					if (Move.IsValid() && Move->HasGameplayTag(TagToMatch, bRequireExactMatch))
diff --git a/Engine/Plugins/Experimental/Mover/Source/Mover/Private/LayeredMoveGroup.cpp b/Engine/Plugins/Experimental/Mover/Source/Mover/Private/LayeredMoveGroup.cpp
--- a/Engine/Plugins/Experimental/Mover/Source/Mover/Private/LayeredMoveGroup.cpp
+++ b/Engine/Plugins/Experimental/Mover/Source/Mover/Private/LayeredMoveGroup.cpp
@@ -263,8 +263,13 @@ const FLayeredMoveInstance* FLayeredMoveInstanceGroup::PrivateFindQueuedMove(con
 
 void FLayeredMoveInstanceGroup::CancelMovesByTag(FGameplayTag Tag, bool bRequireExactMatch)
 {
+	// Queued moves can be removed now before they become active
+	QueuedMoves.RemoveAll([Tag, bRequireExactMatch](const TSharedPtr<FLayeredMoveInstance>& Move)
+		{
+			return (!Move.IsValid() || Move->HasGameplayTag(Tag, bRequireExactMatch));
+		});
 
-	// Schedule a tag cancellation request, to be handled during simulation
+	// Schedule a tag cancellation request, to handle any active moves during simulation
 	TagCancellationRequests.Add(TPair<FGameplayTag, bool>(Tag, bRequireExactMatch));
 
 }
@@ -282,12 +287,10 @@ void FLayeredMoveInstanceGroup::FlushMoveArrays(const FMoverTimeStep& TimeStep,
 			const FGameplayTag TagToMatch = CancelRequest.Key;
 			const bool bRequireExactMatch = CancelRequest.Value;
 
-			QueuedMoves.RemoveAll([TagToMatch, bRequireExactMatch](const TSharedPtr<FLayeredMoveInstance>& Move)
-				{
-					return (!Move.IsValid() || Move->HasGameplayTag(TagToMatch, bRequireExactMatch));
-				});
+
 
 			// Process completion of any active moves that being canceled
+			// Only removing active ones because queued ones were cancelled at the time the request was added
 			ActiveMoves.RemoveAll([&TimeStep, &SimBlackboard, &bResidualVelocityOverridden, &bClampVelocityOverridden, TagToMatch, bRequireExactMatch, this] (const TSharedPtr<FLayeredMoveInstance>& Move)
 				{
 					if (Move.IsValid() && Move->HasGameplayTag(TagToMatch, bRequireExactMatch))
