Mover: Fixed up incorrect UE_REQUIRES arguements in GetMoverComponent and GetMoverComponentChecked functions in MovementModes

#Synced-CL 49033279

diff --git a/Engine/Plugins/Experimental/Mover/Source/Mover/Public/MovementMode.h b/Engine/Plugins/Experimental/Mover/Source/Mover/Public/MovementMode.h
--- a/Engine/Plugins/Experimental/Mover/Source/Mover/Public/MovementMode.h
+++ b/Engine/Plugins/Experimental/Mover/Source/Mover/Public/MovementMode.h
@@ -69,7 +69,7 @@ public:
 	 * Gets the outer mover component of the indicated type. Does not check on the type or the presence of the MoverComp outer. Safe to call on CDOs.
 	 * Note: Since UBaseMovementMode is declared "Within = MoverComponent", all instances of a mode except the CDO are guaranteed to have a valid MoverComponent outer.
 	 */
-	template<typename MoverT = UMoverComponent UE_REQUIRES(std::is_base_of_v<MoverT, UMoverComponent>)>
+	template<typename MoverT = UMoverComponent UE_REQUIRES(std::is_base_of_v<UMoverComponent, MoverT>)>
 	MoverT* GetMoverComponent() const
 	{
 		return Cast<MoverT>(GetOuter());
@@ -79,7 +79,7 @@ public:
 	 * Gets the outer mover component of the indicated type, checked for validity.
 	 * Note: Since UBaseMovementMode is declared "Within = MoverComponent", all instances of a mode except the CDO are guaranteed to have a valid MoverComponent outer.
 	 */
-	template<typename MoverT = UMoverComponent UE_REQUIRES(std::is_base_of_v<MoverT, UMoverComponent>)>
+	template<typename MoverT = UMoverComponent UE_REQUIRES(std::is_base_of_v<UMoverComponent, MoverT>)>
 	MoverT& GetMoverComponentChecked() const
 	{
 		return *CastChecked<MoverT>(GetOuterUMoverComponent());
