Mover: Fix UNavMoverComponent::GetFeetLocationBased nullptr MovementBase access and invalid supplied data causing occasional EPathFollowingResult::Blocked for AI agents due to bad samples.

#Synced-CL 49033279

diff --git a/Engine/Plugins/Experimental/Mover/Source/Mover/Private/DefaultMovementSet/NavMoverComponent.cpp b/Engine/Plugins/Experimental/Mover/Source/Mover/Private/DefaultMovementSet/NavMoverComponent.cpp
--- a/Engine/Plugins/Experimental/Mover/Source/Mover/Private/DefaultMovementSet/NavMoverComponent.cpp
+++ b/Engine/Plugins/Experimental/Mover/Source/Mover/Private/DefaultMovementSet/NavMoverComponent.cpp
@@ -140,21 +140,31 @@ FVector UNavMoverComponent::GetFeetLocationAt(FVector ComponentLocation) const
 
 FBasedPosition UNavMoverComponent::GetFeetLocationBased() const
 {
-	FBasedPosition BasedPosition(NULL, GetFeetLocation());
-	
+	FBasedPosition BasedPosition;
 	if (MoverComponent.IsValid())
 	{
+		bool bSetFromMoverBlackboard = false;
 		if (const UMoverBlackboard* Blackboard = MoverComponent->GetSimBlackboard())
 		{
 			FRelativeBaseInfo MovementBaseInfo;
 			if (Blackboard->TryGet(CommonBlackboard::LastFoundDynamicMovementBase, MovementBaseInfo)) 
 			{
-				BasedPosition.Base = MovementBaseInfo.MovementBase->GetOwner();
-				BasedPosition.Position = MovementBaseInfo.Location;
-				BasedPosition.CachedBaseLocation = MovementBaseInfo.ContactLocalPosition;
-				BasedPosition.CachedBaseRotation = MovementBaseInfo.Rotation.Rotator();
+				if (MovementBaseInfo.HasRelativeInfo())
+				{
+					bSetFromMoverBlackboard = true;
+					BasedPosition.Base = MovementBaseInfo.MovementBase->GetOwner();
+					BasedPosition.Position = MovementBaseInfo.ContactLocalPosition;
+					BasedPosition.CachedBaseLocation = MovementBaseInfo.Location;
+					BasedPosition.CachedBaseRotation = MovementBaseInfo.Rotation.Rotator();
+					BasedPosition.CachedTransPosition = GetFeetLocation();
+				}
 			}
 		}
+		
+		if (!bSetFromMoverBlackboard)
+		{
+			BasedPosition.Set(nullptr, GetFeetLocation());
+		}
 	}
 
 	return BasedPosition;
