Mover: guarding against trajectory prediction using a not-totally-initialized MoverComponent, and adding warning information if it's attempted.

#Synced-CL 49321623

diff --git a/Engine/Plugins/Experimental/Mover/Source/Mover/Private/MoverComponent.cpp b/Engine/Plugins/Experimental/Mover/Source/Mover/Private/MoverComponent.cpp
--- a/Engine/Plugins/Experimental/Mover/Source/Mover/Private/MoverComponent.cpp
+++ b/Engine/Plugins/Experimental/Mover/Source/Mover/Private/MoverComponent.cpp
@@ -191,6 +191,7 @@ void UMoverComponent::UninitializeComponent()
 	{
 		ModeFSM->ClearAllMovementModes();
 		ModeFSM->ClearAllGlobalTransitions();
+		ModeFSM = nullptr;
 	}
 
 	Super::UninitializeComponent();
@@ -2426,6 +2427,8 @@ TArray<FTrajectorySampleInfo> UMoverComponent::GetPredictedTrajectory(FMoverPred
 		}
 	}
 
+	UE_LOG(LogMover, Warning, TEXT("Cannot predict trajectory on a MoverComponent that is not initialized. Returning defaulted samples. Actor: %s  Initialized? %i BegunPlay? %i BeingDestroyed? %i"), *GetNameSafe(GetOwner()), HasBeenInitialized(), HasBegunPlay(), IsBeingDestroyed());
+
 	TArray<FTrajectorySampleInfo> BlankDefaultSamples;
 	BlankDefaultSamples.AddDefaulted(PredictionParams.NumPredictionSamples);
 	return BlankDefaultSamples;
