Mover: Adding read write lock when Queueing LayeredMoveInstances

#Synced-CL 49333356

diff --git a/Engine/Plugins/Experimental/Mover/Source/Mover/Private/MovementModeStateMachine.cpp b/Engine/Plugins/Experimental/Mover/Source/Mover/Private/MovementModeStateMachine.cpp
--- a/Engine/Plugins/Experimental/Mover/Source/Mover/Private/MovementModeStateMachine.cpp
+++ b/Engine/Plugins/Experimental/Mover/Source/Mover/Private/MovementModeStateMachine.cpp
@@ -556,6 +556,7 @@ void UMovementModeStateMachine::QueueInstantMovementEffect_Internal(const FSched
 
 void UMovementModeStateMachine::QueueActiveLayeredMove(const TSharedPtr<FLayeredMoveInstance>& LayeredMove)
 {
+	UE::TRWScopeLock QueueLock(LayeredMoveQueueLock, SLT_Write);
 	QueuedLayeredMoveInstances.Add(LayeredMove);
 }
 
