Mover: implementing single-entry values for rollback blackboard, along with a policy to ignore rollbacks.
Includes some changes for const-correctness.

#Synced-CL 49592226

diff --git a/Engine/Plugins/Experimental/Mover/Source/Mover/Private/MoveLibrary/RollbackBlackboard.cpp b/Engine/Plugins/Experimental/Mover/Source/Mover/Private/MoveLibrary/RollbackBlackboard.cpp
--- a/Engine/Plugins/Experimental/Mover/Source/Mover/Private/MoveLibrary/RollbackBlackboard.cpp
+++ b/Engine/Plugins/Experimental/Mover/Source/Mover/Private/MoveLibrary/RollbackBlackboard.cpp
@@ -22,6 +22,11 @@ uint32 URollbackBlackboard::BlackboardEntryBase::ComputeBufferSize(EBlackboardSi
 			// TODO: Warn if FixedBufferSize is invalid
 			BufferSize = FMath::Max(BufferSize, FixedBufferSize);
 			break;
+		
+		case EBlackboardSizingPolicy::SingleEntry:
+			// TODO: only allocate 1 element if our simulation is not async
+			BufferSize = 2;		// Even for a single value entry, we need 2 slots in async simulations to avoid potential thread contention
+			break;
 
 		default:	//TODO: warn about an unhandled policy type
 			break;
@@ -40,7 +45,12 @@ URollbackBlackboard::BlackboardEntryBase::BlackboardEntryBase(const EntrySetting
 
 void URollbackBlackboard::BlackboardEntryBase::RollBack(uint32 NewPendingFrame)
 {
-	// Goal: adjust entry to point at the value from the prior frame. May make the entry invalidated, if there were no values that old.
+	if (Settings.RollbackPolicy == EBlackboardRollbackPolicy::IgnoreRollback)
+	{
+		return;
+	}
+
+	// Goal: adjust entry indices to point at the value from the frame prior to NewPendingFrame. May make the entire entry invalidated, if there were no value slots that old.
 	check(ExternalIdx == InternalIdx);
 
 	const uint32 LowestPossibleIdx = (ExternalIdx > Timestamps.Capacity()) ? (ExternalIdx - Timestamps.Capacity()) : 0u;
@@ -56,10 +66,9 @@ void URollbackBlackboard::BlackboardEntryBase::RollBack(uint32 NewPendingFrame)
 		}
 	}
 
-	// If we made it here, then there are no entries that weren't rolled back, so let's make it clear
+	// If we made it here, then there are no value slots that weren't rolled back, so let's make it clear
 	Timestamps[ExternalIdx].Invalidate();
 
-
 }
 
 
diff --git a/Engine/Plugins/Experimental/Mover/Source/Mover/Public/MoveLibrary/RollbackBlackboard.h b/Engine/Plugins/Experimental/Mover/Source/Mover/Public/MoveLibrary/RollbackBlackboard.h
--- a/Engine/Plugins/Experimental/Mover/Source/Mover/Public/MoveLibrary/RollbackBlackboard.h
+++ b/Engine/Plugins/Experimental/Mover/Source/Mover/Public/MoveLibrary/RollbackBlackboard.h
@@ -52,12 +52,12 @@ enum class EBlackboardSizingPolicy : uint8
 	 */
 	FixedBackendBufferSize UMETA(Hidden),
 
-	/** Minimizes data size for blackboard entries where thread contention and rolling back shouldn't change the value.
-	 * Example: a statistic tracking the maximum time spent in a particular mode's simulation tick. 
-	 * Note this will still have multiple elements for asynchronous simulations, to prevent thread contention over the entry.
-	 * NOT YET IMPLEMENTED
+	/** Buffer size is minimal, for a single value slot. Useful for cases where rolling back shouldn't change the value, 
+	 * or for short-lived values that are only set & read within the same simulation frame.
+	 * Example: a statistic tracking the maximum time spent in a particular mode's simulation tick.
+	 * Note this will still have 2 elements for asynchronous simulations, to prevent thread contention over the entry.
 	 */
-	SingleEntry UMETA(Hidden),
+	SingleEntry,
 };
 
 /**
@@ -83,9 +83,12 @@ enum class EBlackboardPersistencePolicy : uint8
 UENUM()
 enum class EBlackboardRollbackPolicy : uint8
 {
-	/** Entries occurring after the rollback time are effectively erased and no longer retrievable. New entries may be written during resimulation. */
+	/** Entry changes occurring after the rollback time are effectively erased and no longer retrievable. New entries may be written during resimulation. */
 	InvalidatedOnRollback,
 
+	/** Current value does not change when a rollback occurs. Typically paired with the SingleEntry sizing policy. */
+	IgnoreRollback,
+
 	/** Entries are not invalidated during rollback, and may be read/written during resim. If left alone, they will persist. NOT YET IMPLEMENTED. */
 	WritableDuringResimulation UMETA(Hidden),
 
@@ -289,7 +292,7 @@ public:
 	{
 		// TODO: Verify this isn't happening mid-sim on the sim thread (_Internal should be used instead)
 
-		if (BlackboardEntry<EntryT>* Entry = static_cast<BlackboardEntry<EntryT>*>(FindEntry(ObjName)))
+		if (BlackboardEntry<EntryT>* Entry = static_cast<BlackboardEntry<EntryT>*>(FindEntry_Internal(ObjName)))
 		{
 			Entry->SetEntryValue(Obj);
 			return true;
@@ -303,7 +306,7 @@ public:
 	template<typename EntryT>
 	bool TryGet(FName ObjName, EntryT& OutFoundValue) const
 	{
-		if (BlackboardEntry<EntryT>* Entry = static_cast<BlackboardEntry<EntryT>*>(FindEntry(ObjName)))
+		if (BlackboardEntry<EntryT>* Entry = static_cast<BlackboardEntry<EntryT>*>(FindEntry_Internal(ObjName)))
 		{
 			return Entry->TryGetEntryValue(CurrentSimTimeStamp, OutFoundValue);
 		}
@@ -317,7 +320,7 @@ public:
 	template<typename EntryT>
 	EntryT* TryGetRef(FName ObjName) const
 	{
-		if (BlackboardEntry<EntryT>* Entry = static_cast<BlackboardEntry<EntryT>*>(FindEntry(ObjName)))
+		if (BlackboardEntry<EntryT>* Entry = static_cast<BlackboardEntry<EntryT>*>(FindEntry_Internal(ObjName)))
 		{
 			return Entry->TryGetEntryRef(CurrentSimTimeStamp);
 		}
@@ -334,7 +337,7 @@ private:
 		check((bIsSimulationInProgress && (InProgressSimFrameThreadId == FPlatformTLS::GetCurrentThreadId())) || 
 			  (bIsRollbackInProgress && (InRollbackThreadId == FPlatformTLS::GetCurrentThreadId()))) ;
 
-		if (BlackboardEntry<EntryT>* Entry = static_cast<BlackboardEntry<EntryT>*>(FindEntry(ObjName)))
+		if (BlackboardEntry<EntryT>* Entry = static_cast<BlackboardEntry<EntryT>*>(FindEntry_Internal(ObjName)))
 		{
 			Entry->SetEntryValue_Internal(Obj, InProgressSimTimeStamp);
 			return true;
@@ -349,7 +352,7 @@ private:
 		check((bIsSimulationInProgress && (InProgressSimFrameThreadId == FPlatformTLS::GetCurrentThreadId())) ||
 			  (bIsRollbackInProgress && (InRollbackThreadId == FPlatformTLS::GetCurrentThreadId())));
 
-		if (BlackboardEntry<EntryT>* Entry = static_cast<BlackboardEntry<EntryT>*>(FindEntry(ObjName)))
+		if (BlackboardEntry<EntryT>* Entry = static_cast<BlackboardEntry<EntryT>*>(FindEntry_Internal(ObjName)))
 		{
 			return Entry->TryGetEntryValue_Internal(InProgressSimTimeStamp, OutFoundValue);
 		}
@@ -358,12 +361,12 @@ private:
 	}
 
 	template<typename EntryT>
-	EntryT* TryGetRef_Internal(FName ObjName) const
+	const EntryT* TryGetRef_Internal(FName ObjName) const
 	{
 		check((bIsSimulationInProgress && (InProgressSimFrameThreadId == FPlatformTLS::GetCurrentThreadId())) ||
 			  (bIsRollbackInProgress && (InRollbackThreadId == FPlatformTLS::GetCurrentThreadId())));
 
-		if (BlackboardEntry<EntryT>* Entry = static_cast<BlackboardEntry<EntryT>*>(FindEntry(ObjName)))
+		if (BlackboardEntry<EntryT>* Entry = static_cast<BlackboardEntry<EntryT>*>(FindEntry_Internal(ObjName)))
 		{
 			return Entry->TryGetEntryRef_Internal(InProgressSimTimeStamp);
 		}
@@ -373,7 +376,7 @@ private:
 
 
 
-	BlackboardEntryBase* FindEntry(FName EntryName) const
+	BlackboardEntryBase* FindEntry_Internal(FName EntryName) const
 	{
 		if (const TUniquePtr<BlackboardEntryBase>* EntryPtr = EntryMap.Find(EntryName))
 		{
@@ -452,10 +455,12 @@ public:
 	}
 
 	template<typename EntryT>
-	EntryT* TryGetRef(FName ObjName) const
+	const EntryT* TryGetRef(FName ObjName) const
 	{
 		return Blackboard->TryGetRef_Internal<EntryT>(ObjName);
 	}
+	
+	// TODO: Also support a mutable TryGetRef that gives you a writable reference
 
 	UE_API void BeginSimulationFrame(const FMoverTimeStep& PendingTimeStep);
 	UE_API void EndSimulationFrame();
