Mover: fixes bug where the Play Montage BP node for Mover would find the first skeletal mesh component, rather than favoring the primary visual component.

#Synced-CL 49732268

diff --git a/Engine/Plugins/Experimental/Mover/Source/Mover/Private/MoveLibrary/PlayMoverMontageCallbackProxy.cpp b/Engine/Plugins/Experimental/Mover/Source/Mover/Private/MoveLibrary/PlayMoverMontageCallbackProxy.cpp
--- a/Engine/Plugins/Experimental/Mover/Source/Mover/Private/MoveLibrary/PlayMoverMontageCallbackProxy.cpp
+++ b/Engine/Plugins/Experimental/Mover/Source/Mover/Private/MoveLibrary/PlayMoverMontageCallbackProxy.cpp
@@ -20,7 +20,17 @@ UPlayMoverMontageCallbackProxy* UPlayMoverMontageCallbackProxy::CreateProxyObjec
 	float StartingPosition,
 	FName StartingSection)
 {
-	USkeletalMeshComponent* SkelMeshComp = InMoverComponent ? InMoverComponent->GetOwner()->GetComponentByClass<USkeletalMeshComponent>() : nullptr;
+	USkeletalMeshComponent* SkelMeshComp = nullptr; 
+	
+	if (InMoverComponent)
+	{
+		SkelMeshComp = Cast<USkeletalMeshComponent>(InMoverComponent->GetPrimaryVisualComponent());
+
+		if (!SkelMeshComp)
+		{
+			SkelMeshComp = InMoverComponent->GetOwner()->GetComponentByClass<USkeletalMeshComponent>();
+		}
+	}
 
 	UPlayMoverMontageCallbackProxy* Proxy = NewObject<UPlayMoverMontageCallbackProxy>();
 	Proxy->SetFlags(RF_StrongRefOnFrame);
