Mover: fix crash that would occur with the default movements modes if someone tried something other than a capsule for an actor's root. Falls back to more general bounding cylinder check now.  Also makes step-up / step-over work for non-capsule shapes.
More broadly, this eliminates the requirement of the default movement modes and movement utils that the root shape is a capsule. Movement seems to work reasonably well even with boxes and spheres.

#Synced-CL 50195363

diff --git a/Engine/Plugins/Experimental/Mover/Source/Mover/Private/MoveLibrary/FloorQueryUtils.cpp b/Engine/Plugins/Experimental/Mover/Source/Mover/Private/MoveLibrary/FloorQueryUtils.cpp
--- a/Engine/Plugins/Experimental/Mover/Source/Mover/Private/MoveLibrary/FloorQueryUtils.cpp
+++ b/Engine/Plugins/Experimental/Mover/Source/Mover/Private/MoveLibrary/FloorQueryUtils.cpp
@@ -43,9 +43,15 @@ void UFloorQueryUtils::ComputeFloorDist(const FMovingComponentSet& MovingComps,
 	// TODO: pluggable shapes
 	float PawnRadius = 0.0f;
 	float PawnHalfHeight = 0.0f;
-	UCapsuleComponent* CapsuleComponent = Cast<UCapsuleComponent>(MovingComps.UpdatedComponent);
-	CapsuleComponent->GetScaledCapsuleSize(PawnRadius, PawnHalfHeight);
-	//MovingComps.UpdatedPrimitive->CalcBoundingCylinder(PawnRadius, PawnHalfHeight);
+	if (UCapsuleComponent* CapsuleComponent = Cast<UCapsuleComponent>(MovingComps.UpdatedComponent))
+	{
+		CapsuleComponent->GetScaledCapsuleSize(PawnRadius, PawnHalfHeight);
+	}
+	else if (MovingComps.UpdatedPrimitive.IsValid())
+	{
+		MovingComps.UpdatedPrimitive->CalcBoundingCylinder(PawnRadius, PawnHalfHeight);
+	}
+
 	FVector UpDirection = MovingComps.MoverComponent->GetUpDirection();
 
 	bool bBlockingHit = false;
diff --git a/Engine/Plugins/Experimental/Mover/Source/Mover/Private/MoveLibrary/GroundMovementUtils.cpp b/Engine/Plugins/Experimental/Mover/Source/Mover/Private/MoveLibrary/GroundMovementUtils.cpp
--- a/Engine/Plugins/Experimental/Mover/Source/Mover/Private/MoveLibrary/GroundMovementUtils.cpp
+++ b/Engine/Plugins/Experimental/Mover/Source/Mover/Private/MoveLibrary/GroundMovementUtils.cpp
@@ -57,30 +57,40 @@ static const FName SlideSubstepName = "SlideFromStep";
 
 bool UGroundMovementUtils::TryMoveToStepUp(const FMovingComponentSet& MovingComps, const FVector& GravDir, float MaxStepHeight, float MaxWalkSlopeCosine, bool bUseFlatBaseForFloorChecks, float FloorSweepDistance, const FVector& MoveDelta, const FHitResult& MoveHitResult, const FFloorCheckResult& CurrentFloor, bool bIsFalling, FOptionalFloorCheckResult* OutFloorTestResult, FMovementRecord& MoveRecord)
 {
-	UCapsuleComponent* CapsuleComponent = Cast<UCapsuleComponent>(MovingComps.UpdatedPrimitive.Get());
-
-	if (CapsuleComponent == nullptr || !CanStepUpOnHitSurface(MoveHitResult) || MaxStepHeight <= 0.f)
+	if (!CanStepUpOnHitSurface(MoveHitResult) || MaxStepHeight <= 0.f)
 	{
 		return false;
 	}
 
+	float PawnRadius(0.0f), PawnHalfHeight(0.0f);
+	if (const UCapsuleComponent* CapsuleComponent = Cast<UCapsuleComponent>(MovingComps.UpdatedPrimitive.Get()))
+	{
+		CapsuleComponent->GetScaledCapsuleSize(PawnRadius, PawnHalfHeight);
+	}
+	else if (MovingComps.UpdatedPrimitive.IsValid())
+	{
+		MovingComps.UpdatedPrimitive->CalcBoundingCylinder(PawnRadius, PawnHalfHeight);
+	}
+
+	USceneComponent* RootComponent = MovingComps.UpdatedComponent.Get();
+
+	const FVector OldLocation = RootComponent->GetComponentLocation();
+
 	TArray<FMovementSubstep> QueuedSubsteps;	// keeping track of substeps before committing, because some moves can be backed out
 
 	FVector UpDirection = MovingComps.MoverComponent->GetUpDirection();
 
-	const FVector OldLocation = CapsuleComponent->GetComponentLocation();
 	FVector LastComponentLocation = OldLocation;
 
-	float PawnRadius, PawnHalfHeight;
+	
 
-	CapsuleComponent->GetScaledCapsuleSize(PawnRadius, PawnHalfHeight);
 
-	// Don't bother stepping up if top of capsule is hitting something.
+	// Don't bother stepping up if top of shape is hitting something.
 	const float InitialImpactDot = MoveHitResult.ImpactPoint.Dot(UpDirection);
 	const float OldLocationDot = OldLocation.Dot(UpDirection);
 	if (InitialImpactDot > OldLocationDot + (PawnHalfHeight - PawnRadius))
 	{
-		UE_LOG(LogMover, VeryVerbose, TEXT("Not stepping up due to top of capsule hitting something"));
+		UE_LOG(LogMover, VeryVerbose, TEXT("Not stepping up due to top of shape hitting something"));
 		return false;
 	}
 
@@ -116,7 +126,7 @@ bool UGroundMovementUtils::TryMoveToStepUp(const FMovingComponentSet& MovingComp
 		}
 		else
 		{
-			// Base floor point is the base of the capsule moved down by how far we are hovering over the surface we are hitting.
+			// Base floor point is the base of the shape moved down by how far we are hovering over the surface we are hitting.
 			PawnFloorPointDot -= CurrentFloor.FloorDist;
 		}
 	}
@@ -129,17 +139,17 @@ bool UGroundMovementUtils::TryMoveToStepUp(const FMovingComponentSet& MovingComp
 	}
 
 	// Scope our movement updates, and do not apply them until all intermediate moves are completed.
-	FScopedMovementUpdate ScopedStepUpMovement(CapsuleComponent, EScopedUpdate::DeferredUpdates);
+	FScopedMovementUpdate ScopedStepUpMovement(RootComponent, EScopedUpdate::DeferredUpdates);
 
 	// step up - treat as vertical wall
 	FHitResult SweepUpHit(1.f);
-	const FQuat PawnRotation = CapsuleComponent->GetComponentQuat();
+	const FQuat PawnRotation = RootComponent->GetComponentQuat();
 
 	const FVector UpAdjustment = -GravDir * StepTravelUpHeight;
 	const bool bDidStepUp = UMovementUtils::TryMoveUpdatedComponent_Internal(MovingComps, UpAdjustment, PawnRotation, true, MOVECOMP_NoFlags, &SweepUpHit, ETeleportType::None);
 
 	UE_LOG(LogMover, VeryVerbose, TEXT("TryMoveToStepUp Up: %s (role %i) UpAdjustment=%s DidMove=%i"),
-		*GetNameSafe(CapsuleComponent->GetOwner()), CapsuleComponent->GetOwnerRole(), *UpAdjustment.ToCompactString(), bDidStepUp);
+		*GetNameSafe(RootComponent->GetOwner()), RootComponent->GetOwnerRole(), *UpAdjustment.ToCompactString(), bDidStepUp);
 
 	if (SweepUpHit.bStartPenetrating)
 	{
@@ -150,15 +160,15 @@ bool UGroundMovementUtils::TryMoveToStepUp(const FMovingComponentSet& MovingComp
 	}
 
 	// Cache upwards substep
-	QueuedSubsteps.Add(FMovementSubstep(StepUpSubstepName, CapsuleComponent->GetComponentLocation()-LastComponentLocation, false));
-	LastComponentLocation = CapsuleComponent->GetComponentLocation();
+	QueuedSubsteps.Add(FMovementSubstep(StepUpSubstepName, RootComponent->GetComponentLocation()-LastComponentLocation, false));
+	LastComponentLocation = RootComponent->GetComponentLocation();
 
 	// step fwd
 	FHitResult StepFwdHit(1.f);
 	const bool bDidStepFwd = UMovementUtils::TryMoveUpdatedComponent_Internal(MovingComps, MoveDelta, PawnRotation, true, MOVECOMP_NoFlags, &StepFwdHit, ETeleportType::None);
 
 	UE_LOG(LogMover, VeryVerbose, TEXT("TryMoveToStepUp Fwd: %s (role %i) MoveDelta=%s DidMove=%i"),
-		*GetNameSafe(CapsuleComponent->GetOwner()), CapsuleComponent->GetOwnerRole(), *MoveDelta.ToCompactString(), bDidStepFwd);
+		*GetNameSafe(RootComponent->GetOwner()), RootComponent->GetOwnerRole(), *MoveDelta.ToCompactString(), bDidStepFwd);
 
 	// Check result of forward movement
 	if (StepFwdHit.bBlockingHit)
@@ -190,7 +200,7 @@ bool UGroundMovementUtils::TryMoveToStepUp(const FMovingComponentSet& MovingComp
 		
 		if (bIsFalling)
 		{
-			QueuedSubsteps.Add( FMovementSubstep(StepFwdSubstepName, CapsuleComponent->GetComponentLocation()-LastComponentLocation, true) );
+			QueuedSubsteps.Add( FMovementSubstep(StepFwdSubstepName, RootComponent->GetComponentLocation() - LastComponentLocation, true) );
 
 			// Commit queued substeps to movement record
 			for (FMovementSubstep Substep : QueuedSubsteps)
@@ -202,8 +212,8 @@ bool UGroundMovementUtils::TryMoveToStepUp(const FMovingComponentSet& MovingComp
 		}
 
 		// Cache forwards substep before the slide attempt
-		QueuedSubsteps.Add(FMovementSubstep(StepFwdSubstepName, CapsuleComponent->GetComponentLocation() - LastComponentLocation, true));
-		LastComponentLocation = CapsuleComponent->GetComponentLocation();
+		QueuedSubsteps.Add(FMovementSubstep(StepFwdSubstepName, RootComponent->GetComponentLocation() - LastComponentLocation, true));
+		LastComponentLocation = RootComponent->GetComponentLocation();
 
 		// adjust and try again
 		const float ForwardHitTime = StepFwdHit.Time;
@@ -211,8 +221,8 @@ bool UGroundMovementUtils::TryMoveToStepUp(const FMovingComponentSet& MovingComp
 		// locking relevancy so velocity isn't added until it is needed to (adding it to the QueuedSubsteps so it can get added later)
 		MoveRecord.LockRelevancy(false);
 		const float ForwardSlideAmount = TryWalkToSlideAlongSurface(MovingComps, MoveDelta, 1.f - StepFwdHit.Time, PawnRotation, StepFwdHit.Normal, StepFwdHit, true, MoveRecord, MaxWalkSlopeCosine, MaxStepHeight);
-		QueuedSubsteps.Add( FMovementSubstep(SlideSubstepName, CapsuleComponent->GetComponentLocation()-LastComponentLocation, true) );
-		LastComponentLocation = CapsuleComponent->GetComponentLocation();
+		QueuedSubsteps.Add( FMovementSubstep(SlideSubstepName, RootComponent->GetComponentLocation() - LastComponentLocation, true) );
+		LastComponentLocation = RootComponent->GetComponentLocation();
 		MoveRecord.UnlockRelevancy();
 
 		if (bIsFalling)
@@ -233,17 +243,17 @@ bool UGroundMovementUtils::TryMoveToStepUp(const FMovingComponentSet& MovingComp
 	else
 	{
 		// Our forward move attempt was unobstructed - cache it
-		QueuedSubsteps.Add(FMovementSubstep(StepFwdSubstepName, CapsuleComponent->GetComponentLocation() - LastComponentLocation, true));
-		LastComponentLocation = CapsuleComponent->GetComponentLocation();
+		QueuedSubsteps.Add(FMovementSubstep(StepFwdSubstepName, RootComponent->GetComponentLocation() - LastComponentLocation, true));
+		LastComponentLocation = RootComponent->GetComponentLocation();
 	}
 
 
 	// Step down
 	const FVector StepDownAdjustment = GravDir * StepTravelDownHeight;
-	const bool bDidStepDown = UMovementUtils::TryMoveUpdatedComponent_Internal(MovingComps, StepDownAdjustment, CapsuleComponent->GetComponentQuat(), true, MOVECOMP_NoFlags, &StepFwdHit, ETeleportType::None);
+	const bool bDidStepDown = UMovementUtils::TryMoveUpdatedComponent_Internal(MovingComps, StepDownAdjustment, RootComponent->GetComponentQuat(), true, MOVECOMP_NoFlags, &StepFwdHit, ETeleportType::None);
 
 	UE_LOG(LogMover, VeryVerbose, TEXT("TryMoveToStepUp Down: %s (role %i) StepDownAdjustment=%s DidMove=%i"),
-		*GetNameSafe(CapsuleComponent->GetOwner()), CapsuleComponent->GetOwnerRole(), *StepDownAdjustment.ToCompactString(), bDidStepDown);
+		*GetNameSafe(RootComponent->GetOwner()), RootComponent->GetOwnerRole(), *StepDownAdjustment.ToCompactString(), bDidStepDown);
 
 
 	// If step down was initially penetrating abort the step up
@@ -288,7 +298,7 @@ bool UGroundMovementUtils::TryMoveToStepUp(const FMovingComponentSet& MovingComp
 			}
 		}
 
-		// Reject moves where the downward sweep hit something very close to the edge of the capsule. This maintains consistency with FindFloor as well.
+		// Reject moves where the downward sweep hit something very close to the edge of the shape. This maintains consistency with FindFloor as well.
 		if (!UFloorQueryUtils::IsWithinEdgeTolerance(StepFwdHit.Location, StepFwdHit.ImpactPoint, PawnRadius, UpDirection))
 		{
 			UE_LOG(LogMover, VeryVerbose, TEXT("Reject step-down attempt during step-up/step-fwd, due to being outside edge tolerance"));
@@ -308,7 +318,7 @@ bool UGroundMovementUtils::TryMoveToStepUp(const FMovingComponentSet& MovingComp
 		if (OutFloorTestResult != NULL)
 		{
 
-			UFloorQueryUtils::FindFloor(MovingComps, FloorSweepDistance, MaxWalkSlopeCosine, bUseFlatBaseForFloorChecks, CapsuleComponent->GetComponentLocation(), StepDownResult.FloorTestResult);
+			UFloorQueryUtils::FindFloor(MovingComps, FloorSweepDistance, MaxWalkSlopeCosine, bUseFlatBaseForFloorChecks, RootComponent->GetComponentLocation(), StepDownResult.FloorTestResult);
 
 			// Reject unwalkable normals if we end up higher than our initial height.
 			// It's fine to walk down onto an unwalkable surface, don't reject those moves.
@@ -330,8 +340,8 @@ bool UGroundMovementUtils::TryMoveToStepUp(const FMovingComponentSet& MovingComp
 	}
 
 	// Cache downwards substep
-	QueuedSubsteps.Add(FMovementSubstep(StepDownSubstepName, CapsuleComponent->GetComponentLocation() - LastComponentLocation, false));
-	LastComponentLocation = CapsuleComponent->GetComponentLocation();
+	QueuedSubsteps.Add(FMovementSubstep(StepDownSubstepName, RootComponent->GetComponentLocation() - LastComponentLocation, false));
+	LastComponentLocation = RootComponent->GetComponentLocation();
 
 	// Copy step down result.
 	if (OutFloorTestResult != NULL)
@@ -761,8 +771,6 @@ FVector UGroundMovementUtils::TestMoveToKeepMinHeightAboveFloor(const FMovingCom
 /* static */
 bool UGroundMovementUtils::TestMoveToStepOver(const FMovingComponentSet& MovingComps, const FVector& GravDir, float MaxStepHeight, float MaxWalkSlopeCosine, bool bUseFlatBaseForFloorChecks, float FloorSweepDistance, const FVector& MoveDelta, const FQuat& Rotation, const FHitResult& MoveHitResult, const FFloorCheckResult& CurrentFloor, bool bIsFalling, FOptionalFloorCheckResult* OutFloorTestResult, FVector& OutFinalLocation, FMovementRecord& InOutMoveRecord)
 {
-	UCapsuleComponent* CapsuleComponent = Cast<UCapsuleComponent>(MovingComps.UpdatedPrimitive.Get());
-
 	if (!UGroundMovementUtils::CanStepUpOnHitSurface(MoveHitResult) || MaxStepHeight <= 0.f)
 	{
 		return false;
@@ -775,16 +783,24 @@ bool UGroundMovementUtils::TestMoveToStepOver(const FMovingComponentSet& MovingC
 	const FVector OldLocation = MoveHitResult.TraceStart + ((MoveHitResult.TraceEnd - MoveHitResult.TraceStart) * MoveHitResult.Time);
 	FVector LocationInProgress = OldLocation;
 
-	float PawnRadius, PawnHalfHeight;
+	float PawnRadius(0.0f), PawnHalfHeight(0.0f);
+	if (const UCapsuleComponent* CapsuleComponent = Cast<UCapsuleComponent>(MovingComps.UpdatedPrimitive.Get()))
+	{
+		CapsuleComponent->GetScaledCapsuleSize(PawnRadius, PawnHalfHeight);
+	}
+	else if (MovingComps.UpdatedPrimitive.IsValid())
+	{
+		MovingComps.UpdatedPrimitive->CalcBoundingCylinder(PawnRadius, PawnHalfHeight);
+	}
 
-	CapsuleComponent->GetScaledCapsuleSize(PawnRadius, PawnHalfHeight);
+	USceneComponent* RootComponent = MovingComps.UpdatedComponent.Get();
 
-	// Don't bother stepping up if top of capsule is hitting something.
+	// Don't bother stepping up if top of shape is hitting something.
 	const float InitialImpactDot = MoveHitResult.ImpactPoint.Dot(UpDirection);
 	const float OldLocationDot = OldLocation.Dot(UpDirection);
 	if (InitialImpactDot > OldLocationDot + (PawnHalfHeight - PawnRadius))
 	{
-		UE_LOG(LogMover, VeryVerbose, TEXT("Not stepping up due to top of capsule hitting something"));
+		UE_LOG(LogMover, VeryVerbose, TEXT("Not stepping up due to top of shape hitting something"));
 		return false;
 	}
 
@@ -819,7 +835,7 @@ bool UGroundMovementUtils::TestMoveToStepOver(const FMovingComponentSet& MovingC
 		}
 		else
 		{
-			// Base floor point is the base of the capsule moved down by how far we are hovering over the surface we are hitting.
+			// Base floor point is the base of the shape moved down by how far we are hovering over the surface we are hitting.
 			PawnFloorPointDot -= CurrentFloor.FloorDist;
 		}
 	}
@@ -840,7 +856,7 @@ bool UGroundMovementUtils::TestMoveToStepOver(const FMovingComponentSet& MovingC
 	const bool bDidStepUp = UAsyncMovementUtils::TestMoveComponent_Internal(MovingComps, LocationInProgress, LocationInProgress + UpAdjustment, Rotation, Rotation, /*bShouldSweep=*/ true, CollisionParams, SweepUpHit);
 
 	UE_LOG(LogMover, VeryVerbose, TEXT("TestMoveToStepOver Up: %s (role %i) UpAdjustment=%s DidMove=%i"),
-		*GetNameSafe(CapsuleComponent->GetOwner()), CapsuleComponent->GetOwnerRole(), *UpAdjustment.ToCompactString(), bDidStepUp);
+		*GetNameSafe(RootComponent->GetOwner()), RootComponent->GetOwnerRole(), *UpAdjustment.ToCompactString(), bDidStepUp);
 
 	if (SweepUpHit.bStartPenetrating)
 	{
@@ -862,7 +878,7 @@ bool UGroundMovementUtils::TestMoveToStepOver(const FMovingComponentSet& MovingC
 	FVector FwdStepDelta = (StepFwdHit.TraceStart + ((StepFwdHit.TraceEnd - StepFwdHit.TraceStart) * StepFwdHit.Time)) - LocationInProgress;
 
 	UE_LOG(LogMover, VeryVerbose, TEXT("TestMoveToStepOver Fwd: %s (role %i) MoveDelta=%s DidMove=%i"),
-		*GetNameSafe(CapsuleComponent->GetOwner()), CapsuleComponent->GetOwnerRole(), *MoveDelta.ToCompactString(), bDidStepFwd);
+		*GetNameSafe(RootComponent->GetOwner()), RootComponent->GetOwnerRole(), *MoveDelta.ToCompactString(), bDidStepFwd);
 
 	// Check result of forward movement
 	if (StepFwdHit.bBlockingHit)
@@ -953,7 +969,7 @@ bool UGroundMovementUtils::TestMoveToStepOver(const FMovingComponentSet& MovingC
 
 
 	UE_LOG(LogMover, VeryVerbose, TEXT("TestMoveToStepOver Down: %s (role %i) StepDownAdjustment=%s DidMove=%i"),
-		*GetNameSafe(CapsuleComponent->GetOwner()), CapsuleComponent->GetOwnerRole(), *StepDownAdjustment.ToCompactString(), bDidStepDown);
+		*GetNameSafe(RootComponent->GetOwner()), RootComponent->GetOwnerRole(), *StepDownAdjustment.ToCompactString(), bDidStepDown);
 
 
 	// If step down was initially penetrating abort the step up
@@ -994,7 +1010,7 @@ bool UGroundMovementUtils::TestMoveToStepOver(const FMovingComponentSet& MovingC
 			}
 		}
 
-		// Reject moves where the downward sweep hit something very close to the edge of the capsule. This maintains consistency with FindFloor as well.
+		// Reject moves where the downward sweep hit something very close to the edge of the shape. This maintains consistency with FindFloor as well.
 		if (!UFloorQueryUtils::IsWithinEdgeTolerance(StepFwdHit.Location, StepFwdHit.ImpactPoint, PawnRadius, UpDirection))
 		{
 			UE_LOG(LogMover, VeryVerbose, TEXT("Reject step-down attempt during step-up/step-fwd, due to being outside edge tolerance"));
