Mover: Add USimpleFlyingMode to simplify the process of making a blueprint driven flying mode
GASP: Add example content to use a blueprint derivation of SimpleFlyingMode for a debug fly, bound to keyboard F

#virtualized
#Synced-CL 50205727
#Synced-CL 50205727

diff --git a/Engine/Plugins/Experimental/Mover/Source/Mover/Private/DefaultMovementSet/Modes/SimpleFlyingMode.cpp b/Engine/Plugins/Experimental/Mover/Source/Mover/Private/DefaultMovementSet/Modes/SimpleFlyingMode.cpp
new file mode 100644
--- /dev/null
+++ b/Engine/Plugins/Experimental/Mover/Source/Mover/Private/DefaultMovementSet/Modes/SimpleFlyingMode.cpp
@@ -0,0 +1,105 @@
+// Copyright Epic Games, Inc. All Rights Reserved.
+
+#include "DefaultMovementSet/Modes/SimpleFlyingMode.h"
+
+#include "MoverComponent.h"
+#include "DefaultMovementSet/Settings/CommonLegacyMovementSettings.h"
+#include "MoveLibrary/AirMovementUtils.h"
+#include "MoveLibrary/MovementUtils.h"
+
+#include UE_INLINE_GENERATED_CPP_BY_NAME(SimpleFlyingMode)
+
+void USimpleFlyingMode::GenerateMove_Implementation(const FMoverTickStartData& StartState, const FMoverTimeStep& TimeStep,
+                                                   FProposedMove& OutProposedMove) const
+{
+	const FMoverDefaultSyncState* StartingSyncState = StartState.SyncState.SyncStateCollection.FindDataByType<FMoverDefaultSyncState>();
+
+	if (CommonLegacySettings.Get() == nullptr || !StartingSyncState)
+	{
+		return;
+	}
+
+	const float DeltaSeconds = TimeStep.StepMs * 0.001f;
+	if (DeltaSeconds <= FLT_EPSILON)
+	{
+		return;
+	}
+
+	// Get input
+	FVector DesiredVelocity;
+	EMoveInputType MoveInputType;
+	FVector DesiredFacingDir;
+	 
+	if (const FCharacterDefaultInputs* CharacterInputs = StartState.InputCmd.InputCollection.FindDataByType<FCharacterDefaultInputs>())
+	{
+		DesiredVelocity = CharacterInputs->GetMoveInput_WorldSpace();
+		MoveInputType = CharacterInputs->GetMoveInputType();
+		DesiredFacingDir = CharacterInputs->GetOrientationIntentDir_WorldSpace();
+	}
+	else
+	{
+		// If no input found, most likely a networked sim proxy
+		// Try to deduce an input from the sync state
+		DesiredVelocity = StartingSyncState->GetIntent_WorldSpace();
+		MoveInputType = EMoveInputType::DirectionalIntent;
+		DesiredFacingDir = StartingSyncState->GetOrientation_WorldSpace().Quaternion().GetForwardVector();
+	}
+	float MaxMoveSpeed = MaxSpeedOverride >= 0.0f ? MaxSpeedOverride : CommonLegacySettings->MaxSpeed;
+
+	switch (MoveInputType)
+	{
+	case EMoveInputType::DirectionalIntent:
+		{
+			OutProposedMove.DirectionIntent = DesiredVelocity;	// here, DesiredVelocity is already in "intent space" (unit length for "max intent") so we can use it directly
+			DesiredVelocity *= MaxMoveSpeed;
+		}
+		break;
+	case EMoveInputType::Velocity:
+		{
+			// Clamp to max move speed
+			DesiredVelocity = DesiredVelocity.GetClampedToMaxSize(MaxMoveSpeed);
+			OutProposedMove.DirectionIntent = MaxMoveSpeed > UE_KINDA_SMALL_NUMBER ? DesiredVelocity / MaxMoveSpeed : FVector::ZeroVector; // here, DesiredVelocity is converted to "intent space"
+		}
+		break;
+	case EMoveInputType::None:
+	case EMoveInputType::Invalid:
+	default:
+		{
+			UE_LOG(LogMover, Warning, TEXT("Unhandled MoveInputType %i in USimpleFlyingMode"), MoveInputType);
+			DesiredVelocity = FVector::ZeroVector;
+			OutProposedMove.DirectionIntent = FVector::ZeroVector;
+		}
+	break;
+	}
+
+	OutProposedMove.bHasDirIntent = !OutProposedMove.DirectionIntent.IsNearlyZero();
+	DesiredFacingDir -= DesiredFacingDir.ProjectOnTo(GetMoverComponent()->GetUpDirection());
+	FQuat CurrentFacing = StartingSyncState->GetOrientation_WorldSpace().Quaternion();
+	FQuat DesiredFacing = CurrentFacing;
+	
+	if (DesiredFacingDir.Normalize())
+	{
+		DesiredFacing = FQuat::FindBetween(FVector::ForwardVector, DesiredFacingDir);
+	}
+
+	OutProposedMove.LinearVelocity = StartingSyncState->GetVelocity_WorldSpace();
+	FVector AngularVelocityDegrees = StartingSyncState->GetAngularVelocityDegrees_WorldSpace();
+	
+	// Hack const_cast stuff
+	// Why is this needed?
+	// Because some modes want to mutate their data inside the generate move
+	USimpleFlyingMode* MutableSimpleMode  = const_cast<USimpleFlyingMode*>(this);
+	MutableSimpleMode->GenerateFlyingMove(const_cast<FMoverTickStartData&>(StartState), DeltaSeconds, DesiredVelocity, DesiredFacing, CurrentFacing, AngularVelocityDegrees, OutProposedMove.LinearVelocity);
+
+	// Calc angular velocity from final facing
+	OutProposedMove.AngularVelocityDegrees = AngularVelocityDegrees;
+}
+
+void USimpleFlyingMode::GenerateFlyingMove_Implementation(FMoverTickStartData& StartState, float DeltaSeconds, const FVector& DesiredVelocity,
+	const FQuat& DesiredFacing, const FQuat& CurrentFacing, FVector& InOutAngularVelocityDegrees, FVector& InOutVelocity)
+{
+	InOutVelocity = DesiredVelocity;
+	FQuat ToFacing = CurrentFacing.Inverse() * DesiredFacing;
+	InOutAngularVelocityDegrees = DeltaSeconds > 0.0f ? FMath::RadiansToDegrees(ToFacing.ToRotationVector() / DeltaSeconds) : FVector::ZeroVector;
+}
+
diff --git a/Engine/Plugins/Experimental/Mover/Source/Mover/Private/DefaultMovementSet/Modes/SimpleWalkingMode.cpp b/Engine/Plugins/Experimental/Mover/Source/Mover/Private/DefaultMovementSet/Modes/SimpleWalkingMode.cpp
--- a/Engine/Plugins/Experimental/Mover/Source/Mover/Private/DefaultMovementSet/Modes/SimpleWalkingMode.cpp
+++ b/Engine/Plugins/Experimental/Mover/Source/Mover/Private/DefaultMovementSet/Modes/SimpleWalkingMode.cpp
@@ -72,7 +72,7 @@ void USimpleWalkingMode::GenerateMove_Implementation(const FMoverTickStartData&
 	case EMoveInputType::Invalid:
 	default:
 		{
-			UE_LOG(LogMover, Warning, TEXT("Unhandled MoveInputType %i in USpringWalkingMode"), MoveInputType);
+			UE_LOG(LogMover, Warning, TEXT("Unhandled MoveInputType %i in USimpleWalkingMode"), MoveInputType);
 			DesiredVelocity = FVector::ZeroVector;
 			OutProposedMove.DirectionIntent = FVector::ZeroVector;
 		}
@@ -90,16 +90,13 @@ void USimpleWalkingMode::GenerateMove_Implementation(const FMoverTickStartData&
 	}
 
 	OutProposedMove.LinearVelocity = StartingSyncState->GetVelocity_WorldSpace();
-	FVector AngularVelocityDegrees = StartingSyncState->GetAngularVelocityDegrees_WorldSpace();
+	OutProposedMove.AngularVelocityDegrees = StartingSyncState->GetAngularVelocityDegrees_WorldSpace();
 	
 	// Hack const_cast stuff
 	// Why is this needed?
 	// Because some modes want to mutate their data inside the generate walk move
 	USimpleWalkingMode* MutableSimpleWalkMode  = const_cast<USimpleWalkingMode*>(this);
-	MutableSimpleWalkMode->GenerateWalkMove(const_cast<FMoverTickStartData&>(StartState), DeltaSeconds, DesiredVelocity, DesiredFacing, CurrentFacing, AngularVelocityDegrees, OutProposedMove.LinearVelocity);
-
-	// Calc angular velocity from final facing
-	OutProposedMove.AngularVelocityDegrees = AngularVelocityDegrees;
+	MutableSimpleWalkMode->GenerateWalkMove(const_cast<FMoverTickStartData&>(StartState), DeltaSeconds, DesiredVelocity, DesiredFacing, CurrentFacing, OutProposedMove.AngularVelocityDegrees, OutProposedMove.LinearVelocity);
 }
 
 void USimpleWalkingMode::GenerateWalkMove_Implementation(FMoverTickStartData& StartState, float DeltaSeconds, const FVector& DesiredVelocity,
diff --git a/Engine/Plugins/Experimental/Mover/Source/Mover/Public/DefaultMovementSet/Modes/SimpleFlyingMode.h b/Engine/Plugins/Experimental/Mover/Source/Mover/Public/DefaultMovementSet/Modes/SimpleFlyingMode.h
new file mode 100644
--- /dev/null
+++ b/Engine/Plugins/Experimental/Mover/Source/Mover/Public/DefaultMovementSet/Modes/SimpleFlyingMode.h
@@ -0,0 +1,33 @@
+// Copyright Epic Games, Inc. All Rights Reserved.
+
+#pragma once
+
+#include "CoreMinimal.h"
+#include "DefaultMovementSet/Modes/FlyingMode.h"
+#include "SimpleFlyingMode.generated.h"
+
+#define UE_API MOVER_API
+
+/**
+ * Basic flying mode that will move exactly where inputs request. Easy to derive from to create custom behaviors
+ */
+UCLASS(BlueprintType, Abstract)
+class USimpleFlyingMode : public UFlyingMode
+{
+	GENERATED_BODY()
+
+public:
+
+	UE_API virtual void GenerateMove_Implementation(const FMoverTickStartData& StartState, const FMoverTimeStep& TimeStep, FProposedMove& OutProposedMove) const override;
+
+	// Override this to make a simple flying mode
+	UFUNCTION(BlueprintNativeEvent, meta = (DisplayName = "Generate Simple Flying Move"))
+	UE_API void GenerateFlyingMove(UPARAM(ref) FMoverTickStartData& StartState, float DeltaSeconds, const FVector& DesiredVelocity,
+		const FQuat& DesiredFacing, const FQuat& CurrentFacing, UPARAM(ref) FVector& InOutAngularVelocityDegrees, UPARAM(ref) FVector& InOutVelocity);
+
+	// If this value is greater or equal to 0, this will override the max speed read from the CommonLegacyMovementSettings
+	UPROPERTY(EditAnywhere, BlueprintReadWrite, Category = "Mover|Flying Settings", meta = (ForceUnits = "cm/s"))
+	float MaxSpeedOverride = -1.0f;
+};
+
+#undef UE_API
