Refactor UAsyncNavWalkingMode::FindBestNavMeshLocation to rely solely on IPhysicsBodyInstanceOwner::GetPhysicsBodyInstanceOwnerFromHitResult, since UPrimitiveComponent implements IPhysicsBodyInstanceOwner.

#Synced-CL 50298075

diff --git a/Engine/Plugins/Experimental/Mover/Source/Mover/Private/DefaultMovementSet/Modes/AsyncNavWalkingMode.cpp b/Engine/Plugins/Experimental/Mover/Source/Mover/Private/DefaultMovementSet/Modes/AsyncNavWalkingMode.cpp
--- a/Engine/Plugins/Experimental/Mover/Source/Mover/Private/DefaultMovementSet/Modes/AsyncNavWalkingMode.cpp
+++ b/Engine/Plugins/Experimental/Mover/Source/Mover/Private/DefaultMovementSet/Modes/AsyncNavWalkingMode.cpp
@@ -468,20 +468,10 @@ void UAsyncNavWalkingMode::FindBestNavMeshLocation(const FVector& TraceStart, co
 
 		FORCEINLINE bool operator()(const FHitResult& TestHit) const
 		{
-			UPrimitiveComponent* PrimComp = TestHit.GetComponent();
-			// Prefer using primitive component if valid
-			if (const IPhysicsBodyInstanceOwner* BodyInstanceOwner = !PrimComp ? IPhysicsBodyInstanceOwner::GetPhysicsBodyInstanceOwnerFromHitResult(TestHit) : nullptr)
-			{
-				const bool bBlockOnWorldStatic = (BodyInstanceOwner->GetCollisionResponseToChannel(ECC_WorldStatic) == ECR_Block);
-				const bool bBlockOnWorldDynamic = (BodyInstanceOwner->GetCollisionResponseToChannel(ECC_WorldDynamic) == ECR_Block);
-				return !bBlockOnWorldStatic && (!bBlockOnWorldDynamic || bCheckOnlyWorldStatic);
-			}
-			else
-			{
-				const bool bBlockOnWorldStatic = PrimComp && (PrimComp->GetCollisionResponseToChannel(ECC_WorldStatic) == ECR_Block);
-				const bool bBlockOnWorldDynamic = PrimComp && (PrimComp->GetCollisionResponseToChannel(ECC_WorldDynamic) == ECR_Block);
-				return !bBlockOnWorldStatic && (!bBlockOnWorldDynamic || bCheckOnlyWorldStatic);
-			}
+			const IPhysicsBodyInstanceOwner* BodyInstanceOwner = IPhysicsBodyInstanceOwner::GetPhysicsBodyInstanceOwnerFromHitResult(TestHit);
+			const bool bBlockOnWorldStatic = BodyInstanceOwner && (BodyInstanceOwner->GetCollisionResponseToChannel(ECC_WorldStatic) == ECR_Block);
+			const bool bBlockOnWorldDynamic = BodyInstanceOwner && (BodyInstanceOwner->GetCollisionResponseToChannel(ECC_WorldDynamic) == ECR_Block);
+			return !bBlockOnWorldStatic && (!bBlockOnWorldDynamic || bCheckOnlyWorldStatic);
 		}
 
 		bool bCheckOnlyWorldStatic;
