Mover: supporting prediction on the rollback blackboard, with dedicated slot in the circular buffer that avoids the need for copying and mem allocs. Also changes acccessor wrappers so there are external and simulation wrappers, with simulation wrapper hiding whether access is predicted or not.  No one except the owner + simulation driver should have direct access to the blackboard object.
All sim cases now have a Mover 'Simulation' object, which now owns the blackboard and controls access to it. This is a preliminary change that will be followed by more changes to how the Simulation object is used as the movement flow is driven.

MoverTests: adds first version of automated testing for the blackboard.

#Synced-CL 50587857

diff --git a/Engine/Plugins/Experimental/ChaosMover/Source/ChaosMover/Private/Backends/ChaosMoverBackend.cpp b/Engine/Plugins/Experimental/ChaosMover/Source/ChaosMover/Private/Backends/ChaosMoverBackend.cpp
--- a/Engine/Plugins/Experimental/ChaosMover/Source/ChaosMover/Private/Backends/ChaosMoverBackend.cpp
+++ b/Engine/Plugins/Experimental/ChaosMover/Source/ChaosMover/Private/Backends/ChaosMoverBackend.cpp
@@ -64,15 +64,16 @@ void UChaosMoverBackendComponent::InitializeComponent()
 
 	if (UWorld* World = GetWorld(); World && World->IsGameWorld())
 	{
+		UMoverComponent& MoverComp = GetMoverComponent();
+
 		Simulation = NewObject<UChaosMoverSimulation>(GetOwner(), SimulationClass, TEXT("ChaosMoverSimulation"), RF_Transient);
-		Simulation->SetOwner(GetMoverComponent().GetOwner());
+		Simulation->SetOwner(MoverComp.GetOwner());
+		MoverComp.SetSimulation(*Simulation);
 
-		NullMovementMode = NewObject<UNullMovementMode>(&GetMoverComponent(), TEXT("NullMovementMode"), RF_Transient);
+		NullMovementMode = NewObject<UNullMovementMode>(&MoverComp, TEXT("NullMovementMode"), RF_Transient);
 		ImmediateModeTransition = NewObject<UImmediateMovementModeTransition>(&GetMoverComponent(), TEXT("ImmediateModeTransition"), RF_Transient);
 		TransformOnInit = GetMoverComponent().GetUpdatedComponentTransform();
 
-		Simulation->SetRollbackBlackboard(GetMoverComponent().GetRollbackBlackboard_Internal());
-
 		// Create NetworkPhysicsComponent
 		if ((World->GetNetMode() != ENetMode::NM_Standalone) && Chaos::FPhysicsSolverBase::IsNetworkPhysicsPredictionEnabled() && bIsUsingAsyncPhysics)
 		{
diff --git a/Engine/Plugins/Experimental/ChaosMover/Source/ChaosMover/Private/ChaosMoverStateMachine.cpp b/Engine/Plugins/Experimental/ChaosMover/Source/ChaosMover/Private/ChaosMoverStateMachine.cpp
--- a/Engine/Plugins/Experimental/ChaosMover/Source/ChaosMover/Private/ChaosMoverStateMachine.cpp
+++ b/Engine/Plugins/Experimental/ChaosMover/Source/ChaosMover/Private/ChaosMoverStateMachine.cpp
@@ -9,7 +9,7 @@
 #include "ChaosMover/ChaosMoverSimulation.h"
 #include "MoveLibrary/MovementMixer.h"
 #include "MoveLibrary/MovementUtils.h"
-#include "MoveLibrary/RollbackBlackboard.h"
+#include "MoveLibrary/RollbackBlackboardWrappers.h"
 #include "MovementModeTransition.h"
 #include "MoverComponent.h"
 #include "MoverDeveloperSettings.h"
@@ -832,16 +832,14 @@ namespace UE::ChaosMover
 						}
 
 						// Track last mode change in the blackboard
-						if (URollbackBlackboard_InternalWrapper* RollbackBlackboard = SimStrongObjPtr->GetRollbackBlackboard_Mutable())
-						{
-							FMovementModeChangeRecord ModeChangeRecord;
-							ModeChangeRecord.ModeName = CurrentModeName;
-							ModeChangeRecord.PrevModeName = PreviousModeName;
-							ModeChangeRecord.Frame = InternalSimTime.ServerFrame;
-							ModeChangeRecord.SimTimeMs = InternalSimTime.BaseSimTimeMs;
-
-							RollbackBlackboard->TrySet(CommonBlackboard::LastModeChangeRecord, ModeChangeRecord);
-						}
+						FRollbackBlackboardSimWrapper RollbackBlackboard = SimStrongObjPtr->GetRollbackBlackboardSimWrapper();
+						FMovementModeChangeRecord ModeChangeRecord;
+						ModeChangeRecord.ModeName = CurrentModeName;
+						ModeChangeRecord.PrevModeName = PreviousModeName;
+						ModeChangeRecord.Frame = InternalSimTime.ServerFrame;
+						ModeChangeRecord.SimTimeMs = InternalSimTime.BaseSimTimeMs;
+
+						RollbackBlackboard.TrySet(CommonBlackboard::LastModeChangeRecord, ModeChangeRecord);
 
 						if (TStrongObjectPtr<UBaseMovementMode> CurrentMode = Modes[CurrentModeName].Pin())
 						{	
diff --git a/Engine/Plugins/Experimental/ChaosMover/Source/ChaosMover/Public/ChaosMover/ChaosMoverSimulation.h b/Engine/Plugins/Experimental/ChaosMover/Source/ChaosMover/Public/ChaosMover/ChaosMoverSimulation.h
--- a/Engine/Plugins/Experimental/ChaosMover/Source/ChaosMover/Public/ChaosMover/ChaosMoverSimulation.h
+++ b/Engine/Plugins/Experimental/ChaosMover/Source/ChaosMover/Public/ChaosMover/ChaosMoverSimulation.h
@@ -7,7 +7,6 @@
 #include "ChaosMover/PathedMovement/ChaosPathedMovementTypes.h"
 #include "ChaosMover/ChaosMoverStateMachine.h"
 #include "MoveLibrary/MovementMixer.h"
-#include "MoveLibrary/MoverBlackboard.h"
 #include "MoverTypes.h"
 #include "MoverSimulation.h"
 #include "PhysicsEngine/ConstraintInstance.h"
diff --git a/Engine/Plugins/Experimental/Mover/Source/Mover/Private/Backends/MoverNetworkPhysicsLiaisonBase.cpp b/Engine/Plugins/Experimental/Mover/Source/Mover/Private/Backends/MoverNetworkPhysicsLiaisonBase.cpp
--- a/Engine/Plugins/Experimental/Mover/Source/Mover/Private/Backends/MoverNetworkPhysicsLiaisonBase.cpp
+++ b/Engine/Plugins/Experimental/Mover/Source/Mover/Private/Backends/MoverNetworkPhysicsLiaisonBase.cpp
@@ -2,8 +2,9 @@
 
 #include "Backends/MoverNetworkPhysicsLiaisonBase.h"
 
-#include "MoveLibrary/RollbackBlackboard.h"
+#include "MoveLibrary/RollbackBlackboardWrappers.h"
 #include "MovementModeStateMachine.h"
+#include "MoverSimulation.h"
 #include "PBDRigidsSolver.h"
 #include "Chaos/PhysicsObjectInterface.h"
 #include "Chaos/PhysicsObjectInternalInterface.h"
@@ -232,6 +233,8 @@ UMoverNetworkPhysicsLiaisonComponentBase::UMoverNetworkPhysicsLiaisonComponentBa
 	{
 		SetIsReplicatedByDefault(true);
 	}
+
+	SimulationClass = UMoverSimulation::StaticClass();
 }
 
 //////////////////////////////////////////////////////////////////////////
@@ -293,12 +296,12 @@ void UMoverNetworkPhysicsLiaisonComponentBase::HandleComponentPhysicsStateChange
 	{
 		if (UMoverComponent& MoverComp = GetMoverComponent(); MoverComp.ModeFSM)
 		{
-			MoverComp.RollbackBlackboard_InternalWrapper->BeginSimulationFrame(GetCurrentAsyncMoverTimeStep_External());
+			RollbackBlackboard->BeginSimulationFrame(GetCurrentAsyncMoverTimeStep_External());
 
 			MoverComp.ModeFSM->SetModeImmediately(MoverComp.StartingMovementMode);
 			MoverComp.SimBlackboard->InvalidateAll();
 
-			MoverComp.RollbackBlackboard_InternalWrapper->EndSimulationFrame();
+			RollbackBlackboard->EndSimulationFrame();
 		}
 
 		InitializeSimOutputData();
@@ -332,6 +335,10 @@ void UMoverNetworkPhysicsLiaisonComponentBase::InitializeComponent()
 		UMoverComponent& MoverComp = GetMoverComponent();
 	}
 
+	Simulation = NewObject<UMoverSimulation>(GetOwner(), SimulationClass, TEXT("MoverSimulation"), RF_Transient);
+	GetMoverComponent().SetSimulation(*Simulation);
+	RollbackBlackboard = Simulation->RollbackBlackboard;
+
 	if (Chaos::FPhysicsSolverBase::IsNetworkPhysicsPredictionEnabled() && bUsingAsyncPhysics)
 	{
 		NetworkPhysicsComponent = NewObject<UNetworkPhysicsComponent>(GetOwner(), TEXT("PhysMover_NetworkPhysicsComponent"));
@@ -835,7 +842,7 @@ void UMoverNetworkPhysicsLiaisonComponentBase::PerformPreSimulate_Internal(const
 	FMoverTickEndData TickEndData;
 	TickEndData.SyncState = TickStartData.SyncState;
 
-	MoverComp.RollbackBlackboard_InternalWrapper->BeginSimulationFrame(TimeStep);
+	RollbackBlackboard->BeginSimulationFrame(TimeStep);
 
 
 	MoverComp.ModeFSM->OnSimulationTick(MoverComp.UpdatedComponent, MoverComp.UpdatedCompAsPrimitive, MoverComp.GetSimBlackboard_Mutable(), TickStartData, TimeStep, TickEndData);
@@ -862,8 +869,8 @@ void UMoverNetworkPhysicsLiaisonComponentBase::PerformPreSimulate_Internal(const
 
 		//@todo DanH: Does the base need to concern itself with setting W on the particle?
 	}
-
-	MoverComp.RollbackBlackboard_InternalWrapper->EndSimulationFrame();
+	
+	RollbackBlackboard->EndSimulationFrame();
 
 }
 
diff --git a/Engine/Plugins/Experimental/Mover/Source/Mover/Private/Backends/MoverNetworkPredictionLiaison.cpp b/Engine/Plugins/Experimental/Mover/Source/Mover/Private/Backends/MoverNetworkPredictionLiaison.cpp
--- a/Engine/Plugins/Experimental/Mover/Source/Mover/Private/Backends/MoverNetworkPredictionLiaison.cpp
+++ b/Engine/Plugins/Experimental/Mover/Source/Mover/Private/Backends/MoverNetworkPredictionLiaison.cpp
@@ -5,6 +5,7 @@
 #include "NetworkPredictionProxyInit.h"
 #include "NetworkPredictionProxyWrite.h"
 #include "GameFramework/Actor.h"
+#include "MoverSimulation.h"
 #include "MoverComponent.h"
 
 #if WITH_EDITOR
@@ -46,6 +47,8 @@ UMoverNetworkPredictionLiaisonComponent::UMoverNetworkPredictionLiaisonComponent
 	bWantsInitializeComponent = true;
 	bAutoActivate = true;
 	SetIsReplicatedByDefault(true);
+
+	SimulationClass = UMoverSimulation::StaticClass();
 }
 
 void UMoverNetworkPredictionLiaisonComponent::ProduceInput(const int32 DeltaTimeMS, FMoverInputCmdContext* Cmd)
@@ -283,6 +286,15 @@ void UMoverNetworkPredictionLiaisonComponent::BeginPlay()
 void UMoverNetworkPredictionLiaisonComponent::InitializeComponent()
 {
 	Super::InitializeComponent();
+
+	Simulation = NewObject<UMoverSimulation>(GetOwner(), SimulationClass, TEXT("MoverSimulation"), RF_Transient);
+	
+	MoverComp = GetOwner()->FindComponentByClass<UMoverComponent>();
+
+	if (MoverComp)
+	{
+		MoverComp->SetSimulation(*Simulation);
+	}
 }
 
 
@@ -290,6 +302,8 @@ void UMoverNetworkPredictionLiaisonComponent::UninitializeComponent()
 {
 	NetworkPredictionProxy.EndPlay();
 
+	Simulation = nullptr;
+
 	Super::UninitializeComponent();
 }
 
diff --git a/Engine/Plugins/Experimental/Mover/Source/Mover/Private/Backends/MoverStandaloneLiaison.cpp b/Engine/Plugins/Experimental/Mover/Source/Mover/Private/Backends/MoverStandaloneLiaison.cpp
--- a/Engine/Plugins/Experimental/Mover/Source/Mover/Private/Backends/MoverStandaloneLiaison.cpp
+++ b/Engine/Plugins/Experimental/Mover/Source/Mover/Private/Backends/MoverStandaloneLiaison.cpp
@@ -5,6 +5,7 @@
 #include "GameFramework/Pawn.h"
 #include "GameFramework/Controller.h"
 #include "MoverComponent.h"
+#include "MoverSimulation.h"
 
 #include UE_INLINE_GENERATED_CPP_BY_NAME(MoverStandaloneLiaison)
 
@@ -84,6 +85,8 @@ UMoverStandaloneLiaisonComponent::UMoverStandaloneLiaisonComponent(const FObject
 	bAutoActivate = true;
 	SetIsReplicatedByDefault(false);
 
+	SimulationClass = UMoverSimulation::StaticClass();
+
 	CurrentSimTimeMs = 0.0;
 	CurrentSimFrame = 0;
 
@@ -127,6 +130,27 @@ bool UMoverStandaloneLiaisonComponent::WritePendingSyncState(const FMoverSyncSta
 	return true;
 }
 
+void UMoverStandaloneLiaisonComponent::InitializeComponent()
+{
+	Super::InitializeComponent();
+
+	Simulation = NewObject<UMoverSimulation>(GetOwner(), SimulationClass, TEXT("MoverSimulation"), RF_Transient);
+
+	MoverComp = GetOwner()->FindComponentByClass<UMoverComponent>();
+
+	if (MoverComp)
+	{
+		MoverComp->SetSimulation(*Simulation);
+	}
+}
+
+void UMoverStandaloneLiaisonComponent::UninitializeComponent()
+{
+	Simulation = nullptr;
+	Super::UninitializeComponent();
+}
+
+
 void UMoverStandaloneLiaisonComponent::BeginPlay()
 {
 	Super::BeginPlay();
diff --git a/Engine/Plugins/Experimental/Mover/Source/Mover/Private/MoveLibrary/RollbackBlackboard.cpp b/Engine/Plugins/Experimental/Mover/Source/Mover/Private/MoveLibrary/RollbackBlackboard.cpp
--- a/Engine/Plugins/Experimental/Mover/Source/Mover/Private/MoveLibrary/RollbackBlackboard.cpp
+++ b/Engine/Plugins/Experimental/Mover/Source/Mover/Private/MoveLibrary/RollbackBlackboard.cpp
@@ -35,9 +35,9 @@ uint32 URollbackBlackboard::BlackboardEntryBase::ComputeBufferSize(EBlackboardSi
 	return BufferSize;
 }
 
-URollbackBlackboard::BlackboardEntryBase::BlackboardEntryBase(const EntrySettings& InSettings, uint32 BufferSize)
+URollbackBlackboard::BlackboardEntryBase::BlackboardEntryBase(const EntrySettings& InSettings, uint32 BufferSize, ERollbackBufferFlags BufferFlags)
 	: Settings(InSettings)
-	, Timestamps(BufferSize, EntryTimeStamp())
+	, Timestamps(BufferSize, BufferFlags, EntryTimeStamp())
 	, ExternalIdx(0)
 	, InternalIdx(0)
 {
@@ -71,13 +71,19 @@ void URollbackBlackboard::BlackboardEntryBase::RollBack(uint32 NewPendingFrame)
 
 }
 
+void URollbackBlackboard::BlackboardEntryBase::InvalidatePredictiveState()
+{
+	Timestamps.GetPredictiveElement().Invalidate();
+
+	ResetPredictiveElement();
+}
 
 
-bool URollbackBlackboard::BlackboardEntryBase::CanReadEntryAt(const EntryTimeStamp& ReaderTimeStamp, EEntryIndexType IndexType) const
+
+bool URollbackBlackboard::BlackboardEntryBase::CanReadEntryAtTime(const EntryTimeStamp& ReaderTimeStamp, const EntryTimeStamp& EntryTimeStamp) const
 {
-	const uint32 TimestampIdx = IndexType == EEntryIndexType::External ? ExternalIdx : InternalIdx;
 
-	if (!Timestamps[TimestampIdx].IsValid())
+	if (!EntryTimeStamp.IsValid() || !ReaderTimeStamp.IsValid())
 	{
 		return false;	// entry isn't initialized yet or never set
 	}
@@ -87,15 +93,20 @@ bool URollbackBlackboard::BlackboardEntryBase::CanReadEntryAt(const EntryTimeSta
 		case EBlackboardPersistencePolicy::NextFrameOnly:
 		{
 			// if value wasn't last set during the current or immediate prior sim frame, then it can't be read
-			if (!(ReaderTimeStamp.Frame == Timestamps[TimestampIdx].Frame ||
-				  ReaderTimeStamp.Frame == Timestamps[TimestampIdx].Frame+1))
+			if (!(ReaderTimeStamp.Frame == EntryTimeStamp.Frame ||
+				  ReaderTimeStamp.Frame == EntryTimeStamp.Frame+1))
 			{
 				return false;
 			}
 		}
 		break;
 
-		case EBlackboardPersistencePolicy::Forever:		// Always allow reading
+		case EBlackboardPersistencePolicy::Forever:
+			// If the reader is in the past compared to the time of the entry, then it can't be read
+			if (ReaderTimeStamp.Frame < EntryTimeStamp.Frame)
+			{
+				return false;
+			}
 			break;
 
 		default:		// Allow reading of any not-yet-implemented policies
@@ -105,6 +116,22 @@ bool URollbackBlackboard::BlackboardEntryBase::CanReadEntryAt(const EntryTimeSta
 	return true;
 }
 
+bool URollbackBlackboard::BlackboardEntryBase::CanReadEntryAtTime(const EntryTimeStamp& ReaderTimeStamp, EEntryIndexType IndexType) const
+{
+	const EntryTimeStamp* TargetEntryTimeStamp = nullptr;
+
+	if (IndexType == EEntryIndexType::Predictive)
+	{
+		TargetEntryTimeStamp = &Timestamps.GetPredictiveElement();
+	}
+	else
+	{
+		const uint32 TimestampIdx = ((IndexType == EEntryIndexType::External) ? ExternalIdx : InternalIdx);
+		TargetEntryTimeStamp = &Timestamps[TimestampIdx];
+	}
+	
+	return CanReadEntryAtTime(ReaderTimeStamp, *TargetEntryTimeStamp);
+}
 
 
 void URollbackBlackboard::BeginSimulationFrame(const FMoverTimeStep& PendingTimeStep)
@@ -122,6 +149,7 @@ void URollbackBlackboard::EndSimulationFrame()
 {
 	check(bIsSimulationInProgress && (InProgressSimFrameThreadId == FPlatformTLS::GetCurrentThreadId()));
 	bIsSimulationInProgress = bIsResimulating = false;
+	InProgressSimFrameThreadId = 0;
 
 	// TODO: need some kind of lock so we wait on any in-progress operations before advancing the CurrentSimTimeStamp
 	CurrentSimTimeStamp = InProgressSimTimeStamp;
@@ -162,25 +190,41 @@ void URollbackBlackboard::EndRollback()
 {
 	check (bIsRollbackInProgress && (InRollbackThreadId == FPlatformTLS::GetCurrentThreadId()));
 	bIsRollbackInProgress = false;
+	InRollbackThreadId = 0;
 }
 
-
-void URollbackBlackboard_InternalWrapper::BeginSimulationFrame(const FMoverTimeStep& PendingTimeStep)
+void URollbackBlackboard::BeginPredictionFrame(const FMoverTimeStep& PendingTimeStep)
 {
-	Blackboard->BeginSimulationFrame(PendingTimeStep);
+	check(BufferFlags & ERollbackBufferFlags::Prediction);
+	check(!bIsPredictionInProgress || InPredictionThreadId == FPlatformTLS::GetCurrentThreadId());
+	
+	if (!bIsPredictionInProgress)
+	{
+		bIsPredictionInProgress = true;
+		InPredictionThreadId = FPlatformTLS::GetCurrentThreadId();
+	}
 }
 
-void URollbackBlackboard_InternalWrapper::EndSimulationFrame()
+void URollbackBlackboard::EndPrediction()
 {
-	Blackboard->EndSimulationFrame();
+	check(((BufferFlags & ERollbackBufferFlags::Prediction) == ERollbackBufferFlags::Prediction) && bIsPredictionInProgress && IsInPredictionThread());
+
+	bIsPredictionInProgress = false;
+	InPredictionThreadId = 0;
+
+	for (const TPair<FName, TUniquePtr<BlackboardEntryBase>>& KVP : EntryMap)
+	{
+		KVP.Value->InvalidatePredictiveState();
+	}
 }
 
-void URollbackBlackboard_InternalWrapper::BeginRollback(const FMoverTimeStep& NewBaseTimeStep)
+bool URollbackBlackboard::IsInSimulationThread() const
 {
-	Blackboard->BeginRollback(NewBaseTimeStep);
+	return (InProgressSimFrameThreadId == FPlatformTLS::GetCurrentThreadId());
 }
 
-void URollbackBlackboard_InternalWrapper::EndRollback()
+bool URollbackBlackboard::IsInPredictionThread() const
 {
-	Blackboard->EndRollback();
+	return (InPredictionThreadId == FPlatformTLS::GetCurrentThreadId());
 }
+
diff --git a/Engine/Plugins/Experimental/Mover/Source/Mover/Private/MovementModeStateMachine.cpp b/Engine/Plugins/Experimental/Mover/Source/Mover/Private/MovementModeStateMachine.cpp
--- a/Engine/Plugins/Experimental/Mover/Source/Mover/Private/MovementModeStateMachine.cpp
+++ b/Engine/Plugins/Experimental/Mover/Source/Mover/Private/MovementModeStateMachine.cpp
@@ -7,8 +7,8 @@
 #include "MoverLog.h"
 #include "MoveLibrary/MovementUtils.h"
 #include "MoverComponent.h"
+#include "MoverSimulation.h"
 #include "MoveLibrary/MovementMixer.h"
-#include "MoveLibrary/RollbackBlackboard.h"
 #include "Templates/SubclassOf.h"
 #include "Engine/World.h"
 #include "GameFramework/GameStateBase.h"
@@ -123,6 +123,11 @@ void UMovementModeStateMachine::ClearAllGlobalTransitions()
 	GlobalTransitions.Empty();
 }
 
+void UMovementModeStateMachine::SetSimulation(UMoverSimulation& InSimulation)
+{
+	Simulation = &InSimulation;
+}
+
 void UMovementModeStateMachine::QueueNextMode(FName DesiredNextModeName, bool bShouldReenter)
 {
 	QueueNextModeInternal(DesiredNextModeName, bShouldReenter, /*bIsFromRollback*/ false);
@@ -762,16 +767,15 @@ void UMovementModeStateMachine::AdvanceToNextMode()
 		// signal movement mode change event
 		const UMoverComponent* MoverComp = CastChecked<UMoverComponent>(GetOuter());
 
-		if (URollbackBlackboard_InternalWrapper* RollbackBlackboard = MoverComp->GetRollbackBlackboard_Internal())
-		{
-			FMovementModeChangeRecord ModeChangeRecord;
-			ModeChangeRecord.ModeName = CurrentModeName;
-			ModeChangeRecord.PrevModeName = PreviousModeName;
-			ModeChangeRecord.Frame = CurrentBaseTimeStep.ServerFrame;
-			ModeChangeRecord.SimTimeMs = CurrentBaseTimeStep.BaseSimTimeMs;
+		FRollbackBlackboardSimWrapper RollbackBlackboard = Simulation->GetRollbackBlackboardSimWrapper();
 
-			RollbackBlackboard->TrySet(CommonBlackboard::LastModeChangeRecord, ModeChangeRecord);
-		}
+		FMovementModeChangeRecord ModeChangeRecord;
+		ModeChangeRecord.ModeName = CurrentModeName;
+		ModeChangeRecord.PrevModeName = PreviousModeName;
+		ModeChangeRecord.Frame = CurrentBaseTimeStep.ServerFrame;
+		ModeChangeRecord.SimTimeMs = CurrentBaseTimeStep.BaseSimTimeMs;
+
+		RollbackBlackboard.TrySet(CommonBlackboard::LastModeChangeRecord, ModeChangeRecord);
 
 		Modes[CurrentModeName]->Activate(MoverEventContext);
 
diff --git a/Engine/Plugins/Experimental/Mover/Source/Mover/Private/MoverComponent.cpp b/Engine/Plugins/Experimental/Mover/Source/Mover/Private/MoverComponent.cpp
--- a/Engine/Plugins/Experimental/Mover/Source/Mover/Private/MoverComponent.cpp
+++ b/Engine/Plugins/Experimental/Mover/Source/Mover/Private/MoverComponent.cpp
@@ -4,6 +4,7 @@
 #include "MoverComponent.h"
 #include "MoverSimulationTypes.h"
 #include "MovementModeStateMachine.h"
+#include "MoverSimulation.h"
 #include "MotionWarpingMoverAdapter.h"
 #include "DefaultMovementSet/Modes/WalkingMode.h"
 #include "DefaultMovementSet/Modes/FallingMode.h"
@@ -11,7 +12,7 @@
 #include "MoveLibrary/MovementMixer.h"
 #include "MoveLibrary/MovementUtils.h"
 #include "MoveLibrary/FloorQueryUtils.h"
-#include "MoveLibrary/RollbackBlackboard.h"
+#include "MoveLibrary/RollbackBlackboardWrappers.h"
 #include "InputContainerStruct.h"
 #include "MoverLog.h"
 #include "InstantMovementEffect.h"
@@ -95,21 +96,29 @@ void UMoverComponent::InitializeComponent()
 
 		SimBlackboard = NewObject<UMoverBlackboard>(this, TEXT("MoverBlackboard"), RF_Transient);
 
-		RollbackBlackboard = NewObject<URollbackBlackboard>(this, TEXT("RollbackBlackboard"), RF_Transient);
-		RollbackBlackboard_InternalWrapper = NewObject<URollbackBlackboard_InternalWrapper>(this, TEXT("RollbackBlackboard_Internal"), RF_Transient);
-		RollbackBlackboard_InternalWrapper->Init(*RollbackBlackboard);
 
-		// create any internal entries
-		static URollbackBlackboard::EntrySettings ModeChangeRecordSettings;
-		ModeChangeRecordSettings.SizingPolicy = EBlackboardSizingPolicy::FixedDeclaredSize;
-		ModeChangeRecordSettings.FixedSize = 4;
-		ModeChangeRecordSettings.PersistencePolicy = EBlackboardPersistencePolicy::Forever;
-		ModeChangeRecordSettings.RollbackPolicy = EBlackboardRollbackPolicy::InvalidatedOnRollback;
+		FindDefaultUpdatedComponent();
 
-		RollbackBlackboard->CreateEntry<FMovementModeChangeRecord>(CommonBlackboard::LastModeChangeRecord, ModeChangeRecordSettings);
+
+		// Instantiate our sister backend component that will actually talk to the system driving the simulation
+		if (BackendClass)
+		{
+			UActorComponent* NewLiaisonComp = NewObject<UActorComponent>(this, BackendClass, TEXT("BackendLiaisonComponent"));
+			BackendLiaisonComp.SetObject(NewLiaisonComp);
+			BackendLiaisonComp.SetInterface(CastChecked<IMoverBackendLiaisonInterface>(NewLiaisonComp));
+			if (BackendLiaisonComp)
+			{
+				NewLiaisonComp->RegisterComponent();
+				NewLiaisonComp->InitializeComponent();
+				NewLiaisonComp->SetNetAddressable();
+			}
+		}
+		else
+		{
+			UE_LOG(LogMover, Error, TEXT("No backend class set on %s. Mover actor will not function."), *GetNameSafe(GetOwner()));
+		}
 
 
-		FindDefaultUpdatedComponent();
 
 		// Set up FSM and initial movement states
 		ModeFSM = NewObject<UMovementModeStateMachine>(this, TEXT("MoverStateMachine"), RF_Transient);
@@ -145,24 +154,6 @@ void UMoverComponent::InitializeComponent()
 			ModeFSM->SetDefaultMode(StartingMovementMode);
 			ModeFSM->QueueNextMode(StartingMovementMode);
 		}
-
-		// Instantiate our sister backend component that will actually talk to the system driving the simulation
-		if (BackendClass)
-		{
-			UActorComponent* NewLiaisonComp = NewObject<UActorComponent>(this, BackendClass, TEXT("BackendLiaisonComponent"));
-			BackendLiaisonComp.SetObject(NewLiaisonComp);
-			BackendLiaisonComp.SetInterface(CastChecked<IMoverBackendLiaisonInterface>(NewLiaisonComp));
-			if (BackendLiaisonComp)
-			{
-				NewLiaisonComp->RegisterComponent();
-				NewLiaisonComp->InitializeComponent();
-				NewLiaisonComp->SetNetAddressable();
-			}
-		}
-		else
-		{
-			UE_LOG(LogMover, Error, TEXT("No backend class set on %s. Mover actor will not function."), *GetNameSafe(GetOwner()));
-		}
 	}
 
 	// Gather initial state to fulfill queries
@@ -252,7 +243,19 @@ void UMoverComponent::PostLoad()
 void UMoverComponent::BeginPlay()
 {
 	Super::BeginPlay();
-	
+
+	if (ensureMsgf(Simulation != nullptr, TEXT("MoverComponent requires a Simulation to be set during initialization! The backend should have done this already.")))
+	{
+		ModeFSM->SetSimulation(*Simulation);
+
+		static URollbackBlackboard::EntrySettings ModeChangeRecordSettings;
+		ModeChangeRecordSettings.SizingPolicy = EBlackboardSizingPolicy::FixedDeclaredSize;
+		ModeChangeRecordSettings.FixedSize = 4;
+		ModeChangeRecordSettings.PersistencePolicy = EBlackboardPersistencePolicy::Forever;
+		ModeChangeRecordSettings.RollbackPolicy = EBlackboardRollbackPolicy::InvalidatedOnRollback;
+		RollbackBlackboardExternalWrapper.CreateEntry<FMovementModeChangeRecord>(CommonBlackboard::LastModeChangeRecord, ModeChangeRecordSettings);
+	}
+
 	FindDefaultUpdatedComponent();
 	ensureMsgf(UpdatedComponent != nullptr, TEXT("No root component found on %s. Simulation initialization will most likely fail."), *GetPathNameSafe(GetOwner()));
 
@@ -611,8 +614,8 @@ void UMoverComponent::SimulationTick(const FMoverTimeStep& InTimeStep, const FMo
 		OnPreMovement.Broadcast(MoverTimeStep, SimInput.InputCmd, SimInput.SyncState, SimInput.AuxState);
 	}
 
-	RollbackBlackboard_InternalWrapper->BeginSimulationFrame(MoverTimeStep);
-
+	RollbackBlackboard->BeginSimulationFrame(MoverTimeStep);
+	
 	// Tick the actual simulation. This is where the proposed moves are queried and executed, affecting change to the moving actor's gameplay state and captured in the output sim state
 	if (IsInGameThread())
 	{
@@ -640,7 +643,7 @@ void UMoverComponent::SimulationTick(const FMoverTimeStep& InTimeStep, const FMo
 		}
 	}
 
-	RollbackBlackboard_InternalWrapper->EndSimulationFrame();
+	RollbackBlackboard->EndSimulationFrame();
 
 	if (!SimOutput.MoveRecord.GetTotalMoveDelta().IsZero())
 	{
@@ -1026,6 +1029,15 @@ void UMoverComponent::CreateDefaultInputAndState(FMoverInputCmdContext& OutInput
 
 }
 
+void UMoverComponent::SetSimulation(UMoverSimulation& InSimulation)
+{
+	Simulation = &InSimulation;
+	check(Simulation->RollbackBlackboard != nullptr);
+	
+	RollbackBlackboard = Simulation->RollbackBlackboard;	// TODO: fetching the true blackboard is temporary: this will be removed once simulations are used consistently
+	RollbackBlackboardExternalWrapper = FRollbackBlackboardExternalWrapper(*Simulation->RollbackBlackboard);
+}
+
 void UMoverComponent::HandleImpact(FMoverOnImpactParams& ImpactParams)
 {
 	if (ImpactParams.MovementModeName.IsNone())
@@ -1126,11 +1138,11 @@ void UMoverComponent::OnSimulationRollback(const FMoverSyncState* SyncState, con
 {
 	SimBlackboard->Invalidate(EInvalidationReason::Rollback);
 
-	RollbackBlackboard_InternalWrapper->BeginRollback(NewBaseTimeStep);
+	RollbackBlackboard->BeginRollback(NewBaseTimeStep);
 
 	ModeFSM->OnSimulationRollback(SyncState, AuxState, NewBaseTimeStep);
 
-	RollbackBlackboard_InternalWrapper->EndRollback();
+	RollbackBlackboard->EndRollback();
 	bHasJustRolledBack = true;
 }
 
diff --git a/Engine/Plugins/Experimental/Mover/Source/Mover/Private/MoverSimulation.cpp b/Engine/Plugins/Experimental/Mover/Source/Mover/Private/MoverSimulation.cpp
--- a/Engine/Plugins/Experimental/Mover/Source/Mover/Private/MoverSimulation.cpp
+++ b/Engine/Plugins/Experimental/Mover/Source/Mover/Private/MoverSimulation.cpp
@@ -3,13 +3,17 @@
 #include "MoverSimulation.h"
 
 #include "MoveLibrary/MoverBlackboard.h"
-#include "MoveLibrary/RollbackBlackboard.h"
+#include "MoveLibrary/RollbackBlackboardWrappers.h"
 
 #include UE_INLINE_GENERATED_CPP_BY_NAME(MoverSimulation)
 
 UMoverSimulation::UMoverSimulation()
 {
 	Blackboard = CreateDefaultSubobject<UMoverBlackboard>(TEXT("MoverSimulationBlackboard"));
+	RollbackBlackboard = CreateDefaultSubobject<URollbackBlackboard>(TEXT("MoverSimulationRollbackBlackboard"));
+
+	// TODO: this should be set up depending on whether the simulation is real or purely for prediction 
+	RollbackBlackboardSimWrapper = FRollbackBlackboardSimWrapper(*RollbackBlackboard, /*bIsPredictive*/ false);
 }
 
 const UMoverBlackboard* UMoverSimulation::GetBlackboard() const
@@ -22,17 +26,8 @@ UMoverBlackboard* UMoverSimulation::GetBlackboard_Mutable()
 	return Blackboard;
 }
 
-const URollbackBlackboard_InternalWrapper* UMoverSimulation::GetRollbackBlackboard() const
+FRollbackBlackboardSimWrapper UMoverSimulation::GetRollbackBlackboardSimWrapper() const
 {
-	return RollbackBlackboard;
+	return RollbackBlackboardSimWrapper;
 }
 
-URollbackBlackboard_InternalWrapper* UMoverSimulation::GetRollbackBlackboard_Mutable()
-{
-	return RollbackBlackboard;
-}
-
-void UMoverSimulation::SetRollbackBlackboard(URollbackBlackboard_InternalWrapper* RollbackSimBlackboard)
-{
-	RollbackBlackboard = RollbackSimBlackboard;
-}
\ No newline at end of file
diff --git a/Engine/Plugins/Experimental/Mover/Source/Mover/Public/Backends/MoverNetworkPhysicsLiaisonBase.h b/Engine/Plugins/Experimental/Mover/Source/Mover/Public/Backends/MoverNetworkPhysicsLiaisonBase.h
--- a/Engine/Plugins/Experimental/Mover/Source/Mover/Public/Backends/MoverNetworkPhysicsLiaisonBase.h
+++ b/Engine/Plugins/Experimental/Mover/Source/Mover/Public/Backends/MoverNetworkPhysicsLiaisonBase.h
@@ -9,6 +9,8 @@
 #include "PhysicsMover/PhysicsMoverSimulationTypes.h"
 #include "MoverNetworkPhysicsLiaisonBase.generated.h"
 
+class UMoverSimulation;
+
 #define UE_API MOVER_API
 
 //////////////////////////////////////////////////////////////////////////
@@ -314,6 +316,16 @@ protected:
 	
 	bool bCachedInputIsValid = false;
 	bool bUsingAsyncPhysics = false;
+
+	// Type of simulation to instantiate as the system spins up
+	TSubclassOf<UMoverSimulation> SimulationClass;
+
+	UPROPERTY(Transient)
+	TObjectPtr<UMoverSimulation> Simulation;
+
+	UPROPERTY(Transient)
+	TObjectPtr<URollbackBlackboard> RollbackBlackboard;
+
 };
 
 #undef UE_API
diff --git a/Engine/Plugins/Experimental/Mover/Source/Mover/Public/Backends/MoverNetworkPredictionLiaison.h b/Engine/Plugins/Experimental/Mover/Source/Mover/Public/Backends/MoverNetworkPredictionLiaison.h
--- a/Engine/Plugins/Experimental/Mover/Source/Mover/Public/Backends/MoverNetworkPredictionLiaison.h
+++ b/Engine/Plugins/Experimental/Mover/Source/Mover/Public/Backends/MoverNetworkPredictionLiaison.h
@@ -16,6 +16,7 @@
 #define UE_API MOVER_API
 
 class UMoverComponent;
+class UMoverSimulation;
 
 
 using KinematicMoverStateTypes = TNetworkPredictionStateTypes<FMoverInputCmdContext, FMoverSyncState, FMoverAuxStateContext>;
@@ -91,6 +92,12 @@ protected:
 	TObjectPtr<UMoverComponent> MoverComp;	// the component that we're in charge of driving
 	FMoverSyncState* StartingOutSync;
 	FMoverAuxStateContext* StartingOutAux;
+
+	// Type of simulation to instantiate as the system spins up
+	TSubclassOf<UMoverSimulation> SimulationClass;
+
+	UPROPERTY(Transient)
+	TObjectPtr<UMoverSimulation> Simulation;
 };
 
 #undef UE_API
diff --git a/Engine/Plugins/Experimental/Mover/Source/Mover/Public/Backends/MoverStandaloneLiaison.h b/Engine/Plugins/Experimental/Mover/Source/Mover/Public/Backends/MoverStandaloneLiaison.h
--- a/Engine/Plugins/Experimental/Mover/Source/Mover/Public/Backends/MoverStandaloneLiaison.h
+++ b/Engine/Plugins/Experimental/Mover/Source/Mover/Public/Backends/MoverStandaloneLiaison.h
@@ -8,6 +8,7 @@
 
 #define UE_API MOVER_API
 
+class UMoverSimulation;
 class UMoverComponent;
 class UMoverStandaloneLiaisonComponent;
 class AController;
@@ -111,6 +112,8 @@ public:
 
 	// Begin UActorComponent interface
 	UE_API virtual void RegisterComponentTickFunctions(bool bRegister) override;
+	UE_API virtual void InitializeComponent() override;
+	UE_API virtual void UninitializeComponent() override;
 	UE_API virtual void BeginPlay() override;
 	// End UActorComponent interface
 
@@ -173,6 +176,12 @@ protected:
 	UPROPERTY(EditDefaultsOnly, Category = "Async Movement")
 	bool bUseAsyncMovementSimulationTick = false;
 
+	// Type of simulation to instantiate as the system spins up
+	TSubclassOf<UMoverSimulation> SimulationClass;
+
+	UPROPERTY(Transient)
+	TObjectPtr<UMoverSimulation> Simulation;
+
 	TObjectPtr<UMoverComponent> MoverComp;	// the component that we're in charge of driving
 
 	double CurrentSimTimeMs;
diff --git a/Engine/Plugins/Experimental/Mover/Source/Mover/Public/MoveLibrary/RollbackBlackboard.h b/Engine/Plugins/Experimental/Mover/Source/Mover/Public/MoveLibrary/RollbackBlackboard.h
--- a/Engine/Plugins/Experimental/Mover/Source/Mover/Public/MoveLibrary/RollbackBlackboard.h
+++ b/Engine/Plugins/Experimental/Mover/Source/Mover/Public/MoveLibrary/RollbackBlackboard.h
@@ -2,7 +2,7 @@
 
 #pragma once
 
-#include "Containers/CircularBuffer.h"
+#include "RollbackCircularBuffer.h"
 #include "MoverLog.h"
 #include "UObject/Interface.h"
 #include "UObject/Object.h"
@@ -100,8 +100,8 @@ enum class EBlackboardRollbackPolicy : uint8
 
 
 /**
- * This is the core rollback blackboard, with API for external access.
- * See @URollbackBlackboard_InternalWrapper for in-simulation use.
+ * This is the core rollback blackboard, with API for external access and admin use. 
+ * See @FRollbackBlackboardSimWrapper for in-simulation use.
  */ 
 UCLASS(MinimalAPI)
 class URollbackBlackboard : public UObject
@@ -143,8 +143,9 @@ private:
 
 	enum EEntryIndexType : uint8
 	{
-		External,	// entry index for outside systems to read
-		Internal	// entry index for in-simulation system to read/write
+		External,   // entry index for outside systems to read
+		Internal,   // entry index for in-simulation system to read/write
+		Predictive  // entry index for a predicting system to read/write
 	};
 
 	// untyped base, so we can use entries of multiple types in a single-typed container
@@ -152,7 +153,7 @@ private:
 	{
 	public:
 		BlackboardEntryBase() = delete;
-		UE_API BlackboardEntryBase(const EntrySettings& InSettings, uint32 BufferSize);
+		UE_API BlackboardEntryBase(const EntrySettings& InSettings, uint32 BufferSize, ERollbackBufferFlags BufferFlags);
 
 		virtual ~BlackboardEntryBase()
 		{
@@ -167,14 +168,20 @@ private:
 		// Note that NewPendingFrame means the frame that will now be re-simulated. So any existing entries with a Frame >= NewPendingFrame can be invalidated.
 		UE_API void RollBack(uint32 NewPendingFrame);
 
+		UE_API void InvalidatePredictiveState();
+		
+
 		UE_API static uint32 ComputeBufferSize(EBlackboardSizingPolicy InSizingPolicy, uint32 FixedBufferSize = -1);
 
 	protected:
-		UE_API bool CanReadEntryAt(const EntryTimeStamp& ReaderTimeStamp, EEntryIndexType IndexType) const;
+		UE_API bool CanReadEntryAtTime(const EntryTimeStamp& ReaderTimeStamp, EEntryIndexType IndexType) const;
+		UE_API bool CanReadEntryAtTime(const EntryTimeStamp& ReaderTimeStamp, const EntryTimeStamp& EntryTimeStamp) const;
 
-		EntrySettings Settings;
+		UE_API virtual void ResetPredictiveElement() {}
 
-		TCircularBuffer<EntryTimeStamp> Timestamps;
+		EntrySettings Settings;
+		
+		TRollbackCircularBuffer<EntryTimeStamp> Timestamps;
 
 		uint32 ExternalIdx;	// Indexes to the last entry on a committed frame
 		uint32 InternalIdx;	// Indexes to where in-progress simulations should read/write (matches ExternalIdx when not mid-simulation, and advances if written during simulation)
@@ -190,9 +197,9 @@ private:
 	public:
 
 		// TODO: consider what options need to be specified
-		BlackboardEntry(const EntrySettings& InSettings, uint32 InBufferSize, const EntryT* OptionalInitialObj=nullptr)
-			: BlackboardEntryBase(InSettings, InBufferSize)
-			, EntryBuffer(InBufferSize)
+		BlackboardEntry(const EntrySettings& InSettings, uint32 InBufferSize, ERollbackBufferFlags InBufferFlags, const EntryT* OptionalInitialObj=nullptr)
+			: BlackboardEntryBase(InSettings, InBufferSize, InBufferFlags)
+			, EntryBuffer(InBufferSize, InBufferFlags), DefaultValue()
 		{
 			// JAH TODO: Consider making the initial object required, so we can also initialize where in the buffer we point to
 			if (OptionalInitialObj)
@@ -201,13 +208,20 @@ private:
 				{
 					EntryBuffer[i] = *OptionalInitialObj;
 				}
+
+				DefaultValue = *OptionalInitialObj;
 			}
 		}
 
+		void SetValueAtTime_External(const EntryT& Value, const EntryTimeStamp& TimeStamp)
+		{
+			EntryBuffer[ExternalIdx] = Value;
+			Timestamps[ExternalIdx] = TimeStamp;
+		}
 
-		bool TryGetEntryValue(const EntryTimeStamp& CurrentExternalTime, EntryT& ValueOut)
+		bool TryGetValueAtTime_External(const EntryTimeStamp& CurrentExternalTime, EntryT& ValueOut) const
 		{
-			if (CanReadEntryAt(CurrentExternalTime, EEntryIndexType::External))
+			if (CanReadEntryAtTime(CurrentExternalTime, EEntryIndexType::External))
 			{
 				ValueOut = EntryBuffer[ExternalIdx];
 				return true;
@@ -217,9 +231,9 @@ private:
 		}
 
 		// TODO: Consider whether it's worthwhile to have this function... dangerous to let external folks hold onto references to internal storage, but may allow more efficient editing
-		EntryT* TryGetEntryRef(const EntryTimeStamp& CurrentExternalTime)
+		const EntryT* TryGetRefAtTime_External(const EntryTimeStamp& CurrentExternalTime) const
 		{
-			if (CanReadEntryAt(CurrentExternalTime, EEntryIndexType::External))
+			if (CanReadEntryAtTime(CurrentExternalTime, EEntryIndexType::External))
 			{ 
 				return &EntryBuffer[ExternalIdx];
 			}
@@ -227,9 +241,9 @@ private:
 			return nullptr;
 		}
 
-		bool TryGetEntryValue_Internal(const EntryTimeStamp& InternalTime, EntryT& ValueOut)
+		bool TryGetValueAtTime_Internal(const EntryTimeStamp& InternalTime, EntryT& ValueOut) const
 		{
-			if (CanReadEntryAt(InternalTime, EEntryIndexType::Internal))
+			if (CanReadEntryAtTime(InternalTime, EEntryIndexType::Internal))
 			{
 				ValueOut = EntryBuffer[InternalIdx];
 				return true;
@@ -238,9 +252,9 @@ private:
 			return false;
 		}
 
-		EntryT* TryGetEntryRef_Internal(const EntryTimeStamp& InternalTime)
+		const EntryT* TryGetRefAtTime_Internal(const EntryTimeStamp& InternalTime) const
 		{
-			if (CanReadEntryAt(InternalTime, EEntryIndexType::Internal))
+			if (CanReadEntryAtTime(InternalTime, EEntryIndexType::Internal))
 			{
 				return &EntryBuffer[InternalIdx];
 			}
@@ -248,7 +262,7 @@ private:
 			return nullptr;
 		}
 
-		void SetEntryValue_Internal(const EntryT& Value, const EntryTimeStamp& TimeStamp)
+		void SetValueAtTime_Internal(const EntryT& Value, const EntryTimeStamp& TimeStamp)
 		{
 			InternalIdx = ExternalIdx + 1;	// always writing internal, one slot ahead of external slot
 
@@ -256,10 +270,66 @@ private:
 			Timestamps[InternalIdx] = TimeStamp;
 		}
 
+		bool TryGetValueAtTime_Predictive(const EntryTimeStamp& PredictiveReaderTime, EntryT& ValueOut) const
+		{
+			// For predictive reading, we can either be reading the current External entry
+			// OR the Predictive entry if it's been written to during this Predictive window.
+			const bool bHasPredictiveValue = Timestamps.GetPredictiveElement().IsValid();
+			
+			if (bHasPredictiveValue)
+			{
+				if (CanReadEntryAtTime(PredictiveReaderTime, EEntryIndexType::Predictive))
+				{
+					ValueOut = EntryBuffer.GetPredictiveElement();
+					return true;
+				}
+			}
+			else
+			{
+				// No prediction writes have occurred yet, so we'll try to read the external value
+				return TryGetValueAtTime_External(PredictiveReaderTime, ValueOut);
+			}
+
+			return false;
+		}
+
+		const EntryT* TryGetRefAtTime_Predictive(const EntryTimeStamp& PredictiveReaderTime) const
+		{
+			// For predictive reading, we can either be reading the current External entry
+			// OR the Predictive entry if it's been written to during this Predictive window.
+			const bool bHasPredictiveValue = Timestamps.GetPredictiveElement().IsValid();
+
+			if (bHasPredictiveValue)
+			{
+				if (CanReadEntryAtTime(PredictiveReaderTime, EEntryIndexType::Predictive))
+				{
+					return &EntryBuffer.GetPredictiveElement();
+				}
+			}
+			else
+			{
+				// No prediction writes have occurred yet, so we'll try to read the external value
+				return TryGetRefAtTime_External(PredictiveReaderTime);
+			}
+
+			return nullptr;
+		}
+
+		void SetValueAtTime_Predictive(const EntryT& Value, const EntryTimeStamp& TimeStamp)
+		{
+			// Writes ALWAYS go to the predictive element
+			EntryBuffer.GetPredictiveElement() = Value;
+			Timestamps.GetPredictiveElement() = TimeStamp;
+		}
+
+		virtual void ResetPredictiveElement() override
+		{
+			EntryBuffer.GetPredictiveElement() = DefaultValue;
+		}
+
 	
-		// TODO: Consider whether these need to be TCircularBuffers, or whether a TArray would be just fine.  TCircularBuffers round up to the next power of 2 in capacity, so may be wasting significant memory for larger structures.
-		// Note that these EntryBuffer is always the same capacity as Timestamps
-		TCircularBuffer<EntryT> EntryBuffer;
+		TRollbackCircularBuffer<EntryT> EntryBuffer;
+		EntryT DefaultValue;
 
 
 	};	// end BlackboardEntry
@@ -277,7 +347,7 @@ public:
 		}
 
 		const uint32 BufferSize = BlackboardEntryBase::ComputeBufferSize(InSettings.SizingPolicy, InSettings.FixedSize);
-		EntryMap.Emplace(EntryName, MakeUnique<BlackboardEntry<EntryT>>(InSettings, BufferSize, nullptr));
+		EntryMap.Emplace(EntryName, MakeUnique<BlackboardEntry<EntryT>>(InSettings, BufferSize, BufferFlags, nullptr));
 		return true;
 	}
 
@@ -286,15 +356,16 @@ public:
 		return EntryMap.Contains(EntryName);
 	}
 
+private:
 	/** Store object by a named key, overwriting any existing object */
 	template<typename EntryT>
-	bool TrySet(FName ObjName, EntryT& Obj)
+	bool TrySet_External(FName ObjName, const EntryT& Obj)
 	{
 		// TODO: Verify this isn't happening mid-sim on the sim thread (_Internal should be used instead)
 
-		if (BlackboardEntry<EntryT>* Entry = static_cast<BlackboardEntry<EntryT>*>(FindEntry_Internal(ObjName)))
+		if (BlackboardEntry<EntryT>* Entry = static_cast<BlackboardEntry<EntryT>*>(FindEntry(ObjName)))
 		{
-			Entry->SetEntryValue(Obj);
+			Entry->SetValueAtTime_External(Obj, CurrentSimTimeStamp);
 			return true;
 		}
 
@@ -304,11 +375,11 @@ public:
 
 	/** Attempt to retrieve an object from the blackboard. If found, OutFoundValue will be set. Returns true/false to indicate whether it was found. */
 	template<typename EntryT>
-	bool TryGet(FName ObjName, EntryT& OutFoundValue) const
+	bool TryGet_External(FName ObjName, EntryT& OutFoundValue) const
 	{
-		if (BlackboardEntry<EntryT>* Entry = static_cast<BlackboardEntry<EntryT>*>(FindEntry_Internal(ObjName)))
+		if (BlackboardEntry<EntryT>* Entry = static_cast<BlackboardEntry<EntryT>*>(FindEntry(ObjName)))
 		{
-			return Entry->TryGetEntryValue(CurrentSimTimeStamp, OutFoundValue);
+			return Entry->TryGetValueAtTime_External(CurrentSimTimeStamp, OutFoundValue);
 		}
 
 		return false;
@@ -318,28 +389,27 @@ public:
 	 * Although you can get access to the reference, it is not safe to hold references over time. 
 	 */
 	template<typename EntryT>
-	EntryT* TryGetRef(FName ObjName) const
+	const EntryT* TryGetRef_External(FName ObjName) const
 	{
-		if (BlackboardEntry<EntryT>* Entry = static_cast<BlackboardEntry<EntryT>*>(FindEntry_Internal(ObjName)))
+		if (BlackboardEntry<EntryT>* Entry = static_cast<BlackboardEntry<EntryT>*>(FindEntry(ObjName)))
 		{
-			return Entry->TryGetEntryRef(CurrentSimTimeStamp);
+			return Entry->TryGetRefAtTime_External(CurrentSimTimeStamp);
 		}
 
 		return nullptr;
 	}
 
 
-private:
-	/** Store object by a named key, overwriting any existing object. Should only be called from within the simulation. */
+	/** Store object of a named key, overwriting any existing object. Should only be called from within the simulation. */
 	template<typename EntryT>
 	bool TrySet_Internal(FName ObjName, const EntryT& Obj)
 	{
 		check((bIsSimulationInProgress && (InProgressSimFrameThreadId == FPlatformTLS::GetCurrentThreadId())) || 
 			  (bIsRollbackInProgress && (InRollbackThreadId == FPlatformTLS::GetCurrentThreadId()))) ;
 
-		if (BlackboardEntry<EntryT>* Entry = static_cast<BlackboardEntry<EntryT>*>(FindEntry_Internal(ObjName)))
+		if (BlackboardEntry<EntryT>* Entry = static_cast<BlackboardEntry<EntryT>*>(FindEntry(ObjName)))
 		{
-			Entry->SetEntryValue_Internal(Obj, InProgressSimTimeStamp);
+			Entry->SetValueAtTime_Internal(Obj, InProgressSimTimeStamp);
 			return true;
 		}
 
@@ -352,9 +422,9 @@ private:
 		check((bIsSimulationInProgress && (InProgressSimFrameThreadId == FPlatformTLS::GetCurrentThreadId())) ||
 			  (bIsRollbackInProgress && (InRollbackThreadId == FPlatformTLS::GetCurrentThreadId())));
 
-		if (BlackboardEntry<EntryT>* Entry = static_cast<BlackboardEntry<EntryT>*>(FindEntry_Internal(ObjName)))
+		if (BlackboardEntry<EntryT>* Entry = static_cast<BlackboardEntry<EntryT>*>(FindEntry(ObjName)))
 		{
-			return Entry->TryGetEntryValue_Internal(InProgressSimTimeStamp, OutFoundValue);
+			return Entry->TryGetValueAtTime_Internal(InProgressSimTimeStamp, OutFoundValue);
 		}
 
 		return false;
@@ -366,17 +436,59 @@ private:
 		check((bIsSimulationInProgress && (InProgressSimFrameThreadId == FPlatformTLS::GetCurrentThreadId())) ||
 			  (bIsRollbackInProgress && (InRollbackThreadId == FPlatformTLS::GetCurrentThreadId())));
 
-		if (BlackboardEntry<EntryT>* Entry = static_cast<BlackboardEntry<EntryT>*>(FindEntry_Internal(ObjName)))
+		if (BlackboardEntry<EntryT>* Entry = static_cast<BlackboardEntry<EntryT>*>(FindEntry(ObjName)))
 		{
-			return Entry->TryGetEntryRef_Internal(InProgressSimTimeStamp);
+			return Entry->TryGetRefAtTime_Internal(InProgressSimTimeStamp);
 		}
 
 		return nullptr;
 	}
 
+private:
+	/** Store object by a named key, overwriting any existing object. Should only be called from within the predicting simulation. */
+	template<typename EntryT>
+	bool TrySet_Predictive(FName ObjName, const EntryT& Obj)
+	{
+		check(bIsPredictionInProgress && (InPredictionThreadId == FPlatformTLS::GetCurrentThreadId()));
 
+		if (BlackboardEntry<EntryT>* Entry = static_cast<BlackboardEntry<EntryT>*>(FindEntry(ObjName)))
+		{
+			Entry->SetValueAtTime_Predictive(Obj, InProgressSimTimeStamp);
+			return true;
+		}
 
-	BlackboardEntryBase* FindEntry_Internal(FName EntryName) const
+		return false;
+	}
+
+	template<typename EntryT>
+	bool TryGet_Predictive(FName ObjName, EntryT& OutFoundValue) const
+	{
+		check(bIsPredictionInProgress && (InPredictionThreadId == FPlatformTLS::GetCurrentThreadId()));
+
+		if (BlackboardEntry<EntryT>* Entry = static_cast<BlackboardEntry<EntryT>*>(FindEntry(ObjName)))
+		{
+			return Entry->TryGetValueAtTime_Predictive(InProgressSimTimeStamp, OutFoundValue);
+		}
+
+		return false;
+	}
+
+	template<typename EntryT>
+	const EntryT* TryGetRef_Predictive(FName ObjName) const
+	{
+		check(bIsPredictionInProgress && (InPredictionThreadId == FPlatformTLS::GetCurrentThreadId()));
+
+		if (BlackboardEntry<EntryT>* Entry = static_cast<BlackboardEntry<EntryT>*>(FindEntry(ObjName)))
+		{
+			return Entry->TryGetRefAtTime_Predictive(InProgressSimTimeStamp);
+		}
+
+		return nullptr;
+	}
+
+private:
+
+	BlackboardEntryBase* FindEntry(FName EntryName) const
 	{
 		if (const TUniquePtr<BlackboardEntryBase>* EntryPtr = EntryMap.Find(EntryName))
 		{
@@ -387,90 +499,46 @@ private:
 	}
 
 
-private:
+public:	// Administration functions for keeping the blackboard in sync with the simulation
+
+	// Begin/EndSimulationFrames: advances true simulation flow, regardless of whether this is resimulation or not. Calls will bookend a simulation frame.
 	UE_API void BeginSimulationFrame(const FMoverTimeStep& PendingTimeStep);
 	UE_API void EndSimulationFrame();
 
-	UE_API void BeginRollback(const FMoverTimeStep& NewBaseTimeStep);	// NewBaseTimeStep represents the current time and the PENDING frame that's next to be resimulated
+	// Begin/EndRollback: reverts the blackboard back to an older state. EndRollback will be called before any simulation frames are started (resim or not).
+	UE_API void BeginRollback(const FMoverTimeStep& NewBaseTimeStep);	// NewBaseTimeStep represents the current time and the PENDING step that's next to be resimulated
 	UE_API void EndRollback();
 
+	// BeginPredictionFrame/EndPrediction: this is for predictive simulation that is temporary and can't rollback. Note there may be multiple calls to BeginPredictionFrame before a single call to EndPrediction.
+	UE_API void BeginPredictionFrame(const FMoverTimeStep& PendingTimeStep);	// Called to start a new prediction OR continue a sequential frame
+	UE_API void EndPrediction();	// Called at the end of ALL prediction frames
 
 private:
+	UE_API bool IsInSimulationThread() const;
+	UE_API bool IsInPredictionThread() const;
+
+	ERollbackBufferFlags BufferFlags = ERollbackBufferFlags::Prediction;
 	
 	TMap<FName, TUniquePtr<BlackboardEntryBase>> EntryMap;
 
 
 	bool bIsSimulationInProgress = false;
 	bool bIsResimulating = false;
-	uint32 InProgressSimFrameThreadId;
+	uint32 InProgressSimFrameThreadId = 0;
 
 	bool bIsRollbackInProgress = false;
-	uint32 InRollbackThreadId;
+	uint32 InRollbackThreadId = 0;
 
+	bool bIsPredictionInProgress = false;
+	uint32 InPredictionThreadId = 0;
 
 	EntryTimeStamp CurrentSimTimeStamp;	// this is the "committed" simulation time, and will lag behind InProgressSimTimeStamp while in the middle of a simulation frame
 	EntryTimeStamp InProgressSimTimeStamp;	// this is the sim time that's actively being used mid-simulation. It is only valid during a simulation step.
 
-	friend class URollbackBlackboard_InternalWrapper;	// expose non-public API to this wrapper class
+	// we expose user API through these wrapper types
+	friend struct FRollbackBlackboardSimWrapper;
+	friend struct FRollbackBlackboardExternalWrapper;
 };
 
 
-
-/** 
- * Wrapper class for in-simulation access to the blackboard. This exposes otherwise private API, and redirects
- * set/get operations to the internal versions.
- */
-UCLASS(MinimalAPI)
-class URollbackBlackboard_InternalWrapper : public UObject
-{
-	GENERATED_BODY()
-
-public:
-	void Init(URollbackBlackboard& InBlackboard)
-	{
-		Blackboard = &InBlackboard;
-	}
-
-	template<typename EntryT>
-	bool CreateEntry(FName EntryName, const URollbackBlackboard::EntrySettings& InSettings)
-	{
-		return Blackboard->CreateEntry<EntryT>(EntryName, InSettings);
-	}
-
-	bool HasEntry(FName EntryName) const
-	{
-		return Blackboard->HasEntry(EntryName);
-	}
-
-	template<typename EntryT>
-	bool TrySet(FName ObjName, const EntryT& Obj)
-	{
-		return Blackboard->TrySet_Internal<EntryT>(ObjName, Obj);
-	}
-
-	template<typename EntryT>
-	bool TryGet(FName ObjName, EntryT& OutFoundValue) const
-	{
-		return Blackboard->TryGet_Internal(ObjName, OutFoundValue);
-	}
-
-	template<typename EntryT>
-	const EntryT* TryGetRef(FName ObjName) const
-	{
-		return Blackboard->TryGetRef_Internal<EntryT>(ObjName);
-	}
-	
-	// TODO: Also support a mutable TryGetRef that gives you a writable reference
-
-	UE_API void BeginSimulationFrame(const FMoverTimeStep& PendingTimeStep);
-	UE_API void EndSimulationFrame();
-
-	UE_API void BeginRollback(const FMoverTimeStep& NewBaseTimeStep);
-	UE_API void EndRollback();
-
-private: 
-	UPROPERTY()
-	TObjectPtr<URollbackBlackboard> Blackboard;
-};
-
 #undef UE_API
diff --git a/Engine/Plugins/Experimental/Mover/Source/Mover/Public/MoveLibrary/RollbackBlackboardWrappers.h b/Engine/Plugins/Experimental/Mover/Source/Mover/Public/MoveLibrary/RollbackBlackboardWrappers.h
new file mode 100644
--- /dev/null
+++ b/Engine/Plugins/Experimental/Mover/Source/Mover/Public/MoveLibrary/RollbackBlackboardWrappers.h
@@ -0,0 +1,200 @@
+// Copyright Epic Games, Inc. All Rights Reserved.
+
+#pragma once
+
+#include "RollbackBlackboard.h"
+#include "RollbackBlackboardWrappers.generated.h"
+
+
+#define UE_API MOVER_API
+
+
+/**
+ * Wrapper class for external parties to access the blackboard. 
+ * External parties = users outside of an active movement simulation such as animation and gameplay scripting
+ */
+USTRUCT(BlueprintType)
+struct FRollbackBlackboardExternalWrapper
+{
+	GENERATED_BODY()
+
+public:
+	UE_API FRollbackBlackboardExternalWrapper(URollbackBlackboard& InBlackboard)
+		: Blackboard(&InBlackboard) {}
+
+	UE_API FRollbackBlackboardExternalWrapper(const FRollbackBlackboardExternalWrapper& Other)
+		: Blackboard(Other.Blackboard) {}
+
+	// Default constructor exists only for reflection / BP support.  Typically constructed with a blackboard or as a copy.
+	FRollbackBlackboardExternalWrapper() {}
+
+	template<typename EntryT>
+	bool CreateEntry(FName EntryName, const URollbackBlackboard::EntrySettings& InSettings)
+	{
+		if (Blackboard == nullptr)
+		{
+			return false;
+		}
+
+		return Blackboard->CreateEntry<EntryT>(EntryName, InSettings);
+	}
+
+	bool HasEntry(FName EntryName) const
+	{
+		if (Blackboard == nullptr)
+		{
+			return false;
+		}
+
+		return Blackboard->HasEntry(EntryName);
+	}
+
+	template<typename EntryT>
+	bool TrySet(FName ObjName, const EntryT& Obj)
+	{
+		if (Blackboard == nullptr)
+		{
+			return false;
+		}
+
+		return Blackboard->TrySet_External<EntryT>(ObjName, Obj);
+	}
+
+	template<typename EntryT>
+	bool TryGet(FName ObjName, EntryT& OutFoundValue) const
+	{
+		if (Blackboard == nullptr)
+		{
+			return false;
+		}
+
+		return Blackboard->TryGet_External(ObjName, OutFoundValue);
+	}
+
+	template<typename EntryT>
+	const EntryT* TryGetRef(FName ObjName) const
+	{
+		if (Blackboard == nullptr)
+		{
+			return false;
+		}
+
+		return Blackboard->TryGetRef_External<EntryT>(ObjName);
+	}
+
+private:
+	UPROPERTY()
+	TObjectPtr<URollbackBlackboard> Blackboard = nullptr;
+};
+
+
+/**
+ * Wrapper class for in-simulation access to the blackboard. This exposes otherwise private API, and redirects
+ * set/get operations to the in-simulation locations. This is what all movement mechanisms will interact with, including
+ * movement modes, layered moves, etc.
+ * This also supports a predictive mode, so the blackboard can be used and even temporarily written to, without
+ * affecting the true simulation's state or external parties. Users of this wrapper will not be aware whether it's
+ * part of a true movement sim or a predictive one.
+ */
+USTRUCT(BlueprintType)
+struct FRollbackBlackboardSimWrapper
+{
+	GENERATED_BODY()
+
+public:
+	UE_API FRollbackBlackboardSimWrapper(URollbackBlackboard& InBlackboard, bool bUsePredictiveMode = false)
+		: Blackboard(&InBlackboard), bInPredictiveMode(bUsePredictiveMode) {}
+
+	UE_API FRollbackBlackboardSimWrapper(const FRollbackBlackboardSimWrapper& Other)
+		: Blackboard(Other.Blackboard), bInPredictiveMode(Other.bInPredictiveMode) {}
+
+	// Default constructor exists only for reflection / BP support. Typically constructed with a blackboard or as a copy.
+	FRollbackBlackboardSimWrapper() {}
+
+	template<typename EntryT>
+	bool CreateEntry(FName EntryName, const URollbackBlackboard::EntrySettings& InSettings)
+	{
+		if (Blackboard == nullptr)
+		{
+			return false;
+		}
+
+		return Blackboard->CreateEntry<EntryT>(EntryName, InSettings);
+	}
+
+	bool HasEntry(FName EntryName) const
+	{
+		if (Blackboard == nullptr)
+		{
+			return false;
+		}
+
+		return Blackboard->HasEntry(EntryName);
+	}
+
+	template<typename EntryT>
+	bool TrySet(FName ObjName, const EntryT& Obj)
+	{
+		if (Blackboard == nullptr)
+		{
+			return false;
+		}
+
+		if (bInPredictiveMode)
+		{
+			return Blackboard->TrySet_Predictive<EntryT>(ObjName, Obj);
+		}
+		else
+		{
+			return Blackboard->TrySet_Internal<EntryT>(ObjName, Obj);
+		}
+	}
+
+	template<typename EntryT>
+	bool TryGet(FName ObjName, EntryT& OutFoundValue) const
+	{
+		if (Blackboard == nullptr)
+		{
+			return false;
+		}
+
+		if (bInPredictiveMode)
+		{
+			return Blackboard->TryGet_Predictive(ObjName, OutFoundValue);
+		}
+		else
+		{
+			return Blackboard->TryGet_Internal(ObjName, OutFoundValue);
+		}
+	}
+
+	template<typename EntryT>
+	const EntryT* TryGetRef(FName ObjName) const
+	{
+		if (Blackboard == nullptr)
+		{
+			return false;
+		}
+
+		if (bInPredictiveMode)
+		{
+			return Blackboard->TryGetRef_Predictive<EntryT>(ObjName);
+		}
+		else
+		{
+			return Blackboard->TryGetRef_Internal<EntryT>(ObjName);
+		}
+	}
+
+	// TODO: Also support a mutable TryGetRef that gives you a writable reference
+
+private:
+
+	UPROPERTY()
+	TObjectPtr<URollbackBlackboard> Blackboard = nullptr;
+
+	UPROPERTY()
+	bool bInPredictiveMode = false;
+};
+
+#undef UE_API
\ No newline at end of file
diff --git a/Engine/Plugins/Experimental/Mover/Source/Mover/Public/MoveLibrary/RollbackCircularBuffer.h b/Engine/Plugins/Experimental/Mover/Source/Mover/Public/MoveLibrary/RollbackCircularBuffer.h
new file mode 100644
--- /dev/null
+++ b/Engine/Plugins/Experimental/Mover/Source/Mover/Public/MoveLibrary/RollbackCircularBuffer.h
@@ -0,0 +1,171 @@
+// Copyright Epic Games, Inc. All Rights Reserved.
+
+#pragma once
+
+#include "CoreTypes.h"
+#include "Containers/Array.h"
+
+enum class ERollbackBufferFlags
+{
+	None = 0,
+	Prediction = (1 << 0),		// whether this buffer will allocate additional space for use with prediction
+};
+ENUM_CLASS_FLAGS(ERollbackBufferFlags)
+
+
+/**
+ * Circular buffer specialized for use with rollback simulations providing concurrent access 
+ * between active simulation and external users, as well as extra space for predictive 
+ * work that's temporary and to be discarded.
+ *
+ * This is based on TCircularBuffer's implementation. The size of the buffer is rounded up to 
+ * the next power of two in order speed up indexing operations using a simple bit mask instead 
+ * of the commonly used modulus operator that may be slow on some platforms.
+ */
+template<typename InElementType> class TRollbackCircularBuffer
+{
+public:
+	using ElementType = InElementType;
+
+	/**
+	 * Creates and initializes a new instance of the TRollbackCircularBuffer class.
+	 *
+	 * @param Capacity The number of elements that the buffer can store
+	 */
+	[[nodiscard]] explicit TRollbackCircularBuffer(uint32 Capacity, ERollbackBufferFlags Flags)
+	{
+		checkSlow(Capacity <= 0x8000000u);	// Round-up doesn't work for values above this.
+
+		Capacity = FMath::Max(Capacity, 1u);	// The circular buffer indexmask doesn't work for a size of 0, so make 1 the minimum
+
+		uint32 NumElementsToAllocate = FMath::RoundUpToPowerOfTwo(Capacity);
+		NonPredictedCapacity = NumElementsToAllocate;
+		
+		// If prediction is enabled, we'll use the last element as a "shadow" element for writes during prediction
+		if ((Flags & ERollbackBufferFlags::Prediction) == ERollbackBufferFlags::Prediction)
+		{
+			++NumElementsToAllocate;	// we need 1 additional predictive entry
+		}
+
+		check(NonPredictedCapacity <= NumElementsToAllocate);
+
+		Elements.AddDefaulted(NumElementsToAllocate);
+		IndexMask = NonPredictedCapacity - 1;
+	}
+
+	/**
+	 * Creates and initializes a new instance of the TRollbackCircularBuffer class.
+	 *
+	 * @param Capacity The number of elements that the buffer can store (will be rounded up to the next power of 2).
+	 * @param InFlags
+	 * @param InitialValue The initial value for the buffer's elements.
+	 */
+	[[nodiscard]] TRollbackCircularBuffer(uint32 Capacity, ERollbackBufferFlags Flags, const ElementType& InitialValue)
+	{
+		checkSlow(Capacity <= 0x8000000u);	// Round-up doesn't work for values above this.
+
+		Capacity = FMath::Max(Capacity, 1u);	// The circular buffer indexmask doesn't work for a size of 0, so make 1 the minimum
+
+		uint32 NumElementsToAllocate = FMath::RoundUpToPowerOfTwo(Capacity);
+		NonPredictedCapacity = NumElementsToAllocate;
+
+		// If prediction is enabled, we'll use the last element as a "shadow" element for writes during prediction
+		if ((Flags & ERollbackBufferFlags::Prediction) == ERollbackBufferFlags::Prediction)
+		{
+			++NumElementsToAllocate;	// we need 1 additional predictive entry
+		}
+
+		check(NonPredictedCapacity <= NumElementsToAllocate);
+
+		Elements.Init(InitialValue, NumElementsToAllocate);
+		IndexMask = NonPredictedCapacity - 1;
+	}
+
+public:
+
+	/**
+	 * Returns the mutable element at the specified index.
+	 *
+	 * @param Index The index of the element to return.
+	 */
+	[[nodiscard]] UE_FORCEINLINE_HINT ElementType& operator[](uint32 Index)
+	{
+		return Elements[Index & IndexMask];
+	}
+
+	/**
+	 * Returns the mutable predictive element. Only valid if the buffer was created with prediction supported.
+	 */
+	[[nodiscard]] UE_FORCEINLINE_HINT ElementType& GetPredictiveElement()
+	{
+		check((uint32)Elements.Num() > NonPredictedCapacity);	// If this fails, it's likely this buffer was created without prediction support, but someone is trying to use it in a predictive context.
+
+		return Elements[Elements.Num()-1];	// the predictive element is always the last one in the array
+	}
+
+	/*
+	 * Returns the immutable element at the specified index.
+	 *
+	 * @param Index The index of the element to return.
+	 */
+	[[nodiscard]] UE_FORCEINLINE_HINT const ElementType& operator[](uint32 Index) const
+	{
+		return Elements[Index & IndexMask];
+	}
+
+
+	/**
+	 * Returns the immutable predictive element. Only valid if the buffer was created with prediction supported.
+	 */
+	[[nodiscard]] UE_FORCEINLINE_HINT const ElementType& GetPredictiveElement() const
+	{
+		check((uint32)Elements.Num() > NonPredictedCapacity);  // If this fails, it's likely this buffer was created without prediction support, but someone is trying to use it in a predictive context. We don't have a valid element ref to return.
+
+		return Elements[Elements.Num() - 1];	// the predictive element is always the last one in the array
+	}
+
+public:
+
+	/**
+	 * Returns the number of elements that the buffer can naturally hold. 
+	 *
+	 * @return Buffer capacity.
+	 */
+	[[nodiscard]] UE_FORCEINLINE_HINT uint32 Capacity() const
+	{
+		return NonPredictedCapacity;
+	}
+
+	/**
+	 * Calculates the index that follows the given index.
+	 *
+	 * @param CurrentIndex The current index.
+	 * @return The next index.
+	 */
+	[[nodiscard]] UE_FORCEINLINE_HINT uint32 GetNextIndex(uint32 CurrentIndex) const
+	{
+		return ((CurrentIndex + 1) & IndexMask);
+	}
+
+	/**
+	 * Calculates the index previous to the given index.
+	 *
+	 * @param CurrentIndex The current index.
+	 * @return The previous index.
+	 */
+	[[nodiscard]] UE_FORCEINLINE_HINT uint32 GetPreviousIndex(uint32 CurrentIndex) const
+	{
+		return ((CurrentIndex - 1) & IndexMask);
+	}
+
+private:
+
+	/** Holds the mask for indexing the buffer's elements. This mask covers only the non-predicted elements. The predictive element in the array will not be accessible. */
+	uint32 IndexMask;
+
+	/** This is the true number of usable element slots. If prediction is enabled, this will be smaller than Elements.Num() */
+	uint32 NonPredictedCapacity;	
+
+	/** Holds the buffer's elements. If prediction is enabled, this will have 1 additional element at the end of the array. */
+	TArray<ElementType> Elements;
+};
diff --git a/Engine/Plugins/Experimental/Mover/Source/Mover/Public/MovementModeStateMachine.h b/Engine/Plugins/Experimental/Mover/Source/Mover/Public/MovementModeStateMachine.h
--- a/Engine/Plugins/Experimental/Mover/Source/Mover/Public/MovementModeStateMachine.h
+++ b/Engine/Plugins/Experimental/Mover/Source/Mover/Public/MovementModeStateMachine.h
@@ -6,6 +6,7 @@
 #include "Misc/TransactionallySafeRWLock.h"
 #include "Templates/SubclassOf.h"
 #include "MovementMode.h"
+#include "MoveLibrary/RollbackBlackboardWrappers.h"
 #include "InstantMovementEffect.h"
 #include "MovementModeStateMachine.generated.h"
 
@@ -37,6 +38,9 @@ public:
 	UE_API void UnregisterGlobalTransition(TObjectPtr<UBaseMovementModeTransition> Transition);
 	UE_API void ClearAllGlobalTransitions();
 
+	// TODO: Warning: this function will be removed once the full simulation context (including blackboard) comes in as part of the Simulation**** function params
+	UE_API void SetSimulation(UMoverSimulation& InSimulation);
+
 	UE_API void SetDefaultMode(FName NewDefaultModeName);
 
 	UE_API void QueueNextMode(FName DesiredNextModeName, bool bShouldReenter=false);
@@ -88,6 +92,9 @@ protected:
 	UPROPERTY(Transient)
 	TObjectPtr<UImmediateMovementModeTransition> QueuedModeTransition;
 
+	UPROPERTY(Transient)
+	TObjectPtr<UMoverSimulation> Simulation;
+
 	FName DefaultModeName = NAME_None;
 	FName CurrentModeName = NAME_None;
 
diff --git a/Engine/Plugins/Experimental/Mover/Source/Mover/Public/MoverComponent.h b/Engine/Plugins/Experimental/Mover/Source/Mover/Public/MoverComponent.h
--- a/Engine/Plugins/Experimental/Mover/Source/Mover/Public/MoverComponent.h
+++ b/Engine/Plugins/Experimental/Mover/Source/Mover/Public/MoverComponent.h
@@ -13,6 +13,7 @@
 #include "LayeredMoveBase.h"
 #include "MoveLibrary/BasedMovementUtils.h"
 #include "MoveLibrary/ConstrainedMoveUtils.h"
+#include "MoveLibrary/RollbackBlackboardWrappers.h"
 #if UE_ENABLE_INCLUDE_ORDER_DEPRECATED_IN_5_4
 #include "Engine/HitResult.h"
 #endif // UE_ENABLE_INCLUDE_ORDER_DEPRECATED_IN_5_4
@@ -29,7 +30,6 @@ struct FMoverSimulationEventData;
 class UMovementModeStateMachine;
 class UMovementMixer;
 class URollbackBlackboard;
-class URollbackBlackboard_InternalWrapper;
 
 namespace MoverComponentConstants
 {
@@ -541,9 +541,8 @@ public:	// Queries
 
 	MOVER_API UMoverBlackboard* GetSimBlackboard_Mutable() const;
 
-
-	URollbackBlackboard* GetRollbackBlackboard() const { return RollbackBlackboard.Get(); }
-	URollbackBlackboard_InternalWrapper* GetRollbackBlackboard_Internal() const { return RollbackBlackboard_InternalWrapper.Get(); }
+	// Access the version of the movement blackboard for users outside of the movement simulation
+	FRollbackBlackboardExternalWrapper GetRollbackBlackboardExternal() const { return RollbackBlackboardExternalWrapper; }
 
 	/** Find settings object by type. Returns null if there is none of that type */
 	const IMovementSettingsInterface* FindSharedSettings(const UClass* ByType) const { return FindSharedSettings_Mutable(ByType); }
@@ -695,6 +694,8 @@ protected:
 public:
 	MOVER_API virtual void CreateDefaultInputAndState(FMoverInputCmdContext& OutInputCmd, FMoverSyncState& OutSyncState, FMoverAuxStateContext& OutAuxState) const;
 
+	MOVER_API void SetSimulation(UMoverSimulation& Simulation);		// Temporary support function. Will be removed once simulations are used consistently.
+
 	/** Handle a blocking impact.*/
 	UFUNCTION(BlueprintCallable, Category = Mover)
 	MOVER_API void HandleImpact(FMoverOnImpactParams& ImpactParams);
@@ -789,12 +790,17 @@ protected:
 	UPROPERTY(Transient)
 	TObjectPtr<UMoverBlackboard> SimBlackboard = nullptr;
 
-	/** Used to store cached data & computations between decoupled systems, that can be referenced by name. Rollback-aware. */
+	/** This is what all movement mechanisms should be interacting with during simulation, rather than relying on actors/components. External parties should generally not see it. */
 	UPROPERTY(Transient)
-	TObjectPtr<URollbackBlackboard> RollbackBlackboard = nullptr;
+	TObjectPtr<UMoverSimulation> Simulation = nullptr;
 
+	// Accessors used to interface with the simulation's rollback blackboard
 	UPROPERTY(Transient)
-	TObjectPtr<URollbackBlackboard_InternalWrapper> RollbackBlackboard_InternalWrapper = nullptr;	// This is a thin layer for use only by in-simulation users
+	FRollbackBlackboardExternalWrapper RollbackBlackboardExternalWrapper;
+
+	/** Ref to the actual blackboard. This is temporary and will be removed once simulations are used consistently. */
+	UPROPERTY(Transient)
+	TObjectPtr<URollbackBlackboard> RollbackBlackboard = nullptr;
 
 	/**
 	 * Layered moves registered on this component that can be activated regardless of the current mode
diff --git a/Engine/Plugins/Experimental/Mover/Source/Mover/Public/MoverSimulation.h b/Engine/Plugins/Experimental/Mover/Source/Mover/Public/MoverSimulation.h
--- a/Engine/Plugins/Experimental/Mover/Source/Mover/Public/MoverSimulation.h
+++ b/Engine/Plugins/Experimental/Mover/Source/Mover/Public/MoverSimulation.h
@@ -3,14 +3,13 @@
 #pragma once
 
 #include "UObject/Object.h"
-
+#include "MoveLibrary/RollbackBlackboardWrappers.h"
 #include "MoverSimulation.generated.h"
 
 #define UE_API MOVER_API
 
 class UMoverBlackboard;
 struct FMoverSyncState;
-class URollbackBlackboard_InternalWrapper;
 
 /**
 * WIP Base class for a Mover simulation.
@@ -34,11 +33,7 @@ public:
 	UE_API UMoverBlackboard* GetBlackboard_Mutable();
 
 	UFUNCTION(BlueprintPure, Category = Mover)
-	UE_API const URollbackBlackboard_InternalWrapper* GetRollbackBlackboard() const;
-
-	UFUNCTION(BlueprintPure, Category = Mover)
-	UE_API URollbackBlackboard_InternalWrapper* GetRollbackBlackboard_Mutable();
-
+	UE_API FRollbackBlackboardSimWrapper GetRollbackBlackboardSimWrapper() const;
 
 	/**
 	* Attempt to teleport to TargetTransform. The teleport is not guaranteed to happen. This function is meant to be called by an instant movement effect as part of its effect application.
@@ -52,16 +47,18 @@ public:
 	virtual void AttemptTeleport(const FMoverTimeStep& TimeStep, const FTransform& TargetTransform, bool bUseActorRotation, FMoverSyncState& OutputState) {}
 
 
-	// Used during initialization only
-	UE_API void SetRollbackBlackboard(URollbackBlackboard_InternalWrapper* RollbackSimBlackboard);
-
 protected:
 	UPROPERTY(Transient)
 	TObjectPtr<UMoverBlackboard> Blackboard = nullptr;
 
 	UPROPERTY(Transient)
-	TObjectPtr<URollbackBlackboard_InternalWrapper> RollbackBlackboard = nullptr;
+	TObjectPtr<URollbackBlackboard> RollbackBlackboard = nullptr;
+
+	FRollbackBlackboardSimWrapper RollbackBlackboardSimWrapper;
 
+// TODO: These friend clases are temporary until the control flow changes to include a Simulation object consistently
+friend class UMoverComponent;
+friend class UMoverNetworkPhysicsLiaisonComponentBase;
 };
 
 #undef UE_API
diff --git a/Engine/Plugins/Experimental/MoverTests/MoverTests.uplugin b/Engine/Plugins/Experimental/MoverTests/MoverTests.uplugin
--- a/Engine/Plugins/Experimental/MoverTests/MoverTests.uplugin
+++ b/Engine/Plugins/Experimental/MoverTests/MoverTests.uplugin
@@ -29,6 +29,10 @@
 		{
 			"Name": "MoverExamples",
 			"Enabled": true
+		},
+		{
+			"Name": "RuntimeTests",
+			"Enabled": true
 		}
 	]
 }
\ No newline at end of file
diff --git a/Engine/Plugins/Experimental/MoverTests/Source/MoverTests/MoverTests.Build.cs b/Engine/Plugins/Experimental/MoverTests/Source/MoverTests/MoverTests.Build.cs
--- a/Engine/Plugins/Experimental/MoverTests/Source/MoverTests/MoverTests.Build.cs
+++ b/Engine/Plugins/Experimental/MoverTests/Source/MoverTests/MoverTests.Build.cs
@@ -31,6 +31,7 @@ public class MoverTests : ModuleRules
 				"Engine",
 				"Slate",
 				"SlateCore",
+				"RuntimeTests",
 				"Mover",
 				"MoverExamples"
 				// ... add other public dependencies that you statically link with here ...
diff --git a/Engine/Plugins/Experimental/MoverTests/Source/MoverTests/Private/RollbackBlackboardTests.cpp b/Engine/Plugins/Experimental/MoverTests/Source/MoverTests/Private/RollbackBlackboardTests.cpp
new file mode 100644
--- /dev/null
+++ b/Engine/Plugins/Experimental/MoverTests/Source/MoverTests/Private/RollbackBlackboardTests.cpp
@@ -0,0 +1,175 @@
+// Copyright Epic Games, Inc. All Rights Reserved.
+
+#include "CoreMinimal.h"
+#include "MoveLibrary/RollbackBlackboardWrappers.h"
+#include "MoverTypes.h"
+#include "MoverSimulationTypes.h"
+#include "Misc/LowLevelTestAdapter.h"
+#include "Misc/AutomationTest.h"
+#include "EngineRuntimeTests.h"
+
+static int FixedMsPerFrame = 10;
+static void AdvanceTimeStepFrame(FMoverTimeStep& TimeStepToModify)
+{
+	++TimeStepToModify.ServerFrame;
+	TimeStepToModify.BaseSimTimeMs += FixedMsPerFrame;
+	TimeStepToModify.StepMs = FixedMsPerFrame;
+}
+
+// Simple test cases
+TEST_CASE_NAMED(FRollbackBlackboardTest, "Mover::RollbackBlackboard::TestOperations", "[Mover][Blackboard]")
+{
+	TObjectPtr<URollbackBlackboard> RBBB = NewObject<URollbackBlackboard>(GetTransientPackage(), TEXT("TestRollbackBlackboard"), RF_Transient);
+
+	REQUIRE(RBBB != nullptr);
+
+	FRollbackBlackboardSimWrapper RBBBSimAccessor(*RBBB, /*bPredictiveMode*/ false);
+	FRollbackBlackboardExternalWrapper RBBBExternalAccessor(*RBBB);
+
+	FMoverTimeStep CurrentSimTime;
+	CurrentSimTime.ServerFrame = 22;
+	CurrentSimTime.BaseSimTimeMs = 300;
+	CurrentSimTime.StepMs = 10;
+
+	{
+		URollbackBlackboard::EntrySettings ModeChangeRecordSettings;
+		ModeChangeRecordSettings.SizingPolicy = EBlackboardSizingPolicy::FixedDeclaredSize;
+		ModeChangeRecordSettings.FixedSize = 4;
+		ModeChangeRecordSettings.PersistencePolicy = EBlackboardPersistencePolicy::Forever;
+		ModeChangeRecordSettings.RollbackPolicy = EBlackboardRollbackPolicy::InvalidatedOnRollback;
+
+		RBBBSimAccessor.CreateEntry<FMovementModeChangeRecord>(CommonBlackboard::LastModeChangeRecord, ModeChangeRecordSettings);
+	}
+
+	{
+		URollbackBlackboard::EntrySettings Integer32TestSettings;
+		Integer32TestSettings.SizingPolicy = EBlackboardSizingPolicy::FixedDeclaredSize;
+		Integer32TestSettings.FixedSize = 6;
+		Integer32TestSettings.PersistencePolicy = EBlackboardPersistencePolicy::Forever;
+		Integer32TestSettings.RollbackPolicy = EBlackboardRollbackPolicy::InvalidatedOnRollback;
+
+		bool bWasGetSuccessful, bWasSetSuccessful;
+
+		// Test that a value can be created, set, and get inside a simulation
+		int32 Int32Val = -1;
+		static const FName Integer32TestEntryName = TEXT("Integer32TestEntry");
+
+		// Test: premature read before the element has been created
+		bWasGetSuccessful = RBBBExternalAccessor.TryGet(Integer32TestEntryName, Int32Val);
+
+		CHECK_EQUALS(TEXT("Verify premature external retrieval of value, before creation, should fail"), bWasGetSuccessful, false);
+
+
+		RBBBSimAccessor.CreateEntry<int32>(Integer32TestEntryName, Integer32TestSettings);
+
+		{
+			RBBB->BeginSimulationFrame(CurrentSimTime);
+		
+
+			bWasSetSuccessful = RBBBSimAccessor.TrySet(Integer32TestEntryName, 32);
+			bWasGetSuccessful = RBBBSimAccessor.TryGet(Integer32TestEntryName, Int32Val);
+
+			CHECK_EQUALS(TEXT("Verify in-simulation setting of a value indicates success"), bWasSetSuccessful, true);
+			CHECK_EQUALS(TEXT("Verify in-simulation getting of a value indicates success"), bWasGetSuccessful, true);
+			CHECK_EQUALS(TEXT("Verify in-simulation getting of a value matches"), Int32Val, 32);
+
+			// Test that the value CANNOT yet be retrieved from an external accessor
+			bWasGetSuccessful = RBBBExternalAccessor.TryGet(Integer32TestEntryName, Int32Val);
+			CHECK_EQUALS(TEXT("Verify premature external retrieval of value, before end of first sim frame, should fail"), bWasGetSuccessful, false);
+
+			RBBB->EndSimulationFrame();
+		}
+
+		// Test that a value can now be retrieved externally
+		bWasGetSuccessful = RBBBExternalAccessor.TryGet(Integer32TestEntryName, Int32Val);
+
+		CHECK_EQUALS(TEXT("Verify external getting operation success"), bWasGetSuccessful, true);
+		CHECK_EQUALS(TEXT("Verify external getting value matches"), Int32Val, 32);
+
+		{
+			AdvanceTimeStepFrame(CurrentSimTime);
+			RBBB->BeginSimulationFrame(CurrentSimTime);
+
+			// Test that the value can be retrieved externally while a sim is ongoing, but that the prior value is returned
+			bWasSetSuccessful = RBBBSimAccessor.TrySet(Integer32TestEntryName, 64);
+
+			bWasGetSuccessful = RBBBExternalAccessor.TryGet(Integer32TestEntryName, Int32Val);
+
+			CHECK_EQUALS(TEXT("Verify external getting value remains the same from the prior frame until the sim frame ends"), Int32Val, 32);
+
+
+			RBBB->EndSimulationFrame();
+
+			bWasGetSuccessful = RBBBExternalAccessor.TryGet(Integer32TestEntryName, Int32Val);
+
+			CHECK_EQUALS(TEXT("Verify external getting value reads the most recent value after the sim frame ends"), Int32Val, 64);
+
+		}
+
+
+		// Test predictive use
+		AdvanceTimeStepFrame(CurrentSimTime);
+
+		FRollbackBlackboardSimWrapper RBBBPredictiveAccessor(*RBBB, /*bPredictiveMode*/ true);
+
+		const int32 NumFramesToPredict = 10;
+
+		FMoverTimeStep PredictionTime = CurrentSimTime;
+
+		int32 ExternalReadValue = -1;
+		int32 EarlyPredictReadValue = -2;
+		int32 LatePredictReadValue = -3;
+		int32 FinalPredictReadValue = -4;
+
+
+		const int32 LastNonPredictedValue = 222;
+		RBBBExternalAccessor.TrySet(Integer32TestEntryName, LastNonPredictedValue);
+
+		int32 LastWrittenPredictValue = LastNonPredictedValue;
+
+
+		for (int32 i=0; i < NumFramesToPredict; ++i)
+		{
+			AdvanceTimeStepFrame(PredictionTime);
+
+			bool bDidRead(false), bDidWrite(false);
+
+			// Check that EXTERNAL reader is never polluted by predicted writes
+			bDidRead = RBBBExternalAccessor.TryGet(Integer32TestEntryName, ExternalReadValue);
+			CHECK(bDidRead);
+			CHECK_EQUALS(TEXT("Verify external getting value reads the most recent true value (not predicted, early)"), ExternalReadValue, LastNonPredictedValue);
+
+			RBBB->BeginPredictionFrame(PredictionTime);
+			{
+				bDidRead = RBBBPredictiveAccessor.TryGet(Integer32TestEntryName, EarlyPredictReadValue);
+				CHECK(bDidRead);
+				CHECK_EQUALS(TEXT("Verify early predicted read the most recent true value (predicted, early-prediction)"), EarlyPredictReadValue, LastWrittenPredictValue);
+
+				// Advance predicted value and write it
+				LastWrittenPredictValue = LastWrittenPredictValue + (i*10);
+				bDidWrite = RBBBPredictiveAccessor.TrySet(Integer32TestEntryName, LastWrittenPredictValue);	// predictive SET
+				CHECK(bDidWrite);
+
+				// Test that the predictive read is the new value but the external accessor is unaffected
+				bDidRead = RBBBPredictiveAccessor.TryGet(Integer32TestEntryName, LatePredictReadValue);
+				CHECK(bDidRead);
+				CHECK_EQUALS(TEXT("Verify late predicted read gets the most recent predicted value (predicted, late-prediction)"), LatePredictReadValue, LastWrittenPredictValue);
+
+				bDidRead = RBBBExternalAccessor.TryGet(Integer32TestEntryName, ExternalReadValue);
+				CHECK(bDidRead);
+				CHECK_EQUALS(TEXT("Verify external getting value reads the most recent true value (not predicted, late-prediction)"), ExternalReadValue, LastNonPredictedValue);
+			}
+
+			bDidRead = RBBBPredictiveAccessor.TryGet(Integer32TestEntryName, FinalPredictReadValue);
+			CHECK(bDidRead);
+			CHECK_EQUALS(TEXT("Verify late predicted read gets the final predicted value (predicted, post-prediction)"), FinalPredictReadValue, LastWrittenPredictValue);
+
+			bDidRead = RBBBExternalAccessor.TryGet(Integer32TestEntryName, ExternalReadValue);
+			CHECK(bDidRead);
+			CHECK_EQUALS(TEXT("Verify external getting value reads the most recent true value (not predicted, post-predicted frame)"), ExternalReadValue, LastNonPredictedValue);
+
+		}
+
+		RBBB->EndPrediction();
+	}
+}
