[Backout] - CL50608996
Due to missing exports in shipping build
#rnx
Original CL Desc
-----------------------------------------------------------------
Mover: supporting prediction on the rollback blackboard, with dedicated slot in the circular buffer that avoids the need for copying and mem allocs. Also changes acccessor wrappers so there are external and simulation wrappers, with simulation wrapper hiding whether access is predicted or not.  No one except the owner + simulation driver should have direct access to the blackboard object.
All sim cases now have a Mover 'Simulation' object, which now owns the blackboard and controls access to it. This is a preliminary change that will be followed by more changes to how the Simulation object is used as the movement flow is driven.

MoverTests: adds first version of automated testing for the blackboard.

#Synced-CL 50587857

diff --git a/Engine/Plugins/Experimental/ChaosMover/Source/ChaosMover/Private/Backends/ChaosMoverBackend.cpp b/Engine/Plugins/Experimental/ChaosMover/Source/ChaosMover/Private/Backends/ChaosMoverBackend.cpp
--- a/Engine/Plugins/Experimental/ChaosMover/Source/ChaosMover/Private/Backends/ChaosMoverBackend.cpp
+++ b/Engine/Plugins/Experimental/ChaosMover/Source/ChaosMover/Private/Backends/ChaosMoverBackend.cpp
@@ -64,16 +64,15 @@ void UChaosMoverBackendComponent::InitializeComponent()
 
 	if (UWorld* World = GetWorld(); World && World->IsGameWorld())
 	{
-		UMoverComponent& MoverComp = GetMoverComponent();
-
 		Simulation = NewObject<UChaosMoverSimulation>(GetOwner(), SimulationClass, TEXT("ChaosMoverSimulation"), RF_Transient);
-		Simulation->SetOwner(MoverComp.GetOwner());
-		MoverComp.SetSimulation(*Simulation);
+		Simulation->SetOwner(GetMoverComponent().GetOwner());
 
-		NullMovementMode = NewObject<UNullMovementMode>(&MoverComp, TEXT("NullMovementMode"), RF_Transient);
+		NullMovementMode = NewObject<UNullMovementMode>(&GetMoverComponent(), TEXT("NullMovementMode"), RF_Transient);
 		ImmediateModeTransition = NewObject<UImmediateMovementModeTransition>(&GetMoverComponent(), TEXT("ImmediateModeTransition"), RF_Transient);
 		TransformOnInit = GetMoverComponent().GetUpdatedComponentTransform();
 
+		Simulation->SetRollbackBlackboard(GetMoverComponent().GetRollbackBlackboard_Internal());
+
 		// Create NetworkPhysicsComponent
 		if ((World->GetNetMode() != ENetMode::NM_Standalone) && Chaos::FPhysicsSolverBase::IsNetworkPhysicsPredictionEnabled() && bIsUsingAsyncPhysics)
 		{
diff --git a/Engine/Plugins/Experimental/ChaosMover/Source/ChaosMover/Private/ChaosMoverStateMachine.cpp b/Engine/Plugins/Experimental/ChaosMover/Source/ChaosMover/Private/ChaosMoverStateMachine.cpp
--- a/Engine/Plugins/Experimental/ChaosMover/Source/ChaosMover/Private/ChaosMoverStateMachine.cpp
+++ b/Engine/Plugins/Experimental/ChaosMover/Source/ChaosMover/Private/ChaosMoverStateMachine.cpp
@@ -9,7 +9,7 @@
 #include "ChaosMover/ChaosMoverSimulation.h"
 #include "MoveLibrary/MovementMixer.h"
 #include "MoveLibrary/MovementUtils.h"
-#include "MoveLibrary/RollbackBlackboardWrappers.h"
+#include "MoveLibrary/RollbackBlackboard.h"
 #include "MovementModeTransition.h"
 #include "MoverComponent.h"
 #include "MoverDeveloperSettings.h"
@@ -832,14 +832,16 @@ namespace UE::ChaosMover
 						}
 
 						// Track last mode change in the blackboard
-						FRollbackBlackboardSimWrapper RollbackBlackboard = SimStrongObjPtr->GetRollbackBlackboardSimWrapper();
-						FMovementModeChangeRecord ModeChangeRecord;
-						ModeChangeRecord.ModeName = CurrentModeName;
-						ModeChangeRecord.PrevModeName = PreviousModeName;
-						ModeChangeRecord.Frame = InternalSimTime.ServerFrame;
-						ModeChangeRecord.SimTimeMs = InternalSimTime.BaseSimTimeMs;
-
-						RollbackBlackboard.TrySet(CommonBlackboard::LastModeChangeRecord, ModeChangeRecord);
+						if (URollbackBlackboard_InternalWrapper* RollbackBlackboard = SimStrongObjPtr->GetRollbackBlackboard_Mutable())
+						{
+							FMovementModeChangeRecord ModeChangeRecord;
+							ModeChangeRecord.ModeName = CurrentModeName;
+							ModeChangeRecord.PrevModeName = PreviousModeName;
+							ModeChangeRecord.Frame = InternalSimTime.ServerFrame;
+							ModeChangeRecord.SimTimeMs = InternalSimTime.BaseSimTimeMs;
+
+							RollbackBlackboard->TrySet(CommonBlackboard::LastModeChangeRecord, ModeChangeRecord);
+						}
 
 						if (TStrongObjectPtr<UBaseMovementMode> CurrentMode = Modes[CurrentModeName].Pin())
 						{	
diff --git a/Engine/Plugins/Experimental/ChaosMover/Source/ChaosMover/Public/ChaosMover/ChaosMoverSimulation.h b/Engine/Plugins/Experimental/ChaosMover/Source/ChaosMover/Public/ChaosMover/ChaosMoverSimulation.h
--- a/Engine/Plugins/Experimental/ChaosMover/Source/ChaosMover/Public/ChaosMover/ChaosMoverSimulation.h
+++ b/Engine/Plugins/Experimental/ChaosMover/Source/ChaosMover/Public/ChaosMover/ChaosMoverSimulation.h
@@ -7,6 +7,7 @@
 #include "ChaosMover/PathedMovement/ChaosPathedMovementTypes.h"
 #include "ChaosMover/ChaosMoverStateMachine.h"
 #include "MoveLibrary/MovementMixer.h"
+#include "MoveLibrary/MoverBlackboard.h"
 #include "MoverTypes.h"
 #include "MoverSimulation.h"
 #include "PhysicsEngine/ConstraintInstance.h"
diff --git a/Engine/Plugins/Experimental/Mover/Source/Mover/Private/Backends/MoverNetworkPhysicsLiaisonBase.cpp b/Engine/Plugins/Experimental/Mover/Source/Mover/Private/Backends/MoverNetworkPhysicsLiaisonBase.cpp
--- a/Engine/Plugins/Experimental/Mover/Source/Mover/Private/Backends/MoverNetworkPhysicsLiaisonBase.cpp
+++ b/Engine/Plugins/Experimental/Mover/Source/Mover/Private/Backends/MoverNetworkPhysicsLiaisonBase.cpp
@@ -2,9 +2,8 @@
 
 #include "Backends/MoverNetworkPhysicsLiaisonBase.h"
 
-#include "MoveLibrary/RollbackBlackboardWrappers.h"
+#include "MoveLibrary/RollbackBlackboard.h"
 #include "MovementModeStateMachine.h"
-#include "MoverSimulation.h"
 #include "PBDRigidsSolver.h"
 #include "Chaos/PhysicsObjectInterface.h"
 #include "Chaos/PhysicsObjectInternalInterface.h"
@@ -233,8 +232,6 @@ UMoverNetworkPhysicsLiaisonComponentBase::UMoverNetworkPhysicsLiaisonComponentBa
 	{
 		SetIsReplicatedByDefault(true);
 	}
-
-	SimulationClass = UMoverSimulation::StaticClass();
 }
 
 //////////////////////////////////////////////////////////////////////////
@@ -296,12 +293,12 @@ void UMoverNetworkPhysicsLiaisonComponentBase::HandleComponentPhysicsStateChange
 	{
 		if (UMoverComponent& MoverComp = GetMoverComponent(); MoverComp.ModeFSM)
 		{
-			RollbackBlackboard->BeginSimulationFrame(GetCurrentAsyncMoverTimeStep_External());
+			MoverComp.RollbackBlackboard_InternalWrapper->BeginSimulationFrame(GetCurrentAsyncMoverTimeStep_External());
 
 			MoverComp.ModeFSM->SetModeImmediately(MoverComp.StartingMovementMode);
 			MoverComp.SimBlackboard->InvalidateAll();
 
-			RollbackBlackboard->EndSimulationFrame();
+			MoverComp.RollbackBlackboard_InternalWrapper->EndSimulationFrame();
 		}
 
 		InitializeSimOutputData();
@@ -335,10 +332,6 @@ void UMoverNetworkPhysicsLiaisonComponentBase::InitializeComponent()
 		UMoverComponent& MoverComp = GetMoverComponent();
 	}
 
-	Simulation = NewObject<UMoverSimulation>(GetOwner(), SimulationClass, TEXT("MoverSimulation"), RF_Transient);
-	GetMoverComponent().SetSimulation(*Simulation);
-	RollbackBlackboard = Simulation->RollbackBlackboard;
-
 	if (Chaos::FPhysicsSolverBase::IsNetworkPhysicsPredictionEnabled() && bUsingAsyncPhysics)
 	{
 		NetworkPhysicsComponent = NewObject<UNetworkPhysicsComponent>(GetOwner(), TEXT("PhysMover_NetworkPhysicsComponent"));
@@ -842,7 +835,7 @@ void UMoverNetworkPhysicsLiaisonComponentBase::PerformPreSimulate_Internal(const
 	FMoverTickEndData TickEndData;
 	TickEndData.SyncState = TickStartData.SyncState;
 
-	RollbackBlackboard->BeginSimulationFrame(TimeStep);
+	MoverComp.RollbackBlackboard_InternalWrapper->BeginSimulationFrame(TimeStep);
 
 
 	MoverComp.ModeFSM->OnSimulationTick(MoverComp.UpdatedComponent, MoverComp.UpdatedCompAsPrimitive, MoverComp.GetSimBlackboard_Mutable(), TickStartData, TimeStep, TickEndData);
@@ -869,8 +862,8 @@ void UMoverNetworkPhysicsLiaisonComponentBase::PerformPreSimulate_Internal(const
 
 		//@todo DanH: Does the base need to concern itself with setting W on the particle?
 	}
-	
-	RollbackBlackboard->EndSimulationFrame();
+
+	MoverComp.RollbackBlackboard_InternalWrapper->EndSimulationFrame();
 
 }
 
diff --git a/Engine/Plugins/Experimental/Mover/Source/Mover/Private/Backends/MoverNetworkPredictionLiaison.cpp b/Engine/Plugins/Experimental/Mover/Source/Mover/Private/Backends/MoverNetworkPredictionLiaison.cpp
--- a/Engine/Plugins/Experimental/Mover/Source/Mover/Private/Backends/MoverNetworkPredictionLiaison.cpp
+++ b/Engine/Plugins/Experimental/Mover/Source/Mover/Private/Backends/MoverNetworkPredictionLiaison.cpp
@@ -5,7 +5,6 @@
 #include "NetworkPredictionProxyInit.h"
 #include "NetworkPredictionProxyWrite.h"
 #include "GameFramework/Actor.h"
-#include "MoverSimulation.h"
 #include "MoverComponent.h"
 
 #if WITH_EDITOR
@@ -47,8 +46,6 @@ UMoverNetworkPredictionLiaisonComponent::UMoverNetworkPredictionLiaisonComponent
 	bWantsInitializeComponent = true;
 	bAutoActivate = true;
 	SetIsReplicatedByDefault(true);
-
-	SimulationClass = UMoverSimulation::StaticClass();
 }
 
 void UMoverNetworkPredictionLiaisonComponent::ProduceInput(const int32 DeltaTimeMS, FMoverInputCmdContext* Cmd)
@@ -286,15 +283,6 @@ void UMoverNetworkPredictionLiaisonComponent::BeginPlay()
 void UMoverNetworkPredictionLiaisonComponent::InitializeComponent()
 {
 	Super::InitializeComponent();
-
-	Simulation = NewObject<UMoverSimulation>(GetOwner(), SimulationClass, TEXT("MoverSimulation"), RF_Transient);
-	
-	MoverComp = GetOwner()->FindComponentByClass<UMoverComponent>();
-
-	if (MoverComp)
-	{
-		MoverComp->SetSimulation(*Simulation);
-	}
 }
 
 
@@ -302,8 +290,6 @@ void UMoverNetworkPredictionLiaisonComponent::UninitializeComponent()
 {
 	NetworkPredictionProxy.EndPlay();
 
-	Simulation = nullptr;
-
 	Super::UninitializeComponent();
 }
 
diff --git a/Engine/Plugins/Experimental/Mover/Source/Mover/Private/Backends/MoverStandaloneLiaison.cpp b/Engine/Plugins/Experimental/Mover/Source/Mover/Private/Backends/MoverStandaloneLiaison.cpp
--- a/Engine/Plugins/Experimental/Mover/Source/Mover/Private/Backends/MoverStandaloneLiaison.cpp
+++ b/Engine/Plugins/Experimental/Mover/Source/Mover/Private/Backends/MoverStandaloneLiaison.cpp
@@ -5,7 +5,6 @@
 #include "GameFramework/Pawn.h"
 #include "GameFramework/Controller.h"
 #include "MoverComponent.h"
-#include "MoverSimulation.h"
 
 #include UE_INLINE_GENERATED_CPP_BY_NAME(MoverStandaloneLiaison)
 
@@ -85,8 +84,6 @@ UMoverStandaloneLiaisonComponent::UMoverStandaloneLiaisonComponent(const FObject
 	bAutoActivate = true;
 	SetIsReplicatedByDefault(false);
 
-	SimulationClass = UMoverSimulation::StaticClass();
-
 	CurrentSimTimeMs = 0.0;
 	CurrentSimFrame = 0;
 
@@ -130,27 +127,6 @@ bool UMoverStandaloneLiaisonComponent::WritePendingSyncState(const FMoverSyncSta
 	return true;
 }
 
-void UMoverStandaloneLiaisonComponent::InitializeComponent()
-{
-	Super::InitializeComponent();
-
-	Simulation = NewObject<UMoverSimulation>(GetOwner(), SimulationClass, TEXT("MoverSimulation"), RF_Transient);
-
-	MoverComp = GetOwner()->FindComponentByClass<UMoverComponent>();
-
-	if (MoverComp)
-	{
-		MoverComp->SetSimulation(*Simulation);
-	}
-}
-
-void UMoverStandaloneLiaisonComponent::UninitializeComponent()
-{
-	Simulation = nullptr;
-	Super::UninitializeComponent();
-}
-
-
 void UMoverStandaloneLiaisonComponent::BeginPlay()
 {
 	Super::BeginPlay();
diff --git a/Engine/Plugins/Experimental/Mover/Source/Mover/Private/MoveLibrary/RollbackBlackboard.cpp b/Engine/Plugins/Experimental/Mover/Source/Mover/Private/MoveLibrary/RollbackBlackboard.cpp
--- a/Engine/Plugins/Experimental/Mover/Source/Mover/Private/MoveLibrary/RollbackBlackboard.cpp
+++ b/Engine/Plugins/Experimental/Mover/Source/Mover/Private/MoveLibrary/RollbackBlackboard.cpp
@@ -35,9 +35,9 @@ uint32 URollbackBlackboard::BlackboardEntryBase::ComputeBufferSize(EBlackboardSi
 	return BufferSize;
 }
 
-URollbackBlackboard::BlackboardEntryBase::BlackboardEntryBase(const EntrySettings& InSettings, uint32 BufferSize, ERollbackBufferFlags BufferFlags)
+URollbackBlackboard::BlackboardEntryBase::BlackboardEntryBase(const EntrySettings& InSettings, uint32 BufferSize)
 	: Settings(InSettings)
-	, Timestamps(BufferSize, BufferFlags, EntryTimeStamp())
+	, Timestamps(BufferSize, EntryTimeStamp())
 	, ExternalIdx(0)
 	, InternalIdx(0)
 {
@@ -71,19 +71,13 @@ void URollbackBlackboard::BlackboardEntryBase::RollBack(uint32 NewPendingFrame)
 
 }
 
-void URollbackBlackboard::BlackboardEntryBase::InvalidatePredictiveState()
-{
-	Timestamps.GetPredictiveElement().Invalidate();
-
-	ResetPredictiveElement();
-}
 
 
-
-bool URollbackBlackboard::BlackboardEntryBase::CanReadEntryAtTime(const EntryTimeStamp& ReaderTimeStamp, const EntryTimeStamp& EntryTimeStamp) const
+bool URollbackBlackboard::BlackboardEntryBase::CanReadEntryAt(const EntryTimeStamp& ReaderTimeStamp, EEntryIndexType IndexType) const
 {
+	const uint32 TimestampIdx = IndexType == EEntryIndexType::External ? ExternalIdx : InternalIdx;
 
-	if (!EntryTimeStamp.IsValid() || !ReaderTimeStamp.IsValid())
+	if (!Timestamps[TimestampIdx].IsValid())
 	{
 		return false;	// entry isn't initialized yet or never set
 	}
@@ -93,20 +87,15 @@ bool URollbackBlackboard::BlackboardEntryBase::CanReadEntryAtTime(const EntryTim
 		case EBlackboardPersistencePolicy::NextFrameOnly:
 		{
 			// if value wasn't last set during the current or immediate prior sim frame, then it can't be read
-			if (!(ReaderTimeStamp.Frame == EntryTimeStamp.Frame ||
-				  ReaderTimeStamp.Frame == EntryTimeStamp.Frame+1))
+			if (!(ReaderTimeStamp.Frame == Timestamps[TimestampIdx].Frame ||
+				  ReaderTimeStamp.Frame == Timestamps[TimestampIdx].Frame+1))
 			{
 				return false;
 			}
 		}
 		break;
 
-		case EBlackboardPersistencePolicy::Forever:
-			// If the reader is in the past compared to the time of the entry, then it can't be read
-			if (ReaderTimeStamp.Frame < EntryTimeStamp.Frame)
-			{
-				return false;
-			}
+		case EBlackboardPersistencePolicy::Forever:		// Always allow reading
 			break;
 
 		default:		// Allow reading of any not-yet-implemented policies
@@ -116,22 +105,6 @@ bool URollbackBlackboard::BlackboardEntryBase::CanReadEntryAtTime(const EntryTim
 	return true;
 }
 
-bool URollbackBlackboard::BlackboardEntryBase::CanReadEntryAtTime(const EntryTimeStamp& ReaderTimeStamp, EEntryIndexType IndexType) const
-{
-	const EntryTimeStamp* TargetEntryTimeStamp = nullptr;
-
-	if (IndexType == EEntryIndexType::Predictive)
-	{
-		TargetEntryTimeStamp = &Timestamps.GetPredictiveElement();
-	}
-	else
-	{
-		const uint32 TimestampIdx = ((IndexType == EEntryIndexType::External) ? ExternalIdx : InternalIdx);
-		TargetEntryTimeStamp = &Timestamps[TimestampIdx];
-	}
-	
-	return CanReadEntryAtTime(ReaderTimeStamp, *TargetEntryTimeStamp);
-}
 
 
 void URollbackBlackboard::BeginSimulationFrame(const FMoverTimeStep& PendingTimeStep)
@@ -149,7 +122,6 @@ void URollbackBlackboard::EndSimulationFrame()
 {
 	check(bIsSimulationInProgress && (InProgressSimFrameThreadId == FPlatformTLS::GetCurrentThreadId()));
 	bIsSimulationInProgress = bIsResimulating = false;
-	InProgressSimFrameThreadId = 0;
 
 	// TODO: need some kind of lock so we wait on any in-progress operations before advancing the CurrentSimTimeStamp
 	CurrentSimTimeStamp = InProgressSimTimeStamp;
@@ -190,41 +162,25 @@ void URollbackBlackboard::EndRollback()
 {
 	check (bIsRollbackInProgress && (InRollbackThreadId == FPlatformTLS::GetCurrentThreadId()));
 	bIsRollbackInProgress = false;
-	InRollbackThreadId = 0;
 }
 
-void URollbackBlackboard::BeginPredictionFrame(const FMoverTimeStep& PendingTimeStep)
+
+void URollbackBlackboard_InternalWrapper::BeginSimulationFrame(const FMoverTimeStep& PendingTimeStep)
 {
-	check(BufferFlags & ERollbackBufferFlags::Prediction);
-	check(!bIsPredictionInProgress || InPredictionThreadId == FPlatformTLS::GetCurrentThreadId());
-	
-	if (!bIsPredictionInProgress)
-	{
-		bIsPredictionInProgress = true;
-		InPredictionThreadId = FPlatformTLS::GetCurrentThreadId();
-	}
+	Blackboard->BeginSimulationFrame(PendingTimeStep);
 }
 
-void URollbackBlackboard::EndPrediction()
+void URollbackBlackboard_InternalWrapper::EndSimulationFrame()
 {
-	check(((BufferFlags & ERollbackBufferFlags::Prediction) == ERollbackBufferFlags::Prediction) && bIsPredictionInProgress && IsInPredictionThread());
-
-	bIsPredictionInProgress = false;
-	InPredictionThreadId = 0;
-
-	for (const TPair<FName, TUniquePtr<BlackboardEntryBase>>& KVP : EntryMap)
-	{
-		KVP.Value->InvalidatePredictiveState();
-	}
+	Blackboard->EndSimulationFrame();
 }
 
-bool URollbackBlackboard::IsInSimulationThread() const
+void URollbackBlackboard_InternalWrapper::BeginRollback(const FMoverTimeStep& NewBaseTimeStep)
 {
-	return (InProgressSimFrameThreadId == FPlatformTLS::GetCurrentThreadId());
+	Blackboard->BeginRollback(NewBaseTimeStep);
 }
 
-bool URollbackBlackboard::IsInPredictionThread() const
+void URollbackBlackboard_InternalWrapper::EndRollback()
 {
-	return (InPredictionThreadId == FPlatformTLS::GetCurrentThreadId());
+	Blackboard->EndRollback();
 }
-
diff --git a/Engine/Plugins/Experimental/Mover/Source/Mover/Private/MovementModeStateMachine.cpp b/Engine/Plugins/Experimental/Mover/Source/Mover/Private/MovementModeStateMachine.cpp
--- a/Engine/Plugins/Experimental/Mover/Source/Mover/Private/MovementModeStateMachine.cpp
+++ b/Engine/Plugins/Experimental/Mover/Source/Mover/Private/MovementModeStateMachine.cpp
@@ -7,8 +7,8 @@
 #include "MoverLog.h"
 #include "MoveLibrary/MovementUtils.h"
 #include "MoverComponent.h"
-#include "MoverSimulation.h"
 #include "MoveLibrary/MovementMixer.h"
+#include "MoveLibrary/RollbackBlackboard.h"
 #include "Templates/SubclassOf.h"
 #include "Engine/World.h"
 #include "GameFramework/GameStateBase.h"
@@ -123,11 +123,6 @@ void UMovementModeStateMachine::ClearAllGlobalTransitions()
 	GlobalTransitions.Empty();
 }
 
-void UMovementModeStateMachine::SetSimulation(UMoverSimulation& InSimulation)
-{
-	Simulation = &InSimulation;
-}
-
 void UMovementModeStateMachine::QueueNextMode(FName DesiredNextModeName, bool bShouldReenter)
 {
 	QueueNextModeInternal(DesiredNextModeName, bShouldReenter, /*bIsFromRollback*/ false);
@@ -767,15 +762,16 @@ void UMovementModeStateMachine::AdvanceToNextMode()
 		// signal movement mode change event
 		const UMoverComponent* MoverComp = CastChecked<UMoverComponent>(GetOuter());
 
-		FRollbackBlackboardSimWrapper RollbackBlackboard = Simulation->GetRollbackBlackboardSimWrapper();
-
-		FMovementModeChangeRecord ModeChangeRecord;
-		ModeChangeRecord.ModeName = CurrentModeName;
-		ModeChangeRecord.PrevModeName = PreviousModeName;
-		ModeChangeRecord.Frame = CurrentBaseTimeStep.ServerFrame;
-		ModeChangeRecord.SimTimeMs = CurrentBaseTimeStep.BaseSimTimeMs;
+		if (URollbackBlackboard_InternalWrapper* RollbackBlackboard = MoverComp->GetRollbackBlackboard_Internal())
+		{
+			FMovementModeChangeRecord ModeChangeRecord;
+			ModeChangeRecord.ModeName = CurrentModeName;
+			ModeChangeRecord.PrevModeName = PreviousModeName;
+			ModeChangeRecord.Frame = CurrentBaseTimeStep.ServerFrame;
+			ModeChangeRecord.SimTimeMs = CurrentBaseTimeStep.BaseSimTimeMs;
 
-		RollbackBlackboard.TrySet(CommonBlackboard::LastModeChangeRecord, ModeChangeRecord);
+			RollbackBlackboard->TrySet(CommonBlackboard::LastModeChangeRecord, ModeChangeRecord);
+		}
 
 		Modes[CurrentModeName]->Activate(MoverEventContext);
 
diff --git a/Engine/Plugins/Experimental/Mover/Source/Mover/Private/MoverComponent.cpp b/Engine/Plugins/Experimental/Mover/Source/Mover/Private/MoverComponent.cpp
--- a/Engine/Plugins/Experimental/Mover/Source/Mover/Private/MoverComponent.cpp
+++ b/Engine/Plugins/Experimental/Mover/Source/Mover/Private/MoverComponent.cpp
@@ -4,7 +4,6 @@
 #include "MoverComponent.h"
 #include "MoverSimulationTypes.h"
 #include "MovementModeStateMachine.h"
-#include "MoverSimulation.h"
 #include "MotionWarpingMoverAdapter.h"
 #include "DefaultMovementSet/Modes/WalkingMode.h"
 #include "DefaultMovementSet/Modes/FallingMode.h"
@@ -12,7 +11,7 @@
 #include "MoveLibrary/MovementMixer.h"
 #include "MoveLibrary/MovementUtils.h"
 #include "MoveLibrary/FloorQueryUtils.h"
-#include "MoveLibrary/RollbackBlackboardWrappers.h"
+#include "MoveLibrary/RollbackBlackboard.h"
 #include "InputContainerStruct.h"
 #include "MoverLog.h"
 #include "InstantMovementEffect.h"
@@ -96,29 +95,21 @@ void UMoverComponent::InitializeComponent()
 
 		SimBlackboard = NewObject<UMoverBlackboard>(this, TEXT("MoverBlackboard"), RF_Transient);
 
+		RollbackBlackboard = NewObject<URollbackBlackboard>(this, TEXT("RollbackBlackboard"), RF_Transient);
+		RollbackBlackboard_InternalWrapper = NewObject<URollbackBlackboard_InternalWrapper>(this, TEXT("RollbackBlackboard_Internal"), RF_Transient);
+		RollbackBlackboard_InternalWrapper->Init(*RollbackBlackboard);
 
-		FindDefaultUpdatedComponent();
-
+		// create any internal entries
+		static URollbackBlackboard::EntrySettings ModeChangeRecordSettings;
+		ModeChangeRecordSettings.SizingPolicy = EBlackboardSizingPolicy::FixedDeclaredSize;
+		ModeChangeRecordSettings.FixedSize = 4;
+		ModeChangeRecordSettings.PersistencePolicy = EBlackboardPersistencePolicy::Forever;
+		ModeChangeRecordSettings.RollbackPolicy = EBlackboardRollbackPolicy::InvalidatedOnRollback;
 
-		// Instantiate our sister backend component that will actually talk to the system driving the simulation
-		if (BackendClass)
-		{
-			UActorComponent* NewLiaisonComp = NewObject<UActorComponent>(this, BackendClass, TEXT("BackendLiaisonComponent"));
-			BackendLiaisonComp.SetObject(NewLiaisonComp);
-			BackendLiaisonComp.SetInterface(CastChecked<IMoverBackendLiaisonInterface>(NewLiaisonComp));
-			if (BackendLiaisonComp)
-			{
-				NewLiaisonComp->RegisterComponent();
-				NewLiaisonComp->InitializeComponent();
-				NewLiaisonComp->SetNetAddressable();
-			}
-		}
-		else
-		{
-			UE_LOG(LogMover, Error, TEXT("No backend class set on %s. Mover actor will not function."), *GetNameSafe(GetOwner()));
-		}
+		RollbackBlackboard->CreateEntry<FMovementModeChangeRecord>(CommonBlackboard::LastModeChangeRecord, ModeChangeRecordSettings);
 
 
+		FindDefaultUpdatedComponent();
 
 		// Set up FSM and initial movement states
 		ModeFSM = NewObject<UMovementModeStateMachine>(this, TEXT("MoverStateMachine"), RF_Transient);
@@ -154,6 +145,24 @@ void UMoverComponent::InitializeComponent()
 			ModeFSM->SetDefaultMode(StartingMovementMode);
 			ModeFSM->QueueNextMode(StartingMovementMode);
 		}
+
+		// Instantiate our sister backend component that will actually talk to the system driving the simulation
+		if (BackendClass)
+		{
+			UActorComponent* NewLiaisonComp = NewObject<UActorComponent>(this, BackendClass, TEXT("BackendLiaisonComponent"));
+			BackendLiaisonComp.SetObject(NewLiaisonComp);
+			BackendLiaisonComp.SetInterface(CastChecked<IMoverBackendLiaisonInterface>(NewLiaisonComp));
+			if (BackendLiaisonComp)
+			{
+				NewLiaisonComp->RegisterComponent();
+				NewLiaisonComp->InitializeComponent();
+				NewLiaisonComp->SetNetAddressable();
+			}
+		}
+		else
+		{
+			UE_LOG(LogMover, Error, TEXT("No backend class set on %s. Mover actor will not function."), *GetNameSafe(GetOwner()));
+		}
 	}
 
 	// Gather initial state to fulfill queries
@@ -243,19 +252,7 @@ void UMoverComponent::PostLoad()
 void UMoverComponent::BeginPlay()
 {
 	Super::BeginPlay();
-
-	if (ensureMsgf(Simulation != nullptr, TEXT("MoverComponent requires a Simulation to be set during initialization! The backend should have done this already.")))
-	{
-		ModeFSM->SetSimulation(*Simulation);
-
-		static URollbackBlackboard::EntrySettings ModeChangeRecordSettings;
-		ModeChangeRecordSettings.SizingPolicy = EBlackboardSizingPolicy::FixedDeclaredSize;
-		ModeChangeRecordSettings.FixedSize = 4;
-		ModeChangeRecordSettings.PersistencePolicy = EBlackboardPersistencePolicy::Forever;
-		ModeChangeRecordSettings.RollbackPolicy = EBlackboardRollbackPolicy::InvalidatedOnRollback;
-		RollbackBlackboardExternalWrapper.CreateEntry<FMovementModeChangeRecord>(CommonBlackboard::LastModeChangeRecord, ModeChangeRecordSettings);
-	}
-
+	
 	FindDefaultUpdatedComponent();
 	ensureMsgf(UpdatedComponent != nullptr, TEXT("No root component found on %s. Simulation initialization will most likely fail."), *GetPathNameSafe(GetOwner()));
 
@@ -614,8 +611,8 @@ void UMoverComponent::SimulationTick(const FMoverTimeStep& InTimeStep, const FMo
 		OnPreMovement.Broadcast(MoverTimeStep, SimInput.InputCmd, SimInput.SyncState, SimInput.AuxState);
 	}
 
-	RollbackBlackboard->BeginSimulationFrame(MoverTimeStep);
-	
+	RollbackBlackboard_InternalWrapper->BeginSimulationFrame(MoverTimeStep);
+
 	// Tick the actual simulation. This is where the proposed moves are queried and executed, affecting change to the moving actor's gameplay state and captured in the output sim state
 	if (IsInGameThread())
 	{
@@ -643,7 +640,7 @@ void UMoverComponent::SimulationTick(const FMoverTimeStep& InTimeStep, const FMo
 		}
 	}
 
-	RollbackBlackboard->EndSimulationFrame();
+	RollbackBlackboard_InternalWrapper->EndSimulationFrame();
 
 	if (!SimOutput.MoveRecord.GetTotalMoveDelta().IsZero())
 	{
@@ -1029,15 +1026,6 @@ void UMoverComponent::CreateDefaultInputAndState(FMoverInputCmdContext& OutInput
 
 }
 
-void UMoverComponent::SetSimulation(UMoverSimulation& InSimulation)
-{
-	Simulation = &InSimulation;
-	check(Simulation->RollbackBlackboard != nullptr);
-	
-	RollbackBlackboard = Simulation->RollbackBlackboard;	// TODO: fetching the true blackboard is temporary: this will be removed once simulations are used consistently
-	RollbackBlackboardExternalWrapper = FRollbackBlackboardExternalWrapper(*Simulation->RollbackBlackboard);
-}
-
 void UMoverComponent::HandleImpact(FMoverOnImpactParams& ImpactParams)
 {
 	if (ImpactParams.MovementModeName.IsNone())
@@ -1138,11 +1126,11 @@ void UMoverComponent::OnSimulationRollback(const FMoverSyncState* SyncState, con
 {
 	SimBlackboard->Invalidate(EInvalidationReason::Rollback);
 
-	RollbackBlackboard->BeginRollback(NewBaseTimeStep);
+	RollbackBlackboard_InternalWrapper->BeginRollback(NewBaseTimeStep);
 
 	ModeFSM->OnSimulationRollback(SyncState, AuxState, NewBaseTimeStep);
 
-	RollbackBlackboard->EndRollback();
+	RollbackBlackboard_InternalWrapper->EndRollback();
 	bHasJustRolledBack = true;
 }
 
diff --git a/Engine/Plugins/Experimental/Mover/Source/Mover/Private/MoverSimulation.cpp b/Engine/Plugins/Experimental/Mover/Source/Mover/Private/MoverSimulation.cpp
--- a/Engine/Plugins/Experimental/Mover/Source/Mover/Private/MoverSimulation.cpp
+++ b/Engine/Plugins/Experimental/Mover/Source/Mover/Private/MoverSimulation.cpp
@@ -3,17 +3,13 @@
 #include "MoverSimulation.h"
 
 #include "MoveLibrary/MoverBlackboard.h"
-#include "MoveLibrary/RollbackBlackboardWrappers.h"
+#include "MoveLibrary/RollbackBlackboard.h"
 
 #include UE_INLINE_GENERATED_CPP_BY_NAME(MoverSimulation)
 
 UMoverSimulation::UMoverSimulation()
 {
 	Blackboard = CreateDefaultSubobject<UMoverBlackboard>(TEXT("MoverSimulationBlackboard"));
-	RollbackBlackboard = CreateDefaultSubobject<URollbackBlackboard>(TEXT("MoverSimulationRollbackBlackboard"));
-
-	// TODO: this should be set up depending on whether the simulation is real or purely for prediction 
-	RollbackBlackboardSimWrapper = FRollbackBlackboardSimWrapper(*RollbackBlackboard, /*bIsPredictive*/ false);
 }
 
 const UMoverBlackboard* UMoverSimulation::GetBlackboard() const
@@ -26,8 +22,17 @@ UMoverBlackboard* UMoverSimulation::GetBlackboard_Mutable()
 	return Blackboard;
 }
 
-FRollbackBlackboardSimWrapper UMoverSimulation::GetRollbackBlackboardSimWrapper() const
+const URollbackBlackboard_InternalWrapper* UMoverSimulation::GetRollbackBlackboard() const
 {
-	return RollbackBlackboardSimWrapper;
+	return RollbackBlackboard;
 }
 
+URollbackBlackboard_InternalWrapper* UMoverSimulation::GetRollbackBlackboard_Mutable()
+{
+	return RollbackBlackboard;
+}
+
+void UMoverSimulation::SetRollbackBlackboard(URollbackBlackboard_InternalWrapper* RollbackSimBlackboard)
+{
+	RollbackBlackboard = RollbackSimBlackboard;
+}
\ No newline at end of file
diff --git a/Engine/Plugins/Experimental/Mover/Source/Mover/Public/Backends/MoverNetworkPhysicsLiaisonBase.h b/Engine/Plugins/Experimental/Mover/Source/Mover/Public/Backends/MoverNetworkPhysicsLiaisonBase.h
--- a/Engine/Plugins/Experimental/Mover/Source/Mover/Public/Backends/MoverNetworkPhysicsLiaisonBase.h
+++ b/Engine/Plugins/Experimental/Mover/Source/Mover/Public/Backends/MoverNetworkPhysicsLiaisonBase.h
@@ -9,8 +9,6 @@
 #include "PhysicsMover/PhysicsMoverSimulationTypes.h"
 #include "MoverNetworkPhysicsLiaisonBase.generated.h"
 
-class UMoverSimulation;
-
 #define UE_API MOVER_API
 
 //////////////////////////////////////////////////////////////////////////
@@ -316,16 +314,6 @@ protected:
 	
 	bool bCachedInputIsValid = false;
 	bool bUsingAsyncPhysics = false;
-
-	// Type of simulation to instantiate as the system spins up
-	TSubclassOf<UMoverSimulation> SimulationClass;
-
-	UPROPERTY(Transient)
-	TObjectPtr<UMoverSimulation> Simulation;
-
-	UPROPERTY(Transient)
-	TObjectPtr<URollbackBlackboard> RollbackBlackboard;
-
 };
 
 #undef UE_API
diff --git a/Engine/Plugins/Experimental/Mover/Source/Mover/Public/Backends/MoverNetworkPredictionLiaison.h b/Engine/Plugins/Experimental/Mover/Source/Mover/Public/Backends/MoverNetworkPredictionLiaison.h
--- a/Engine/Plugins/Experimental/Mover/Source/Mover/Public/Backends/MoverNetworkPredictionLiaison.h
+++ b/Engine/Plugins/Experimental/Mover/Source/Mover/Public/Backends/MoverNetworkPredictionLiaison.h
@@ -16,7 +16,6 @@
 #define UE_API MOVER_API
 
 class UMoverComponent;
-class UMoverSimulation;
 
 
 using KinematicMoverStateTypes = TNetworkPredictionStateTypes<FMoverInputCmdContext, FMoverSyncState, FMoverAuxStateContext>;
@@ -92,12 +91,6 @@ protected:
 	TObjectPtr<UMoverComponent> MoverComp;	// the component that we're in charge of driving
 	FMoverSyncState* StartingOutSync;
 	FMoverAuxStateContext* StartingOutAux;
-
-	// Type of simulation to instantiate as the system spins up
-	TSubclassOf<UMoverSimulation> SimulationClass;
-
-	UPROPERTY(Transient)
-	TObjectPtr<UMoverSimulation> Simulation;
 };
 
 #undef UE_API
diff --git a/Engine/Plugins/Experimental/Mover/Source/Mover/Public/Backends/MoverStandaloneLiaison.h b/Engine/Plugins/Experimental/Mover/Source/Mover/Public/Backends/MoverStandaloneLiaison.h
--- a/Engine/Plugins/Experimental/Mover/Source/Mover/Public/Backends/MoverStandaloneLiaison.h
+++ b/Engine/Plugins/Experimental/Mover/Source/Mover/Public/Backends/MoverStandaloneLiaison.h
@@ -8,7 +8,6 @@
 
 #define UE_API MOVER_API
 
-class UMoverSimulation;
 class UMoverComponent;
 class UMoverStandaloneLiaisonComponent;
 class AController;
@@ -112,8 +111,6 @@ public:
 
 	// Begin UActorComponent interface
 	UE_API virtual void RegisterComponentTickFunctions(bool bRegister) override;
-	UE_API virtual void InitializeComponent() override;
-	UE_API virtual void UninitializeComponent() override;
 	UE_API virtual void BeginPlay() override;
 	// End UActorComponent interface
 
@@ -176,12 +173,6 @@ protected:
 	UPROPERTY(EditDefaultsOnly, Category = "Async Movement")
 	bool bUseAsyncMovementSimulationTick = false;
 
-	// Type of simulation to instantiate as the system spins up
-	TSubclassOf<UMoverSimulation> SimulationClass;
-
-	UPROPERTY(Transient)
-	TObjectPtr<UMoverSimulation> Simulation;
-
 	TObjectPtr<UMoverComponent> MoverComp;	// the component that we're in charge of driving
 
 	double CurrentSimTimeMs;
diff --git a/Engine/Plugins/Experimental/Mover/Source/Mover/Public/MoveLibrary/RollbackBlackboard.h b/Engine/Plugins/Experimental/Mover/Source/Mover/Public/MoveLibrary/RollbackBlackboard.h
--- a/Engine/Plugins/Experimental/Mover/Source/Mover/Public/MoveLibrary/RollbackBlackboard.h
+++ b/Engine/Plugins/Experimental/Mover/Source/Mover/Public/MoveLibrary/RollbackBlackboard.h
@@ -2,7 +2,7 @@
 
 #pragma once
 
-#include "RollbackCircularBuffer.h"
+#include "Containers/CircularBuffer.h"
 #include "MoverLog.h"
 #include "UObject/Interface.h"
 #include "UObject/Object.h"
@@ -100,8 +100,8 @@ enum class EBlackboardRollbackPolicy : uint8
 
 
 /**
- * This is the core rollback blackboard, with API for external access and admin use. 
- * See @FRollbackBlackboardSimWrapper for in-simulation use.
+ * This is the core rollback blackboard, with API for external access.
+ * See @URollbackBlackboard_InternalWrapper for in-simulation use.
  */ 
 UCLASS(MinimalAPI)
 class URollbackBlackboard : public UObject
@@ -143,9 +143,8 @@ private:
 
 	enum EEntryIndexType : uint8
 	{
-		External,   // entry index for outside systems to read
-		Internal,   // entry index for in-simulation system to read/write
-		Predictive  // entry index for a predicting system to read/write
+		External,	// entry index for outside systems to read
+		Internal	// entry index for in-simulation system to read/write
 	};
 
 	// untyped base, so we can use entries of multiple types in a single-typed container
@@ -153,7 +152,7 @@ private:
 	{
 	public:
 		BlackboardEntryBase() = delete;
-		UE_API BlackboardEntryBase(const EntrySettings& InSettings, uint32 BufferSize, ERollbackBufferFlags BufferFlags);
+		UE_API BlackboardEntryBase(const EntrySettings& InSettings, uint32 BufferSize);
 
 		virtual ~BlackboardEntryBase()
 		{
@@ -168,20 +167,14 @@ private:
 		// Note that NewPendingFrame means the frame that will now be re-simulated. So any existing entries with a Frame >= NewPendingFrame can be invalidated.
 		UE_API void RollBack(uint32 NewPendingFrame);
 
-		UE_API void InvalidatePredictiveState();
-		
-
 		UE_API static uint32 ComputeBufferSize(EBlackboardSizingPolicy InSizingPolicy, uint32 FixedBufferSize = -1);
 
 	protected:
-		UE_API bool CanReadEntryAtTime(const EntryTimeStamp& ReaderTimeStamp, EEntryIndexType IndexType) const;
-		UE_API bool CanReadEntryAtTime(const EntryTimeStamp& ReaderTimeStamp, const EntryTimeStamp& EntryTimeStamp) const;
-
-		UE_API virtual void ResetPredictiveElement() {}
+		UE_API bool CanReadEntryAt(const EntryTimeStamp& ReaderTimeStamp, EEntryIndexType IndexType) const;
 
 		EntrySettings Settings;
-		
-		TRollbackCircularBuffer<EntryTimeStamp> Timestamps;
+
+		TCircularBuffer<EntryTimeStamp> Timestamps;
 
 		uint32 ExternalIdx;	// Indexes to the last entry on a committed frame
 		uint32 InternalIdx;	// Indexes to where in-progress simulations should read/write (matches ExternalIdx when not mid-simulation, and advances if written during simulation)
@@ -197,9 +190,9 @@ private:
 	public:
 
 		// TODO: consider what options need to be specified
-		BlackboardEntry(const EntrySettings& InSettings, uint32 InBufferSize, ERollbackBufferFlags InBufferFlags, const EntryT* OptionalInitialObj=nullptr)
-			: BlackboardEntryBase(InSettings, InBufferSize, InBufferFlags)
-			, EntryBuffer(InBufferSize, InBufferFlags), DefaultValue()
+		BlackboardEntry(const EntrySettings& InSettings, uint32 InBufferSize, const EntryT* OptionalInitialObj=nullptr)
+			: BlackboardEntryBase(InSettings, InBufferSize)
+			, EntryBuffer(InBufferSize)
 		{
 			// JAH TODO: Consider making the initial object required, so we can also initialize where in the buffer we point to
 			if (OptionalInitialObj)
@@ -208,20 +201,13 @@ private:
 				{
 					EntryBuffer[i] = *OptionalInitialObj;
 				}
-
-				DefaultValue = *OptionalInitialObj;
 			}
 		}
 
-		void SetValueAtTime_External(const EntryT& Value, const EntryTimeStamp& TimeStamp)
-		{
-			EntryBuffer[ExternalIdx] = Value;
-			Timestamps[ExternalIdx] = TimeStamp;
-		}
 
-		bool TryGetValueAtTime_External(const EntryTimeStamp& CurrentExternalTime, EntryT& ValueOut) const
+		bool TryGetEntryValue(const EntryTimeStamp& CurrentExternalTime, EntryT& ValueOut)
 		{
-			if (CanReadEntryAtTime(CurrentExternalTime, EEntryIndexType::External))
+			if (CanReadEntryAt(CurrentExternalTime, EEntryIndexType::External))
 			{
 				ValueOut = EntryBuffer[ExternalIdx];
 				return true;
@@ -231,9 +217,9 @@ private:
 		}
 
 		// TODO: Consider whether it's worthwhile to have this function... dangerous to let external folks hold onto references to internal storage, but may allow more efficient editing
-		const EntryT* TryGetRefAtTime_External(const EntryTimeStamp& CurrentExternalTime) const
+		EntryT* TryGetEntryRef(const EntryTimeStamp& CurrentExternalTime)
 		{
-			if (CanReadEntryAtTime(CurrentExternalTime, EEntryIndexType::External))
+			if (CanReadEntryAt(CurrentExternalTime, EEntryIndexType::External))
 			{ 
 				return &EntryBuffer[ExternalIdx];
 			}
@@ -241,9 +227,9 @@ private:
 			return nullptr;
 		}
 
-		bool TryGetValueAtTime_Internal(const EntryTimeStamp& InternalTime, EntryT& ValueOut) const
+		bool TryGetEntryValue_Internal(const EntryTimeStamp& InternalTime, EntryT& ValueOut)
 		{
-			if (CanReadEntryAtTime(InternalTime, EEntryIndexType::Internal))
+			if (CanReadEntryAt(InternalTime, EEntryIndexType::Internal))
 			{
 				ValueOut = EntryBuffer[InternalIdx];
 				return true;
@@ -252,9 +238,9 @@ private:
 			return false;
 		}
 
-		const EntryT* TryGetRefAtTime_Internal(const EntryTimeStamp& InternalTime) const
+		EntryT* TryGetEntryRef_Internal(const EntryTimeStamp& InternalTime)
 		{
-			if (CanReadEntryAtTime(InternalTime, EEntryIndexType::Internal))
+			if (CanReadEntryAt(InternalTime, EEntryIndexType::Internal))
 			{
 				return &EntryBuffer[InternalIdx];
 			}
@@ -262,7 +248,7 @@ private:
 			return nullptr;
 		}
 
-		void SetValueAtTime_Internal(const EntryT& Value, const EntryTimeStamp& TimeStamp)
+		void SetEntryValue_Internal(const EntryT& Value, const EntryTimeStamp& TimeStamp)
 		{
 			InternalIdx = ExternalIdx + 1;	// always writing internal, one slot ahead of external slot
 
@@ -270,66 +256,10 @@ private:
 			Timestamps[InternalIdx] = TimeStamp;
 		}
 
-		bool TryGetValueAtTime_Predictive(const EntryTimeStamp& PredictiveReaderTime, EntryT& ValueOut) const
-		{
-			// For predictive reading, we can either be reading the current External entry
-			// OR the Predictive entry if it's been written to during this Predictive window.
-			const bool bHasPredictiveValue = Timestamps.GetPredictiveElement().IsValid();
-			
-			if (bHasPredictiveValue)
-			{
-				if (CanReadEntryAtTime(PredictiveReaderTime, EEntryIndexType::Predictive))
-				{
-					ValueOut = EntryBuffer.GetPredictiveElement();
-					return true;
-				}
-			}
-			else
-			{
-				// No prediction writes have occurred yet, so we'll try to read the external value
-				return TryGetValueAtTime_External(PredictiveReaderTime, ValueOut);
-			}
-
-			return false;
-		}
-
-		const EntryT* TryGetRefAtTime_Predictive(const EntryTimeStamp& PredictiveReaderTime) const
-		{
-			// For predictive reading, we can either be reading the current External entry
-			// OR the Predictive entry if it's been written to during this Predictive window.
-			const bool bHasPredictiveValue = Timestamps.GetPredictiveElement().IsValid();
-
-			if (bHasPredictiveValue)
-			{
-				if (CanReadEntryAtTime(PredictiveReaderTime, EEntryIndexType::Predictive))
-				{
-					return &EntryBuffer.GetPredictiveElement();
-				}
-			}
-			else
-			{
-				// No prediction writes have occurred yet, so we'll try to read the external value
-				return TryGetRefAtTime_External(PredictiveReaderTime);
-			}
-
-			return nullptr;
-		}
-
-		void SetValueAtTime_Predictive(const EntryT& Value, const EntryTimeStamp& TimeStamp)
-		{
-			// Writes ALWAYS go to the predictive element
-			EntryBuffer.GetPredictiveElement() = Value;
-			Timestamps.GetPredictiveElement() = TimeStamp;
-		}
-
-		virtual void ResetPredictiveElement() override
-		{
-			EntryBuffer.GetPredictiveElement() = DefaultValue;
-		}
-
 	
-		TRollbackCircularBuffer<EntryT> EntryBuffer;
-		EntryT DefaultValue;
+		// TODO: Consider whether these need to be TCircularBuffers, or whether a TArray would be just fine.  TCircularBuffers round up to the next power of 2 in capacity, so may be wasting significant memory for larger structures.
+		// Note that these EntryBuffer is always the same capacity as Timestamps
+		TCircularBuffer<EntryT> EntryBuffer;
 
 
 	};	// end BlackboardEntry
@@ -347,7 +277,7 @@ public:
 		}
 
 		const uint32 BufferSize = BlackboardEntryBase::ComputeBufferSize(InSettings.SizingPolicy, InSettings.FixedSize);
-		EntryMap.Emplace(EntryName, MakeUnique<BlackboardEntry<EntryT>>(InSettings, BufferSize, BufferFlags, nullptr));
+		EntryMap.Emplace(EntryName, MakeUnique<BlackboardEntry<EntryT>>(InSettings, BufferSize, nullptr));
 		return true;
 	}
 
@@ -356,16 +286,15 @@ public:
 		return EntryMap.Contains(EntryName);
 	}
 
-private:
 	/** Store object by a named key, overwriting any existing object */
 	template<typename EntryT>
-	bool TrySet_External(FName ObjName, const EntryT& Obj)
+	bool TrySet(FName ObjName, EntryT& Obj)
 	{
 		// TODO: Verify this isn't happening mid-sim on the sim thread (_Internal should be used instead)
 
-		if (BlackboardEntry<EntryT>* Entry = static_cast<BlackboardEntry<EntryT>*>(FindEntry(ObjName)))
+		if (BlackboardEntry<EntryT>* Entry = static_cast<BlackboardEntry<EntryT>*>(FindEntry_Internal(ObjName)))
 		{
-			Entry->SetValueAtTime_External(Obj, CurrentSimTimeStamp);
+			Entry->SetEntryValue(Obj);
 			return true;
 		}
 
@@ -375,11 +304,11 @@ private:
 
 	/** Attempt to retrieve an object from the blackboard. If found, OutFoundValue will be set. Returns true/false to indicate whether it was found. */
 	template<typename EntryT>
-	bool TryGet_External(FName ObjName, EntryT& OutFoundValue) const
+	bool TryGet(FName ObjName, EntryT& OutFoundValue) const
 	{
-		if (BlackboardEntry<EntryT>* Entry = static_cast<BlackboardEntry<EntryT>*>(FindEntry(ObjName)))
+		if (BlackboardEntry<EntryT>* Entry = static_cast<BlackboardEntry<EntryT>*>(FindEntry_Internal(ObjName)))
 		{
-			return Entry->TryGetValueAtTime_External(CurrentSimTimeStamp, OutFoundValue);
+			return Entry->TryGetEntryValue(CurrentSimTimeStamp, OutFoundValue);
 		}
 
 		return false;
@@ -389,27 +318,28 @@ private:
 	 * Although you can get access to the reference, it is not safe to hold references over time. 
 	 */
 	template<typename EntryT>
-	const EntryT* TryGetRef_External(FName ObjName) const
+	EntryT* TryGetRef(FName ObjName) const
 	{
-		if (BlackboardEntry<EntryT>* Entry = static_cast<BlackboardEntry<EntryT>*>(FindEntry(ObjName)))
+		if (BlackboardEntry<EntryT>* Entry = static_cast<BlackboardEntry<EntryT>*>(FindEntry_Internal(ObjName)))
 		{
-			return Entry->TryGetRefAtTime_External(CurrentSimTimeStamp);
+			return Entry->TryGetEntryRef(CurrentSimTimeStamp);
 		}
 
 		return nullptr;
 	}
 
 
-	/** Store object of a named key, overwriting any existing object. Should only be called from within the simulation. */
+private:
+	/** Store object by a named key, overwriting any existing object. Should only be called from within the simulation. */
 	template<typename EntryT>
 	bool TrySet_Internal(FName ObjName, const EntryT& Obj)
 	{
 		check((bIsSimulationInProgress && (InProgressSimFrameThreadId == FPlatformTLS::GetCurrentThreadId())) || 
 			  (bIsRollbackInProgress && (InRollbackThreadId == FPlatformTLS::GetCurrentThreadId()))) ;
 
-		if (BlackboardEntry<EntryT>* Entry = static_cast<BlackboardEntry<EntryT>*>(FindEntry(ObjName)))
+		if (BlackboardEntry<EntryT>* Entry = static_cast<BlackboardEntry<EntryT>*>(FindEntry_Internal(ObjName)))
 		{
-			Entry->SetValueAtTime_Internal(Obj, InProgressSimTimeStamp);
+			Entry->SetEntryValue_Internal(Obj, InProgressSimTimeStamp);
 			return true;
 		}
 
@@ -422,9 +352,9 @@ private:
 		check((bIsSimulationInProgress && (InProgressSimFrameThreadId == FPlatformTLS::GetCurrentThreadId())) ||
 			  (bIsRollbackInProgress && (InRollbackThreadId == FPlatformTLS::GetCurrentThreadId())));
 
-		if (BlackboardEntry<EntryT>* Entry = static_cast<BlackboardEntry<EntryT>*>(FindEntry(ObjName)))
+		if (BlackboardEntry<EntryT>* Entry = static_cast<BlackboardEntry<EntryT>*>(FindEntry_Internal(ObjName)))
 		{
-			return Entry->TryGetValueAtTime_Internal(InProgressSimTimeStamp, OutFoundValue);
+			return Entry->TryGetEntryValue_Internal(InProgressSimTimeStamp, OutFoundValue);
 		}
 
 		return false;
@@ -436,59 +366,17 @@ private:
 		check((bIsSimulationInProgress && (InProgressSimFrameThreadId == FPlatformTLS::GetCurrentThreadId())) ||
 			  (bIsRollbackInProgress && (InRollbackThreadId == FPlatformTLS::GetCurrentThreadId())));
 
-		if (BlackboardEntry<EntryT>* Entry = static_cast<BlackboardEntry<EntryT>*>(FindEntry(ObjName)))
+		if (BlackboardEntry<EntryT>* Entry = static_cast<BlackboardEntry<EntryT>*>(FindEntry_Internal(ObjName)))
 		{
-			return Entry->TryGetRefAtTime_Internal(InProgressSimTimeStamp);
+			return Entry->TryGetEntryRef_Internal(InProgressSimTimeStamp);
 		}
 
 		return nullptr;
 	}
 
-private:
-	/** Store object by a named key, overwriting any existing object. Should only be called from within the predicting simulation. */
-	template<typename EntryT>
-	bool TrySet_Predictive(FName ObjName, const EntryT& Obj)
-	{
-		check(bIsPredictionInProgress && (InPredictionThreadId == FPlatformTLS::GetCurrentThreadId()));
 
-		if (BlackboardEntry<EntryT>* Entry = static_cast<BlackboardEntry<EntryT>*>(FindEntry(ObjName)))
-		{
-			Entry->SetValueAtTime_Predictive(Obj, InProgressSimTimeStamp);
-			return true;
-		}
 
-		return false;
-	}
-
-	template<typename EntryT>
-	bool TryGet_Predictive(FName ObjName, EntryT& OutFoundValue) const
-	{
-		check(bIsPredictionInProgress && (InPredictionThreadId == FPlatformTLS::GetCurrentThreadId()));
-
-		if (BlackboardEntry<EntryT>* Entry = static_cast<BlackboardEntry<EntryT>*>(FindEntry(ObjName)))
-		{
-			return Entry->TryGetValueAtTime_Predictive(InProgressSimTimeStamp, OutFoundValue);
-		}
-
-		return false;
-	}
-
-	template<typename EntryT>
-	const EntryT* TryGetRef_Predictive(FName ObjName) const
-	{
-		check(bIsPredictionInProgress && (InPredictionThreadId == FPlatformTLS::GetCurrentThreadId()));
-
-		if (BlackboardEntry<EntryT>* Entry = static_cast<BlackboardEntry<EntryT>*>(FindEntry(ObjName)))
-		{
-			return Entry->TryGetRefAtTime_Predictive(InProgressSimTimeStamp);
-		}
-
-		return nullptr;
-	}
-
-private:
-
-	BlackboardEntryBase* FindEntry(FName EntryName) const
+	BlackboardEntryBase* FindEntry_Internal(FName EntryName) const
 	{
 		if (const TUniquePtr<BlackboardEntryBase>* EntryPtr = EntryMap.Find(EntryName))
 		{
@@ -499,46 +387,90 @@ private:
 	}
 
 
-public:	// Administration functions for keeping the blackboard in sync with the simulation
-
-	// Begin/EndSimulationFrames: advances true simulation flow, regardless of whether this is resimulation or not. Calls will bookend a simulation frame.
+private:
 	UE_API void BeginSimulationFrame(const FMoverTimeStep& PendingTimeStep);
 	UE_API void EndSimulationFrame();
 
-	// Begin/EndRollback: reverts the blackboard back to an older state. EndRollback will be called before any simulation frames are started (resim or not).
-	UE_API void BeginRollback(const FMoverTimeStep& NewBaseTimeStep);	// NewBaseTimeStep represents the current time and the PENDING step that's next to be resimulated
+	UE_API void BeginRollback(const FMoverTimeStep& NewBaseTimeStep);	// NewBaseTimeStep represents the current time and the PENDING frame that's next to be resimulated
 	UE_API void EndRollback();
 
-	// BeginPredictionFrame/EndPrediction: this is for predictive simulation that is temporary and can't rollback. Note there may be multiple calls to BeginPredictionFrame before a single call to EndPrediction.
-	UE_API void BeginPredictionFrame(const FMoverTimeStep& PendingTimeStep);	// Called to start a new prediction OR continue a sequential frame
-	UE_API void EndPrediction();	// Called at the end of ALL prediction frames
 
 private:
-	UE_API bool IsInSimulationThread() const;
-	UE_API bool IsInPredictionThread() const;
-
-	ERollbackBufferFlags BufferFlags = ERollbackBufferFlags::Prediction;
 	
 	TMap<FName, TUniquePtr<BlackboardEntryBase>> EntryMap;
 
 
 	bool bIsSimulationInProgress = false;
 	bool bIsResimulating = false;
-	uint32 InProgressSimFrameThreadId = 0;
+	uint32 InProgressSimFrameThreadId;
 
 	bool bIsRollbackInProgress = false;
-	uint32 InRollbackThreadId = 0;
+	uint32 InRollbackThreadId;
 
-	bool bIsPredictionInProgress = false;
-	uint32 InPredictionThreadId = 0;
 
 	EntryTimeStamp CurrentSimTimeStamp;	// this is the "committed" simulation time, and will lag behind InProgressSimTimeStamp while in the middle of a simulation frame
 	EntryTimeStamp InProgressSimTimeStamp;	// this is the sim time that's actively being used mid-simulation. It is only valid during a simulation step.
 
-	// we expose user API through these wrapper types
-	friend struct FRollbackBlackboardSimWrapper;
-	friend struct FRollbackBlackboardExternalWrapper;
+	friend class URollbackBlackboard_InternalWrapper;	// expose non-public API to this wrapper class
 };
 
 
+
+/** 
+ * Wrapper class for in-simulation access to the blackboard. This exposes otherwise private API, and redirects
+ * set/get operations to the internal versions.
+ */
+UCLASS(MinimalAPI)
+class URollbackBlackboard_InternalWrapper : public UObject
+{
+	GENERATED_BODY()
+
+public:
+	void Init(URollbackBlackboard& InBlackboard)
+	{
+		Blackboard = &InBlackboard;
+	}
+
+	template<typename EntryT>
+	bool CreateEntry(FName EntryName, const URollbackBlackboard::EntrySettings& InSettings)
+	{
+		return Blackboard->CreateEntry<EntryT>(EntryName, InSettings);
+	}
+
+	bool HasEntry(FName EntryName) const
+	{
+		return Blackboard->HasEntry(EntryName);
+	}
+
+	template<typename EntryT>
+	bool TrySet(FName ObjName, const EntryT& Obj)
+	{
+		return Blackboard->TrySet_Internal<EntryT>(ObjName, Obj);
+	}
+
+	template<typename EntryT>
+	bool TryGet(FName ObjName, EntryT& OutFoundValue) const
+	{
+		return Blackboard->TryGet_Internal(ObjName, OutFoundValue);
+	}
+
+	template<typename EntryT>
+	const EntryT* TryGetRef(FName ObjName) const
+	{
+		return Blackboard->TryGetRef_Internal<EntryT>(ObjName);
+	}
+	
+	// TODO: Also support a mutable TryGetRef that gives you a writable reference
+
+	UE_API void BeginSimulationFrame(const FMoverTimeStep& PendingTimeStep);
+	UE_API void EndSimulationFrame();
+
+	UE_API void BeginRollback(const FMoverTimeStep& NewBaseTimeStep);
+	UE_API void EndRollback();
+
+private: 
+	UPROPERTY()
+	TObjectPtr<URollbackBlackboard> Blackboard;
+};
+
 #undef UE_API
diff --git a/Engine/Plugins/Experimental/Mover/Source/Mover/Public/MoveLibrary/RollbackBlackboardWrappers.h b/Engine/Plugins/Experimental/Mover/Source/Mover/Public/MoveLibrary/RollbackBlackboardWrappers.h
deleted file mode 100644
diff --git a/Engine/Plugins/Experimental/Mover/Source/Mover/Public/MoveLibrary/RollbackCircularBuffer.h b/Engine/Plugins/Experimental/Mover/Source/Mover/Public/MoveLibrary/RollbackCircularBuffer.h
deleted file mode 100644
diff --git a/Engine/Plugins/Experimental/Mover/Source/Mover/Public/MovementModeStateMachine.h b/Engine/Plugins/Experimental/Mover/Source/Mover/Public/MovementModeStateMachine.h
--- a/Engine/Plugins/Experimental/Mover/Source/Mover/Public/MovementModeStateMachine.h
+++ b/Engine/Plugins/Experimental/Mover/Source/Mover/Public/MovementModeStateMachine.h
@@ -6,7 +6,6 @@
 #include "Misc/TransactionallySafeRWLock.h"
 #include "Templates/SubclassOf.h"
 #include "MovementMode.h"
-#include "MoveLibrary/RollbackBlackboardWrappers.h"
 #include "InstantMovementEffect.h"
 #include "MovementModeStateMachine.generated.h"
 
@@ -38,9 +37,6 @@ public:
 	UE_API void UnregisterGlobalTransition(TObjectPtr<UBaseMovementModeTransition> Transition);
 	UE_API void ClearAllGlobalTransitions();
 
-	// TODO: Warning: this function will be removed once the full simulation context (including blackboard) comes in as part of the Simulation**** function params
-	UE_API void SetSimulation(UMoverSimulation& InSimulation);
-
 	UE_API void SetDefaultMode(FName NewDefaultModeName);
 
 	UE_API void QueueNextMode(FName DesiredNextModeName, bool bShouldReenter=false);
@@ -92,9 +88,6 @@ protected:
 	UPROPERTY(Transient)
 	TObjectPtr<UImmediateMovementModeTransition> QueuedModeTransition;
 
-	UPROPERTY(Transient)
-	TObjectPtr<UMoverSimulation> Simulation;
-
 	FName DefaultModeName = NAME_None;
 	FName CurrentModeName = NAME_None;
 
diff --git a/Engine/Plugins/Experimental/Mover/Source/Mover/Public/MoverComponent.h b/Engine/Plugins/Experimental/Mover/Source/Mover/Public/MoverComponent.h
--- a/Engine/Plugins/Experimental/Mover/Source/Mover/Public/MoverComponent.h
+++ b/Engine/Plugins/Experimental/Mover/Source/Mover/Public/MoverComponent.h
@@ -13,7 +13,6 @@
 #include "LayeredMoveBase.h"
 #include "MoveLibrary/BasedMovementUtils.h"
 #include "MoveLibrary/ConstrainedMoveUtils.h"
-#include "MoveLibrary/RollbackBlackboardWrappers.h"
 #if UE_ENABLE_INCLUDE_ORDER_DEPRECATED_IN_5_4
 #include "Engine/HitResult.h"
 #endif // UE_ENABLE_INCLUDE_ORDER_DEPRECATED_IN_5_4
@@ -30,6 +29,7 @@ struct FMoverSimulationEventData;
 class UMovementModeStateMachine;
 class UMovementMixer;
 class URollbackBlackboard;
+class URollbackBlackboard_InternalWrapper;
 
 namespace MoverComponentConstants
 {
@@ -541,8 +541,9 @@ public:	// Queries
 
 	MOVER_API UMoverBlackboard* GetSimBlackboard_Mutable() const;
 
-	// Access the version of the movement blackboard for users outside of the movement simulation
-	FRollbackBlackboardExternalWrapper GetRollbackBlackboardExternal() const { return RollbackBlackboardExternalWrapper; }
+
+	URollbackBlackboard* GetRollbackBlackboard() const { return RollbackBlackboard.Get(); }
+	URollbackBlackboard_InternalWrapper* GetRollbackBlackboard_Internal() const { return RollbackBlackboard_InternalWrapper.Get(); }
 
 	/** Find settings object by type. Returns null if there is none of that type */
 	const IMovementSettingsInterface* FindSharedSettings(const UClass* ByType) const { return FindSharedSettings_Mutable(ByType); }
@@ -694,8 +695,6 @@ protected:
 public:
 	MOVER_API virtual void CreateDefaultInputAndState(FMoverInputCmdContext& OutInputCmd, FMoverSyncState& OutSyncState, FMoverAuxStateContext& OutAuxState) const;
 
-	MOVER_API void SetSimulation(UMoverSimulation& Simulation);		// Temporary support function. Will be removed once simulations are used consistently.
-
 	/** Handle a blocking impact.*/
 	UFUNCTION(BlueprintCallable, Category = Mover)
 	MOVER_API void HandleImpact(FMoverOnImpactParams& ImpactParams);
@@ -790,17 +789,12 @@ protected:
 	UPROPERTY(Transient)
 	TObjectPtr<UMoverBlackboard> SimBlackboard = nullptr;
 
-	/** This is what all movement mechanisms should be interacting with during simulation, rather than relying on actors/components. External parties should generally not see it. */
-	UPROPERTY(Transient)
-	TObjectPtr<UMoverSimulation> Simulation = nullptr;
-
-	// Accessors used to interface with the simulation's rollback blackboard
+	/** Used to store cached data & computations between decoupled systems, that can be referenced by name. Rollback-aware. */
 	UPROPERTY(Transient)
-	FRollbackBlackboardExternalWrapper RollbackBlackboardExternalWrapper;
+	TObjectPtr<URollbackBlackboard> RollbackBlackboard = nullptr;
 
-	/** Ref to the actual blackboard. This is temporary and will be removed once simulations are used consistently. */
 	UPROPERTY(Transient)
-	TObjectPtr<URollbackBlackboard> RollbackBlackboard = nullptr;
+	TObjectPtr<URollbackBlackboard_InternalWrapper> RollbackBlackboard_InternalWrapper = nullptr;	// This is a thin layer for use only by in-simulation users
 
 	/**
 	 * Layered moves registered on this component that can be activated regardless of the current mode
diff --git a/Engine/Plugins/Experimental/Mover/Source/Mover/Public/MoverSimulation.h b/Engine/Plugins/Experimental/Mover/Source/Mover/Public/MoverSimulation.h
--- a/Engine/Plugins/Experimental/Mover/Source/Mover/Public/MoverSimulation.h
+++ b/Engine/Plugins/Experimental/Mover/Source/Mover/Public/MoverSimulation.h
@@ -3,13 +3,14 @@
 #pragma once
 
 #include "UObject/Object.h"
-#include "MoveLibrary/RollbackBlackboardWrappers.h"
+
 #include "MoverSimulation.generated.h"
 
 #define UE_API MOVER_API
 
 class UMoverBlackboard;
 struct FMoverSyncState;
+class URollbackBlackboard_InternalWrapper;
 
 /**
 * WIP Base class for a Mover simulation.
@@ -33,7 +34,11 @@ public:
 	UE_API UMoverBlackboard* GetBlackboard_Mutable();
 
 	UFUNCTION(BlueprintPure, Category = Mover)
-	UE_API FRollbackBlackboardSimWrapper GetRollbackBlackboardSimWrapper() const;
+	UE_API const URollbackBlackboard_InternalWrapper* GetRollbackBlackboard() const;
+
+	UFUNCTION(BlueprintPure, Category = Mover)
+	UE_API URollbackBlackboard_InternalWrapper* GetRollbackBlackboard_Mutable();
+
 
 	/**
 	* Attempt to teleport to TargetTransform. The teleport is not guaranteed to happen. This function is meant to be called by an instant movement effect as part of its effect application.
@@ -47,18 +52,16 @@ public:
 	virtual void AttemptTeleport(const FMoverTimeStep& TimeStep, const FTransform& TargetTransform, bool bUseActorRotation, FMoverSyncState& OutputState) {}
 
 
+	// Used during initialization only
+	UE_API void SetRollbackBlackboard(URollbackBlackboard_InternalWrapper* RollbackSimBlackboard);
+
 protected:
 	UPROPERTY(Transient)
 	TObjectPtr<UMoverBlackboard> Blackboard = nullptr;
 
 	UPROPERTY(Transient)
-	TObjectPtr<URollbackBlackboard> RollbackBlackboard = nullptr;
-
-	FRollbackBlackboardSimWrapper RollbackBlackboardSimWrapper;
+	TObjectPtr<URollbackBlackboard_InternalWrapper> RollbackBlackboard = nullptr;
 
-// TODO: These friend clases are temporary until the control flow changes to include a Simulation object consistently
-friend class UMoverComponent;
-friend class UMoverNetworkPhysicsLiaisonComponentBase;
 };
 
 #undef UE_API
diff --git a/Engine/Plugins/Experimental/MoverTests/MoverTests.uplugin b/Engine/Plugins/Experimental/MoverTests/MoverTests.uplugin
--- a/Engine/Plugins/Experimental/MoverTests/MoverTests.uplugin
+++ b/Engine/Plugins/Experimental/MoverTests/MoverTests.uplugin
@@ -29,10 +29,6 @@
 		{
 			"Name": "MoverExamples",
 			"Enabled": true
-		},
-		{
-			"Name": "RuntimeTests",
-			"Enabled": true
 		}
 	]
 }
\ No newline at end of file
diff --git a/Engine/Plugins/Experimental/MoverTests/Source/MoverTests/MoverTests.Build.cs b/Engine/Plugins/Experimental/MoverTests/Source/MoverTests/MoverTests.Build.cs
--- a/Engine/Plugins/Experimental/MoverTests/Source/MoverTests/MoverTests.Build.cs
+++ b/Engine/Plugins/Experimental/MoverTests/Source/MoverTests/MoverTests.Build.cs
@@ -31,7 +31,6 @@ public class MoverTests : ModuleRules
 				"Engine",
 				"Slate",
 				"SlateCore",
-				"RuntimeTests",
 				"Mover",
 				"MoverExamples"
 				// ... add other public dependencies that you statically link with here ...
diff --git a/Engine/Plugins/Experimental/MoverTests/Source/MoverTests/Private/RollbackBlackboardTests.cpp b/Engine/Plugins/Experimental/MoverTests/Source/MoverTests/Private/RollbackBlackboardTests.cpp
deleted file mode 100644
