Fix for crash that can occur when you edit properties via editor Details panel on an Actor using a MoverComponent during play.
The crash was occurring because the component was attempting to reinitialize something that had already been initialized.

Additionally, an assert was being hit when when using the network physics liaison due to it being marked a "Within = MoverComponent" This assert was due to the component being renamed when being destroyed to a name that did not include that it was indeed within the mover component. This appears to be a long standing bug in ActorConstruction.cpp that only really matters in the case when the component being destroyed is Within another component rather than the actor that is being destroyed.. This change makes it so it gets the component's outer rather than assuming that the outer is the Actor and uses that for the rename

#Synced-CL 50587857

diff --git a/Engine/Plugins/Experimental/Mover/Source/Mover/Private/MoverComponent.cpp b/Engine/Plugins/Experimental/Mover/Source/Mover/Private/MoverComponent.cpp
--- a/Engine/Plugins/Experimental/Mover/Source/Mover/Private/MoverComponent.cpp
+++ b/Engine/Plugins/Experimental/Mover/Source/Mover/Private/MoverComponent.cpp
@@ -149,13 +149,21 @@ void UMoverComponent::InitializeComponent()
 		// Instantiate our sister backend component that will actually talk to the system driving the simulation
 		if (BackendClass)
 		{
-			UActorComponent* NewLiaisonComp = NewObject<UActorComponent>(this, BackendClass, TEXT("BackendLiaisonComponent"));
+			UActorComponent* NewLiaisonComp = Cast<UActorComponent>(GetOwner()->FindComponentByInterface<IMoverBackendLiaisonInterface>());
+			if (!NewLiaisonComp)
+			{
+				NewLiaisonComp = NewObject<UActorComponent>(this, BackendClass, TEXT("BackendLiaisonComponent"));
+			}
+
 			BackendLiaisonComp.SetObject(NewLiaisonComp);
 			BackendLiaisonComp.SetInterface(CastChecked<IMoverBackendLiaisonInterface>(NewLiaisonComp));
 			if (BackendLiaisonComp)
 			{
 				NewLiaisonComp->RegisterComponent();
-				NewLiaisonComp->InitializeComponent();
+				if (!NewLiaisonComp->HasBeenInitialized())
+				{
+					NewLiaisonComp->InitializeComponent();
+				}
 				NewLiaisonComp->SetNetAddressable();
 			}
 		}
diff --git a/Engine/Source/Runtime/Engine/Private/ActorConstruction.cpp b/Engine/Source/Runtime/Engine/Private/ActorConstruction.cpp
--- a/Engine/Source/Runtime/Engine/Private/ActorConstruction.cpp
+++ b/Engine/Source/Runtime/Engine/Private/ActorConstruction.cpp
@@ -207,8 +207,9 @@ void AActor::DestroyConstructedComponents()
 
 				// Rename component to avoid naming conflicts in the case where we rerun the SCS and name the new components the same way.
 				FName const NewBaseName( *(FString::Printf(TEXT("TRASH_%s"), *Component->GetClass()->GetName())) );
-				FName const NewObjectName = MakeUniqueObjectName(this, GetClass(), NewBaseName);
-				Component->Rename(*NewObjectName.ToString(), this, REN_DontCreateRedirectors|REN_NonTransactional|REN_DoNotDirty);
+				FName const NewObjectName = MakeUniqueObjectName(GetTransientOuterForRename(Component->GetClass()), GetClass(), NewBaseName);
+				Component->Rename(*NewObjectName.ToString(), GetTransientOuterForRename(Component->GetClass()), REN_DontCreateRedirectors|REN_NonTransactional|REN_DoNotDirty);
+
 
 				// Transient actors should never mark the package as dirty
 				bMarkPackageDirty = !HasAnyFlags(RF_Transient);
