Mover: fix for crash that could occur when dragging and dropping an asset into an already-running PIE session

#Synced-CL 50896471

diff --git a/Engine/Plugins/Experimental/Mover/Source/Mover/Private/Backends/MoverNetworkPredictionLiaison.cpp b/Engine/Plugins/Experimental/Mover/Source/Mover/Private/Backends/MoverNetworkPredictionLiaison.cpp
--- a/Engine/Plugins/Experimental/Mover/Source/Mover/Private/Backends/MoverNetworkPredictionLiaison.cpp
+++ b/Engine/Plugins/Experimental/Mover/Source/Mover/Private/Backends/MoverNetworkPredictionLiaison.cpp
@@ -138,6 +138,13 @@ void UMoverNetworkPredictionLiaisonComponent::SimulationTick(const FNetSimTimeSt
 {
 	check(MoverComp);
 
+	if (!HasBegunPlay())	// If we're not ready to run yet, maintain the same simulation state
+	{
+		*SimOutput.Sync = *SimInput.Sync;
+		*SimOutput.Aux.Get() = *SimInput.Aux;
+		return;
+	}
+
 	FMoverTickStartData StartData;
 	FMoverTickEndData EndData;
 
