Mover: fix crash when using User-Defined Structs in a Mover data collection when Iris networking is enabled. This works around an Instanced Struct limitation when there is not a native NetSerialize function

#Synced-CL 50896471

diff --git a/Engine/Plugins/Experimental/Mover/Source/Mover/Private/UserDefinedStructSupport.cpp b/Engine/Plugins/Experimental/Mover/Source/Mover/Private/UserDefinedStructSupport.cpp
--- a/Engine/Plugins/Experimental/Mover/Source/Mover/Private/UserDefinedStructSupport.cpp
+++ b/Engine/Plugins/Experimental/Mover/Source/Mover/Private/UserDefinedStructSupport.cpp
@@ -3,6 +3,11 @@
 
 #include "UserDefinedStructSupport.h"
 #include "StructUtils/UserDefinedStruct.h"
+#include "Engine/Engine.h"
+#include "Engine/NetConnection.h"
+#include "Engine/NetDriver.h"
+#include "Engine/PackageMapClient.h"
+#include "Net/RepLayout.h"
 #include "MoverLog.h"
 
 #include UE_INLINE_GENERATED_CPP_BY_NAME(UserDefinedStructSupport)
@@ -71,7 +76,7 @@ bool FMoverUserDefinedDataStruct::NetSerialize(FArchive& Ar, UPackageMap* Map, b
 	bool bSuperSuccess, bStructSuccess;
 
 	Super::NetSerialize(Ar, Map, bSuperSuccess);
-	StructInstance.NetSerialize(Ar, Map, bStructSuccess);
+	bStructSuccess = NetSerializeUserDefinedStruct(StructInstance, Ar, Map);
 
 	bOutSuccess = bSuperSuccess && bStructSuccess;
 
@@ -92,4 +97,67 @@ const UScriptStruct* FMoverUserDefinedDataStruct::GetDataScriptStruct() const
 }
 
 
+bool FMoverUserDefinedDataStruct::NetSerializeUserDefinedStruct(FInstancedStruct& UDSStructInstance, FArchive& Ar, UPackageMap* Map)
+{
+	UScriptStruct* NonConstScriptStruct = nullptr;	// This is not modified directly, but needed for some non-const function calls
+
+	if (Ar.IsLoading())
+	{
+		Ar << NonConstScriptStruct;
+
+		// Initialize only if the type changes.
+		if (UDSStructInstance.GetScriptStruct() != NonConstScriptStruct)
+		{
+			UDSStructInstance.InitializeAs(NonConstScriptStruct);
+		}
+
+		if (!UDSStructInstance.IsValid())
+		{
+			NonConstScriptStruct = nullptr;
+			UE_LOG(LogMover, Error, TEXT("FMoverUserDefinedDataStruct::NetSerializeUserDefinedStruct: Bad script struct serialized when loading, cannot recover."));
+			Ar.SetError();
+		}
+	}
+	else if (Ar.IsSaving())
+	{
+		NonConstScriptStruct = const_cast<UScriptStruct*>(UDSStructInstance.GetScriptStruct());
+		check(::IsValid(NonConstScriptStruct));
+		Ar << NonConstScriptStruct;
+	}
+
+	// Check ScriptStruct here, as loading might have failed. 
+	if (NonConstScriptStruct)
+	{
+		check((NonConstScriptStruct->StructFlags & STRUCT_NetSerializeNative) == 0);	// We should never have a native netserialize function with User-Defined Structs
+
+		// Note that we are doing this method of net serialization because FInstancedStruct's NetSerialize
+		// does not work with Iris yet. We're also able to assume our struct does not have a native NetSerialize
+		// because it's a User-Defined Struct.
+		UNetDriver* NetDriver = nullptr;
+
+		if (UPackageMapClient* MapClient = Cast<UPackageMapClient>(Map))	// Non-Iris cases provide a pkg map that lets us get the net driver
+		{
+			if (UNetConnection* NetConnection = MapClient->GetConnection())
+			{
+				NetDriver = NetConnection->GetDriver();
+			}
+		}
+		else if (UWorld* World = GEngine->GetCurrentPlayWorld())	// Fall back case when Iris is enabled
+		{
+			NetDriver = World->GetNetDriver();
+		}
+
+		check(::IsValid(NetDriver));
+
+		const TSharedPtr<FRepLayout> RepLayout = NetDriver->GetStructRepLayout(NonConstScriptStruct);
+		check(RepLayout.IsValid());
+
+		bool bHasUnmapped = false;
+		RepLayout->SerializePropertiesForStruct(NonConstScriptStruct, static_cast<FBitArchive&>(Ar), Map, UDSStructInstance.GetMutableMemory(), bHasUnmapped);
+		return true;
+	}
+
+	return false;
+}
+
 #undef LOCTEXT_NAMESPACE
diff --git a/Engine/Plugins/Experimental/Mover/Source/Mover/Public/UserDefinedStructSupport.h b/Engine/Plugins/Experimental/Mover/Source/Mover/Public/UserDefinedStructSupport.h
--- a/Engine/Plugins/Experimental/Mover/Source/Mover/Public/UserDefinedStructSupport.h
+++ b/Engine/Plugins/Experimental/Mover/Source/Mover/Public/UserDefinedStructSupport.h
@@ -27,7 +27,7 @@ struct FMoverUserDefinedDataStruct : public FMoverDataStructBase
 
 public:
 	UPROPERTY()
-	FInstancedStruct StructInstance;
+	FInstancedStruct StructInstance;	// this always represents a User-Defined Struct from Blueprint
 
 	// Implementation of FMoverDataStructBase
 	virtual bool ShouldReconcile(const FMoverDataStructBase& AuthorityState) const override;
@@ -42,7 +42,7 @@ public:
 	// This returns the User-Defined Struct type
 	virtual const UScriptStruct* GetDataScriptStruct() const override;
 
-
+	static bool NetSerializeUserDefinedStruct(FInstancedStruct& UDSStructInstance, FArchive& Ar, UPackageMap* Map);
 };
 
 
