Mover: updating tooltip with info about which skeletal mesh will be used.

#Synced-CL 50944842

diff --git a/Engine/Plugins/Experimental/Mover/Source/MoverEditor/Private/K2Node_PlayMontageOnMoverActor.cpp b/Engine/Plugins/Experimental/Mover/Source/MoverEditor/Private/K2Node_PlayMontageOnMoverActor.cpp
--- a/Engine/Plugins/Experimental/Mover/Source/MoverEditor/Private/K2Node_PlayMontageOnMoverActor.cpp
+++ b/Engine/Plugins/Experimental/Mover/Source/MoverEditor/Private/K2Node_PlayMontageOnMoverActor.cpp
@@ -20,7 +20,7 @@ UK2Node_PlayMontageOnMoverActor::UK2Node_PlayMontageOnMoverActor(const FObjectIn
 
 FText UK2Node_PlayMontageOnMoverActor::GetTooltipText() const
 {
-	return LOCTEXT("K2Node_PlayMontageOnMoverActor_Tooltip", "Plays a Montage on an actor with Mover and SkeletalMesh components. Used for networked animation root motion.");
+	return LOCTEXT("K2Node_PlayMontageOnMoverActor_Tooltip", "Plays a Montage on an actor with Mover and SkeletalMesh components. Used for networked animation root motion.\nThe Primary Visual Component will be used if set, or the first skeletal mesh. See @SetPrimaryVisualComponent");
 }
 
 FText UK2Node_PlayMontageOnMoverActor::GetNodeTitle(ENodeTitleType::Type TitleType) const
