Mover: Adjustments to make Movement Modes selectable outside of MoverComponent

#Synced-CL 50905131

diff --git a/Engine/Plugins/Experimental/Mover/Source/Mover/Private/DefaultMovementSet/Modes/SimpleWalkingMode.cpp b/Engine/Plugins/Experimental/Mover/Source/Mover/Private/DefaultMovementSet/Modes/SimpleWalkingMode.cpp
--- a/Engine/Plugins/Experimental/Mover/Source/Mover/Private/DefaultMovementSet/Modes/SimpleWalkingMode.cpp
+++ b/Engine/Plugins/Experimental/Mover/Source/Mover/Private/DefaultMovementSet/Modes/SimpleWalkingMode.cpp
@@ -12,7 +12,7 @@ void USimpleWalkingMode::GenerateMove_Implementation(const FMoverTickStartData&
 {
 	const FMoverDefaultSyncState* StartingSyncState = StartState.SyncState.SyncStateCollection.FindDataByType<FMoverDefaultSyncState>();
 
-	if (CommonLegacySettings.Get() == nullptr || !StartingSyncState)
+	if (!StartingSyncState)
 	{
 		return;
 	}
@@ -42,29 +42,36 @@ void USimpleWalkingMode::GenerateMove_Implementation(const FMoverTickStartData&
 		MoveInputType = EMoveInputType::DirectionalIntent;
 		DesiredFacingDir = StartingSyncState->GetOrientation_WorldSpace().Quaternion().GetForwardVector();
 	}
-	float MaxMoveSpeed = MaxSpeedOverride >= 0.0f ? MaxSpeedOverride : CommonLegacySettings->MaxSpeed;
+
+	const bool bHasMaxMoveSpeed = MaxSpeedOverride >= 0.0f || CommonLegacySettings;
+	float MaxMoveSpeed = MaxSpeedOverride >= 0.0f ? MaxSpeedOverride : (CommonLegacySettings ? CommonLegacySettings->MaxSpeed : 0.0f);
+
+	const UMoverComponent* MoverComponent = GetMoverComponent();
+	const FVector UpVector = MoverComponent ? MoverComponent->GetUpDirection() : FVector::UpVector;
 
 	// Subtract vertical component but keep same magnitude
 	float DesiredVelMag = DesiredVelocity.Length();
-	DesiredVelocity -= DesiredVelocity.ProjectOnTo(GetMoverComponent()->GetUpDirection());
+	DesiredVelocity -= DesiredVelocity.ProjectOnTo(UpVector);
 	float DesiredVel2DSquaredLength = DesiredVelocity.SquaredLength(); 
 	if (DesiredVel2DSquaredLength > 0.0f)
 	{
 		DesiredVelocity *= DesiredVelMag / FMath::Sqrt(DesiredVel2DSquaredLength); 
 	}
 
+	const float DefaultDirectionalIntentSpeed = 100.0f;
+
 	switch (MoveInputType)
 	{
 	case EMoveInputType::DirectionalIntent:
 		{
 			OutProposedMove.DirectionIntent = DesiredVelocity;	// here, DesiredVelocity is already in "intent space" (unit length for "max intent") so we can use it directly
-			DesiredVelocity *= MaxMoveSpeed;
+			DesiredVelocity = bHasMaxMoveSpeed ? DesiredVelocity * MaxMoveSpeed : DesiredVelocity * DefaultDirectionalIntentSpeed;
 		}
 		break;
 	case EMoveInputType::Velocity:
 		{
 			// Clamp to max move speed
-			DesiredVelocity = DesiredVelocity.GetClampedToMaxSize(MaxMoveSpeed);
+			DesiredVelocity = bHasMaxMoveSpeed ? DesiredVelocity.GetClampedToMaxSize(MaxMoveSpeed) : DesiredVelocity;
 			OutProposedMove.DirectionIntent = MaxMoveSpeed > UE_KINDA_SMALL_NUMBER ? DesiredVelocity / MaxMoveSpeed : FVector::ZeroVector; // here, DesiredVelocity is converted to "intent space"
 		}
 		break;
@@ -80,7 +87,7 @@ void USimpleWalkingMode::GenerateMove_Implementation(const FMoverTickStartData&
 	}
 
 	OutProposedMove.bHasDirIntent = !OutProposedMove.DirectionIntent.IsNearlyZero();
-	DesiredFacingDir -= DesiredFacingDir.ProjectOnTo(GetMoverComponent()->GetUpDirection());
+	DesiredFacingDir -= DesiredFacingDir.ProjectOnTo(UpVector);
 	FQuat CurrentFacing = StartingSyncState->GetOrientation_WorldSpace().Quaternion();
 	FQuat DesiredFacing = CurrentFacing;
 	
diff --git a/Engine/Plugins/Experimental/Mover/Source/Mover/Private/MovementMode.cpp b/Engine/Plugins/Experimental/Mover/Source/Mover/Private/MovementMode.cpp
--- a/Engine/Plugins/Experimental/Mover/Source/Mover/Private/MovementMode.cpp
+++ b/Engine/Plugins/Experimental/Mover/Source/Mover/Private/MovementMode.cpp
@@ -17,13 +17,13 @@
 UWorld* UBaseMovementMode::GetWorld() const
 {
 #if WITH_EDITOR
-	// In the editor, GetWorld() is called on the CDO as part of checking ImplementsGetWorld(). Only the CDO can exist without being outer'd to a MoverComponent.
+	// In the editor, GetWorld() is called on the CDO as part of checking ImplementsGetWorld().
 	if (IsTemplate())
 	{
 		return nullptr;
 	}
 #endif
-	return GetOuterUMoverComponent()->GetWorld();
+	return GetOuter()->GetWorld();
 }
 
 
@@ -103,7 +103,7 @@ void UBaseMovementMode::SimulationTick_Implementation(const FSimulationTickParam
 
 UMoverComponent* UBaseMovementMode::K2_GetMoverComponent() const
 {
-	return GetOuterUMoverComponent();
+	return Cast<UMoverComponent>(GetOuter());
 }
 
 #if WITH_EDITOR
diff --git a/Engine/Plugins/Experimental/Mover/Source/Mover/Public/MovementMode.h b/Engine/Plugins/Experimental/Mover/Source/Mover/Public/MovementMode.h
--- a/Engine/Plugins/Experimental/Mover/Source/Mover/Public/MovementMode.h
+++ b/Engine/Plugins/Experimental/Mover/Source/Mover/Public/MovementMode.h
@@ -35,7 +35,7 @@ public:
 /**
  * Base class for all movement modes, exposing simulation update methods for both C++ and blueprint extension
  */
-UCLASS(MinimalAPI, Abstract, Within = MoverComponent, Blueprintable, BlueprintType, EditInlineNew, DefaultToInstanced)
+UCLASS(MinimalAPI, Abstract, Blueprintable, BlueprintType, EditInlineNew, DefaultToInstanced)
 class UBaseMovementMode : public UObject
 {
 	GENERATED_BODY()
@@ -67,7 +67,6 @@ public:
 
 	/**
 	 * Gets the outer mover component of the indicated type. Does not check on the type or the presence of the MoverComp outer. Safe to call on CDOs.
-	 * Note: Since UBaseMovementMode is declared "Within = MoverComponent", all instances of a mode except the CDO are guaranteed to have a valid MoverComponent outer.
 	 */
 	template<typename MoverT = UMoverComponent UE_REQUIRES(std::is_base_of_v<UMoverComponent, MoverT>)>
 	MoverT* GetMoverComponent() const
@@ -77,12 +76,11 @@ public:
 
 	/**
 	 * Gets the outer mover component of the indicated type, checked for validity.
-	 * Note: Since UBaseMovementMode is declared "Within = MoverComponent", all instances of a mode except the CDO are guaranteed to have a valid MoverComponent outer.
 	 */
 	template<typename MoverT = UMoverComponent UE_REQUIRES(std::is_base_of_v<UMoverComponent, MoverT>)>
 	MoverT& GetMoverComponentChecked() const
 	{
-		return *CastChecked<MoverT>(GetOuterUMoverComponent());
+		return *CastChecked<MoverT>(GetOuter());
 	}
 
 	/**
