Mover: Step Up/Over functions now track and return movement time so we can properly keep track of it in movement modes. This fixed a issue in walking transitioning to falling that launched the character.

#Synced-CL 50999505

diff --git a/Engine/Plugins/Experimental/Mover/Source/Mover/Private/DefaultMovementSet/Modes/AsyncWalkingMode.cpp b/Engine/Plugins/Experimental/Mover/Source/Mover/Private/DefaultMovementSet/Modes/AsyncWalkingMode.cpp
--- a/Engine/Plugins/Experimental/Mover/Source/Mover/Private/DefaultMovementSet/Modes/AsyncWalkingMode.cpp
+++ b/Engine/Plugins/Experimental/Mover/Source/Mover/Private/DefaultMovementSet/Modes/AsyncWalkingMode.cpp
@@ -236,7 +236,9 @@ void UAsyncWalkingMode::SimulationTick_Implementation(const FSimulationTickParam
 					FOptionalFloorCheckResult StepUpFloorResult;	// passed to sub-operations, so we can use their final floor results if they did a test
 					FVector PostStepUpLocation; // Valid if step-up succeeded
 
-					if (UGroundMovementUtils::TestMoveToStepOver(Params.MovingComps, DownwardDir, CommonLegacySettings->MaxStepHeight, CommonLegacySettings->MaxWalkSlopeCosine, CommonLegacySettings->bUseFlatBaseForFloorChecks, CommonLegacySettings->FloorSweepDistance, OrigMoveDelta * (1.f - PercentTimeAppliedSoFar), RotationInProgress, MoveHitResult, CurrentFloor, false, &StepUpFloorResult, PostStepUpLocation, MoveRecord))
+					float StepUpMovePercentage = UGroundMovementUtils::TestMoveToStepOver(Params.MovingComps, DownwardDir, CommonLegacySettings->MaxStepHeight, CommonLegacySettings->MaxWalkSlopeCosine, CommonLegacySettings->bUseFlatBaseForFloorChecks, CommonLegacySettings->FloorSweepDistance, OrigMoveDelta * (1.f - PercentTimeAppliedSoFar), RotationInProgress, MoveHitResult, CurrentFloor, false, &StepUpFloorResult, PostStepUpLocation, MoveRecord);
+					PercentTimeAppliedSoFar = FMath::Clamp(PercentTimeAppliedSoFar + (StepUpMovePercentage * (1.f - PercentTimeAppliedSoFar)), 0.f, 1.f);
+					if (StepUpMovePercentage > 0)
 					{
 						LocationInProgress = PostStepUpLocation;
 						RotationInProgress = TargetRotation;
diff --git a/Engine/Plugins/Experimental/Mover/Source/Mover/Private/DefaultMovementSet/Modes/WalkingMode.cpp b/Engine/Plugins/Experimental/Mover/Source/Mover/Private/DefaultMovementSet/Modes/WalkingMode.cpp
--- a/Engine/Plugins/Experimental/Mover/Source/Mover/Private/DefaultMovementSet/Modes/WalkingMode.cpp
+++ b/Engine/Plugins/Experimental/Mover/Source/Mover/Private/DefaultMovementSet/Modes/WalkingMode.cpp
@@ -212,7 +212,10 @@ void UWalkingMode::SimulationTick_Implementation(const FSimulationTickParams& Pa
 					const FVector PreStepUpLocation = UpdatedComponent->GetComponentLocation();
 					const FVector DownwardDir = -MoverComp->GetUpDirection();
 
-					if (!UGroundMovementUtils::TryMoveToStepUp(Params.MovingComps, DownwardDir, CommonLegacySettings->MaxStepHeight, CommonLegacySettings->MaxWalkSlopeCosine, CommonLegacySettings->bUseFlatBaseForFloorChecks, CommonLegacySettings->FloorSweepDistance, OrigMoveDelta * (1.f - PercentTimeAppliedSoFar), MoveHitResult, CurrentFloor, false, &StepUpFloorResult, MoveRecord))
+					
+					float StepUpMovePercentage = UGroundMovementUtils::TryMoveToStepUp(Params.MovingComps, DownwardDir, CommonLegacySettings->MaxStepHeight, CommonLegacySettings->MaxWalkSlopeCosine, CommonLegacySettings->bUseFlatBaseForFloorChecks, CommonLegacySettings->FloorSweepDistance, OrigMoveDelta * (1.f - PercentTimeAppliedSoFar), MoveHitResult, CurrentFloor, false, &StepUpFloorResult, MoveRecord);
+					PercentTimeAppliedSoFar = FMath::Clamp(PercentTimeAppliedSoFar + (StepUpMovePercentage * (1.f - PercentTimeAppliedSoFar)), 0.f, 1.f);
+					if (StepUpMovePercentage <= 0)
 					{
 						FMoverOnImpactParams ImpactParams(DefaultModeNames::Walking, MoveHitResult, OrigMoveDelta);
 						MoverComp->HandleImpact(ImpactParams);
diff --git a/Engine/Plugins/Experimental/Mover/Source/Mover/Private/MoveLibrary/GroundMovementUtils.cpp b/Engine/Plugins/Experimental/Mover/Source/Mover/Private/MoveLibrary/GroundMovementUtils.cpp
--- a/Engine/Plugins/Experimental/Mover/Source/Mover/Private/MoveLibrary/GroundMovementUtils.cpp
+++ b/Engine/Plugins/Experimental/Mover/Source/Mover/Private/MoveLibrary/GroundMovementUtils.cpp
@@ -55,11 +55,12 @@ static const FName StepFwdSubstepName = "StepFwd";
 static const FName StepDownSubstepName = "StepDown";
 static const FName SlideSubstepName = "SlideFromStep";
 
-bool UGroundMovementUtils::TryMoveToStepUp(const FMovingComponentSet& MovingComps, const FVector& GravDir, float MaxStepHeight, float MaxWalkSlopeCosine, bool bUseFlatBaseForFloorChecks, float FloorSweepDistance, const FVector& MoveDelta, const FHitResult& MoveHitResult, const FFloorCheckResult& CurrentFloor, bool bIsFalling, FOptionalFloorCheckResult* OutFloorTestResult, FMovementRecord& MoveRecord)
+float UGroundMovementUtils::TryMoveToStepUp(const FMovingComponentSet& MovingComps, const FVector& GravDir, float MaxStepHeight, float MaxWalkSlopeCosine, bool bUseFlatBaseForFloorChecks, float FloorSweepDistance, const FVector& MoveDelta, const FHitResult& MoveHitResult, const FFloorCheckResult& CurrentFloor, bool bIsFalling, FOptionalFloorCheckResult* OutFloorTestResult, FMovementRecord& MoveRecord)
 {
+	float PercentMovementUsed = 0;
 	if (!CanStepUpOnHitSurface(MoveHitResult) || MaxStepHeight <= 0.f)
 	{
-		return false;
+		return 0;
 	}
 
 	float PawnRadius(0.0f), PawnHalfHeight(0.0f);
@@ -82,23 +83,20 @@ bool UGroundMovementUtils::TryMoveToStepUp(const FMovingComponentSet& MovingComp
 
 	FVector LastComponentLocation = OldLocation;
 
-	
-
-
 	// Don't bother stepping up if top of shape is hitting something.
 	const float InitialImpactDot = MoveHitResult.ImpactPoint.Dot(UpDirection);
 	const float OldLocationDot = OldLocation.Dot(UpDirection);
 	if (InitialImpactDot > OldLocationDot + (PawnHalfHeight - PawnRadius))
 	{
 		UE_LOG(LogMover, VeryVerbose, TEXT("Not stepping up due to top of shape hitting something"));
-		return false;
+		return 0;
 	}
 
 	// TODO: We should rely on movement plane normal, rather than gravity direction
 	if (GravDir.IsZero())
 	{
 		UE_LOG(LogMover, VeryVerbose, TEXT("Not stepping up because there's no gravity"));
-		return false;
+		return 0;
 	}
 
 	// Gravity should be a normalized direction
@@ -135,7 +133,7 @@ bool UGroundMovementUtils::TryMoveToStepUp(const FMovingComponentSet& MovingComp
 	if (InitialImpactDot <= PawnInitialFloorBaseDot)
 	{
 		UE_LOG(LogMover, VeryVerbose, TEXT("Not stepping up because the impact is below us"));
-		return false;
+		return 0;
 	}
 
 	// Scope our movement updates, and do not apply them until all intermediate moves are completed.
@@ -156,7 +154,7 @@ bool UGroundMovementUtils::TryMoveToStepUp(const FMovingComponentSet& MovingComp
 		// Undo movement
 		UE_LOG(LogMover, VeryVerbose, TEXT("Reverting step-up attempt because we started in a penetrating state"));
 		ScopedStepUpMovement.RevertMove();
-		return false;
+		return 0;
 	}
 
 	// Cache upwards substep
@@ -170,6 +168,8 @@ bool UGroundMovementUtils::TryMoveToStepUp(const FMovingComponentSet& MovingComp
 	UE_LOG(LogMover, VeryVerbose, TEXT("TryMoveToStepUp Fwd: %s (role %i) MoveDelta=%s DidMove=%i"),
 		*GetNameSafe(RootComponent->GetOwner()), RootComponent->GetOwnerRole(), *MoveDelta.ToCompactString(), bDidStepFwd);
 
+	PercentMovementUsed += StepFwdHit.Time;
+	
 	// Check result of forward movement
 	if (StepFwdHit.bBlockingHit)
 	{
@@ -178,7 +178,7 @@ bool UGroundMovementUtils::TryMoveToStepUp(const FMovingComponentSet& MovingComp
 			// Undo movement
 			UE_LOG(LogMover, VeryVerbose, TEXT("Reverting step-fwd attempt during step-up, because we started in a penetrating state"));
 			ScopedStepUpMovement.RevertMove();
-			return false;
+			return 0;
 		}
 
 		UMoverComponent* MoverComp = MovingComps.MoverComponent.Get();
@@ -208,7 +208,7 @@ bool UGroundMovementUtils::TryMoveToStepUp(const FMovingComponentSet& MovingComp
 				MoveRecord.Append(Substep);
 			}
 
-			return true;
+			return PercentMovementUsed;
 		}
 
 		// Cache forwards substep before the slide attempt
@@ -225,11 +225,13 @@ bool UGroundMovementUtils::TryMoveToStepUp(const FMovingComponentSet& MovingComp
 		LastComponentLocation = RootComponent->GetComponentLocation();
 		MoveRecord.UnlockRelevancy();
 
+		PercentMovementUsed += ForwardSlideAmount * (1 - PercentMovementUsed);
+		
 		if (bIsFalling)
 		{
 			UE_LOG(LogMover, VeryVerbose, TEXT("Reverting step-fwd attempt during step-up, because we could not adjust without falling"));
 			ScopedStepUpMovement.RevertMove();
-			return false;
+			return 0;
 		}
 
 		// If both the forward hit and the deflection got us nowhere, there is no point in this step up.
@@ -237,7 +239,7 @@ bool UGroundMovementUtils::TryMoveToStepUp(const FMovingComponentSet& MovingComp
 		{
 			UE_LOG(LogMover, VeryVerbose, TEXT("Reverting step-fwd attempt during step-up, because no movement differences occurred"));
 			ScopedStepUpMovement.RevertMove();
-			return false;
+			return 0;
 		}
 	}
 	else
@@ -245,73 +247,75 @@ bool UGroundMovementUtils::TryMoveToStepUp(const FMovingComponentSet& MovingComp
 		// Our forward move attempt was unobstructed - cache it
 		QueuedSubsteps.Add(FMovementSubstep(StepFwdSubstepName, RootComponent->GetComponentLocation() - LastComponentLocation, true));
 		LastComponentLocation = RootComponent->GetComponentLocation();
+		PercentMovementUsed = 1.f;
 	}
 
 
 	// Step down
+	FHitResult StepDownHit(1.f);
 	const FVector StepDownAdjustment = GravDir * StepTravelDownHeight;
-	const bool bDidStepDown = UMovementUtils::TryMoveUpdatedComponent_Internal(MovingComps, StepDownAdjustment, RootComponent->GetComponentQuat(), true, MOVECOMP_NoFlags, &StepFwdHit, ETeleportType::None);
+	const bool bDidStepDown = UMovementUtils::TryMoveUpdatedComponent_Internal(MovingComps, StepDownAdjustment, RootComponent->GetComponentQuat(), true, MOVECOMP_NoFlags, &StepDownHit, ETeleportType::None);
 
 	UE_LOG(LogMover, VeryVerbose, TEXT("TryMoveToStepUp Down: %s (role %i) StepDownAdjustment=%s DidMove=%i"),
 		*GetNameSafe(RootComponent->GetOwner()), RootComponent->GetOwnerRole(), *StepDownAdjustment.ToCompactString(), bDidStepDown);
 
 
 	// If step down was initially penetrating abort the step up
-	if (StepFwdHit.bStartPenetrating)
+	if (StepDownHit.bStartPenetrating)
 	{
 		UE_LOG(LogMover, VeryVerbose, TEXT("Reverting step-down attempt during step-up/step-fwd, because we started in a penetrating state"));
 		ScopedStepUpMovement.RevertMove();
-		return false;
+		return 0;
 	}
 
 	FOptionalFloorCheckResult StepDownResult;
-	if (StepFwdHit.IsValidBlockingHit())
+	if (StepDownHit.IsValidBlockingHit())
 	{
 		// See if this step sequence would have allowed us to travel higher than our max step height allows.
-		const float DeltaDot = (StepFwdHit.ImpactPoint.Dot(UpDirection)) - PawnFloorPointDot;
+		const float DeltaDot = (StepDownHit.ImpactPoint.Dot(UpDirection)) - PawnFloorPointDot;
 		if (DeltaDot > MaxStepHeight)
 		{
-			UE_LOG(LogMover, VeryVerbose, TEXT("Reject step-down attempt during step-up/step-fwd, because it made us travel too high (too high Height %.3f) up from floor base %f to %f"), DeltaDot, PawnInitialFloorBaseDot, StepFwdHit.ImpactPoint.Z);
+			UE_LOG(LogMover, VeryVerbose, TEXT("Reject step-down attempt during step-up/step-fwd, because it made us travel too high (too high Height %.3f) up from floor base %f to %f"), DeltaDot, PawnInitialFloorBaseDot, StepDownHit.ImpactPoint.Z);
 			ScopedStepUpMovement.RevertMove();
-			return false;
+			return 0;
 		}
 
 		// Reject unwalkable surface normals here.
-		if (!UFloorQueryUtils::IsHitSurfaceWalkable(StepFwdHit, UpDirection, MaxWalkSlopeCosine))
+		if (!UFloorQueryUtils::IsHitSurfaceWalkable(StepDownHit, UpDirection, MaxWalkSlopeCosine))
 		{
 			// Reject if normal opposes movement direction
-			const bool bNormalTowardsMe = (MoveDelta | StepFwdHit.ImpactNormal) < 0.f;
+			const bool bNormalTowardsMe = (MoveDelta | StepDownHit.ImpactNormal) < 0.f;
 			if (bNormalTowardsMe)
 			{
-				UE_LOG(LogMover, VeryVerbose, TEXT("Reject step-down attempt during step-up/step-fwd, due to unwalkable normal %s opposed to movement"), *StepFwdHit.ImpactNormal.ToString());
+				UE_LOG(LogMover, VeryVerbose, TEXT("Reject step-down attempt during step-up/step-fwd, due to unwalkable normal %s opposed to movement"), *StepDownHit.ImpactNormal.ToString());
 				ScopedStepUpMovement.RevertMove();
-				return false;
+				return 0;
 			}
 
 			// Also reject if we would end up being higher than our starting location by stepping down.
 			// It's fine to step down onto an unwalkable normal below us, we will just slide off. Rejecting those moves would prevent us from being able to walk off the edge.
-			if (StepFwdHit.Location.Dot(UpDirection) > OldLocationDot)
+			if (StepDownHit.Location.Dot(UpDirection) > OldLocationDot)
 			{
-				UE_LOG(LogMover, VeryVerbose, TEXT("Reject step-down attempt during step-up/step-fwd, due to unwalkable normal %s above old position)"), *StepFwdHit.ImpactNormal.ToString());
+				UE_LOG(LogMover, VeryVerbose, TEXT("Reject step-down attempt during step-up/step-fwd, due to unwalkable normal %s above old position)"), *StepDownHit.ImpactNormal.ToString());
 				ScopedStepUpMovement.RevertMove();
-				return false;
+				return 0;
 			}
 		}
 
 		// Reject moves where the downward sweep hit something very close to the edge of the shape. This maintains consistency with FindFloor as well.
-		if (!UFloorQueryUtils::IsWithinEdgeTolerance(StepFwdHit.Location, StepFwdHit.ImpactPoint, PawnRadius, UpDirection))
+		if (!UFloorQueryUtils::IsWithinEdgeTolerance(StepDownHit.Location, StepDownHit.ImpactPoint, PawnRadius, UpDirection))
 		{
 			UE_LOG(LogMover, VeryVerbose, TEXT("Reject step-down attempt during step-up/step-fwd, due to being outside edge tolerance"));
 			ScopedStepUpMovement.RevertMove();
-			return false;
+			return 0;
 		}
 
 		// Don't step up onto invalid surfaces if traveling higher.
-		if (DeltaDot > 0.f && !CanStepUpOnHitSurface(StepFwdHit))
+		if (DeltaDot > 0.f && !CanStepUpOnHitSurface(StepDownHit))
 		{
 			UE_LOG(LogMover, VeryVerbose, TEXT("Reject step-down attempt during step-up/step-fwd, due to being up onto surface with !CanStepUpOnHitSurface")); 
 			ScopedStepUpMovement.RevertMove();
-			return false;
+			return 0;
 		}
 
 		// See if we can validate the floor as a result of this step down. In almost all cases this should succeed, and we can avoid computing the floor outside this method.
@@ -322,7 +326,7 @@ bool UGroundMovementUtils::TryMoveToStepUp(const FMovingComponentSet& MovingComp
 
 			// Reject unwalkable normals if we end up higher than our initial height.
 			// It's fine to walk down onto an unwalkable surface, don't reject those moves.
-			if (StepFwdHit.Location.Dot(UpDirection) > OldLocationDot)
+			if (StepDownHit.Location.Dot(UpDirection) > OldLocationDot)
 			{
 				// We should reject the floor result if we are trying to step up an actual step where we are not able to perch (this is rare).
 				// In those cases we should instead abort the step up and try to slide along the stair.
@@ -331,7 +335,7 @@ bool UGroundMovementUtils::TryMoveToStepUp(const FMovingComponentSet& MovingComp
 				{
 					UE_LOG(LogMover, VeryVerbose, TEXT("Reject step-down attempt during step-up/step-fwd, due to it being an unperchable step")); 
 					ScopedStepUpMovement.RevertMove();
-					return false;
+					return 0;
 				}
 			}
 
@@ -358,7 +362,7 @@ bool UGroundMovementUtils::TryMoveToStepUp(const FMovingComponentSet& MovingComp
 		MoveRecord.Append(Substep);
 	}
 
-	return true;
+	return PercentMovementUsed;
 
 }
 
@@ -769,13 +773,15 @@ FVector UGroundMovementUtils::TestMoveToKeepMinHeightAboveFloor(const FMovingCom
 
 
 /* static */
-bool UGroundMovementUtils::TestMoveToStepOver(const FMovingComponentSet& MovingComps, const FVector& GravDir, float MaxStepHeight, float MaxWalkSlopeCosine, bool bUseFlatBaseForFloorChecks, float FloorSweepDistance, const FVector& MoveDelta, const FQuat& Rotation, const FHitResult& MoveHitResult, const FFloorCheckResult& CurrentFloor, bool bIsFalling, FOptionalFloorCheckResult* OutFloorTestResult, FVector& OutFinalLocation, FMovementRecord& InOutMoveRecord)
+float UGroundMovementUtils::TestMoveToStepOver(const FMovingComponentSet& MovingComps, const FVector& GravDir, float MaxStepHeight, float MaxWalkSlopeCosine, bool bUseFlatBaseForFloorChecks, float FloorSweepDistance, const FVector& MoveDelta, const FQuat& Rotation, const FHitResult& MoveHitResult, const FFloorCheckResult& CurrentFloor, bool bIsFalling, FOptionalFloorCheckResult* OutFloorTestResult, FVector& OutFinalLocation, FMovementRecord& InOutMoveRecord)
 {
 	if (!UGroundMovementUtils::CanStepUpOnHitSurface(MoveHitResult) || MaxStepHeight <= 0.f)
 	{
-		return false;
+		return 0;
 	}
 
+	float PercentMovementUsed = 0;
+	
 	TArray<FMovementSubstep> QueuedSubsteps;	// keeping track of substeps before committing, because some moves can be backed out
 
 	FVector UpDirection = MovingComps.MoverComponent->GetUpDirection();
@@ -801,14 +807,14 @@ bool UGroundMovementUtils::TestMoveToStepOver(const FMovingComponentSet& MovingC
 	if (InitialImpactDot > OldLocationDot + (PawnHalfHeight - PawnRadius))
 	{
 		UE_LOG(LogMover, VeryVerbose, TEXT("Not stepping up due to top of shape hitting something"));
-		return false;
+		return 0;
 	}
 
 	// TODO: We should rely on movement plane normal, rather than gravity direction
 	if (GravDir.IsZero())
 	{
 		UE_LOG(LogMover, VeryVerbose, TEXT("Not stepping up because there's no gravity"));
-		return false;
+		return 0;
 	}
 
 	// Gravity should be a normalized direction
@@ -844,7 +850,7 @@ bool UGroundMovementUtils::TestMoveToStepOver(const FMovingComponentSet& MovingC
 	if (InitialImpactDot <= PawnInitialFloorBaseDot)
 	{
 		UE_LOG(LogMover, VeryVerbose, TEXT("Not stepping up because the impact is below us"));
-		return false;
+		return 0;
 	}
 
 	FMoverCollisionParams CollisionParams(MovingComps.UpdatedComponent.Get());
@@ -862,7 +868,7 @@ bool UGroundMovementUtils::TestMoveToStepOver(const FMovingComponentSet& MovingC
 	{
 		// Undo movement
 		UE_LOG(LogMover, VeryVerbose, TEXT("Reverting step-up attempt because we started in a penetrating state"));
-		return false;
+		return 0;
 	}
 
 	FVector UpStepDelta = (SweepUpHit.TraceStart + ((SweepUpHit.TraceEnd - SweepUpHit.TraceStart) * SweepUpHit.Time)) - LocationInProgress;
@@ -887,7 +893,7 @@ bool UGroundMovementUtils::TestMoveToStepOver(const FMovingComponentSet& MovingC
 		{
 			// Undo movement
 			UE_LOG(LogMover, VeryVerbose, TEXT("Reverting step-fwd attempt during step-up, because we started in a penetrating state"));
-			return false;
+			return 0;
 		}
 
 		UMoverComponent* MoverComp = MovingComps.MoverComponent.Get();
@@ -911,6 +917,7 @@ bool UGroundMovementUtils::TestMoveToStepOver(const FMovingComponentSet& MovingC
 		QueuedSubsteps.Add(FMovementSubstep(StepFwdSubstepName, FwdStepDelta, true));
 		LocationInProgress += FwdStepDelta;
 
+		PercentMovementUsed = StepFwdHit.Time;
 
 		// If falling, we only need to try to reach up and fwd. We do not need to continue searching downward.
 		if (bIsFalling)
@@ -922,7 +929,7 @@ bool UGroundMovementUtils::TestMoveToStepOver(const FMovingComponentSet& MovingC
 			}
 
 			OutFinalLocation = LocationInProgress;
-			return true;
+			return PercentMovementUsed;
 		}
 
 
@@ -939,19 +946,21 @@ bool UGroundMovementUtils::TestMoveToStepOver(const FMovingComponentSet& MovingC
 		QueuedSubsteps.Add(FMovementSubstep(SlideSubstepName, FwdSlideDelta, true));
 		LocationInProgress += FwdSlideDelta;
 
+		PercentMovementUsed += ForwardSlideAmount * (1 - PercentMovementUsed);
+		
 		InOutMoveRecord.UnlockRelevancy();
 
 		if (bIsFalling)
 		{
 			UE_LOG(LogMover, VeryVerbose, TEXT("Reverting step-fwd attempt during step-up, because we could not adjust without falling"));
-			return false;
+			return 0;
 		}
 
 		// If both the forward hit and the deflection got us nowhere, there is no point in this step up.
 		if (ForwardHitTime == 0.f && ForwardSlideAmount == 0.f)
 		{
 			UE_LOG(LogMover, VeryVerbose, TEXT("Reverting step-fwd attempt during step-up, because no movement differences occurred"));
-			return false;
+			return 0;
 		}
 	}
 	else
@@ -959,6 +968,7 @@ bool UGroundMovementUtils::TestMoveToStepOver(const FMovingComponentSet& MovingC
 		// Our forward move attempt was unobstructed - cache it
 		QueuedSubsteps.Add(FMovementSubstep(StepFwdSubstepName, FwdStepDelta, true));
 		LocationInProgress += FwdStepDelta;
+		PercentMovementUsed = 1.f;
 	}
 
 
@@ -976,7 +986,7 @@ bool UGroundMovementUtils::TestMoveToStepOver(const FMovingComponentSet& MovingC
 	if (StepFwdHit.bStartPenetrating)
 	{
 		UE_LOG(LogMover, VeryVerbose, TEXT("Reverting step-down attempt during step-up/step-fwd, because we started in a penetrating state"));
-		return false;
+		return 0;
 	}
 
 	FOptionalFloorCheckResult StepDownResult;
@@ -987,7 +997,7 @@ bool UGroundMovementUtils::TestMoveToStepOver(const FMovingComponentSet& MovingC
 		if (DeltaDot > MaxStepHeight)
 		{
 			UE_LOG(LogMover, VeryVerbose, TEXT("Reject step-down attempt during step-up/step-fwd, because it made us travel too high (too high Height %.3f) up from floor base %f to %f"), DeltaDot, PawnInitialFloorBaseDot, StepFwdHit.ImpactPoint.Z);
-			return false;
+			return 0;
 		}
 
 		// Reject unwalkable surface normals here.
@@ -998,7 +1008,7 @@ bool UGroundMovementUtils::TestMoveToStepOver(const FMovingComponentSet& MovingC
 			if (bNormalTowardsMe)
 			{
 				UE_LOG(LogMover, VeryVerbose, TEXT("Reject step-down attempt during step-up/step-fwd, due to unwalkable normal %s opposed to movement"), *StepFwdHit.ImpactNormal.ToString());
-				return false;
+				return 0;
 			}
 
 			// Also reject if we would end up being higher than our starting location by stepping down.
@@ -1006,7 +1016,7 @@ bool UGroundMovementUtils::TestMoveToStepOver(const FMovingComponentSet& MovingC
 			if (StepFwdHit.Location.Dot(UpDirection) > OldLocationDot)
 			{
 				UE_LOG(LogMover, VeryVerbose, TEXT("Reject step-down attempt during step-up/step-fwd, due to unwalkable normal %s above old position)"), *StepFwdHit.ImpactNormal.ToString());
-				return false;
+				return 0;
 			}
 		}
 
@@ -1014,14 +1024,14 @@ bool UGroundMovementUtils::TestMoveToStepOver(const FMovingComponentSet& MovingC
 		if (!UFloorQueryUtils::IsWithinEdgeTolerance(StepFwdHit.Location, StepFwdHit.ImpactPoint, PawnRadius, UpDirection))
 		{
 			UE_LOG(LogMover, VeryVerbose, TEXT("Reject step-down attempt during step-up/step-fwd, due to being outside edge tolerance"));
-			return false;
+			return 0;
 		}
 
 		// Don't step up onto invalid surfaces if traveling higher.
 		if (DeltaDot > 0.f && !UGroundMovementUtils::CanStepUpOnHitSurface(StepFwdHit))
 		{
 			UE_LOG(LogMover, VeryVerbose, TEXT("Reject step-down attempt during step-up/step-fwd, due to being up onto surface with !CanStepUpOnHitSurface"));
-			return false;
+			return 0;
 		}
 
 		// See if we can validate the floor as a result of this step down. In almost all cases this should succeed, and we can avoid computing the floor outside this method.
@@ -1040,7 +1050,7 @@ bool UGroundMovementUtils::TestMoveToStepOver(const FMovingComponentSet& MovingC
 				if (!StepDownResult.FloorTestResult.bBlockingHit && StepSideDot < MAX_STEP_SIDE_DOT)
 				{
 					UE_LOG(LogMover, VeryVerbose, TEXT("Reject step-down attempt during step-up/step-fwd, due to it being an unperchable step"));
-					return false;
+					return 0;
 				}
 			}
 
@@ -1068,5 +1078,5 @@ bool UGroundMovementUtils::TestMoveToStepOver(const FMovingComponentSet& MovingC
 	}
 
 	OutFinalLocation = LocationInProgress;
-	return true;
+	return PercentMovementUsed;
 }
diff --git a/Engine/Plugins/Experimental/Mover/Source/Mover/Public/MoveLibrary/GroundMovementUtils.h b/Engine/Plugins/Experimental/Mover/Source/Mover/Public/MoveLibrary/GroundMovementUtils.h
--- a/Engine/Plugins/Experimental/Mover/Source/Mover/Public/MoveLibrary/GroundMovementUtils.h
+++ b/Engine/Plugins/Experimental/Mover/Source/Mover/Public/MoveLibrary/GroundMovementUtils.h
@@ -85,8 +85,8 @@ public:
 	static UE_API FProposedMove ComputeControlledGroundMove(const FGroundMoveParams& InParams);
 	
 	// TODO: Refactor this API for fewer parameters
-	/** Move up steps or slope. Does nothing and returns false if hit surface is invalid for step-up use */
-	static UE_API bool TryMoveToStepUp(const FMovingComponentSet& MovingComps, const FVector& GravDir, float MaxStepHeight, float MaxWalkSlopeCosine, bool bUseFlatBaseForFloorChecks, float FloorSweepDistance, const FVector& MoveDelta, const FHitResult& Hit, const FFloorCheckResult& CurrentFloor, bool bIsFalling, FOptionalFloorCheckResult* OutStepDownResult, FMovementRecord& MoveRecord);
+	/** Move up steps or slope. Returns the percent of time applied, with 0.0 meaning no movement occurred and hit surface is invalid for step-up use. */
+	static UE_API float TryMoveToStepUp(const FMovingComponentSet& MovingComps, const FVector& GravDir, float MaxStepHeight, float MaxWalkSlopeCosine, bool bUseFlatBaseForFloorChecks, float FloorSweepDistance, const FVector& MoveDelta, const FHitResult& Hit, const FFloorCheckResult& CurrentFloor, bool bIsFalling, FOptionalFloorCheckResult* OutStepDownResult, FMovementRecord& MoveRecord);
 
 	/** Moves vertically to stay within range of the walkable floor. Does nothing and returns false if floor is unwalkable or if already in range. */ 
 	static UE_API bool TryMoveToAdjustHeightAboveFloor(const FMovingComponentSet& MovingComps, FFloorCheckResult& CurrentFloor, float MaxWalkSlopeCosine, FMovementRecord& MoveRecord);
@@ -138,9 +138,9 @@ public:
 	 * Sets OutFloorTestResult with final floor info (optional)
 	 * Sets OutFinalLocation with final component location
      * Appends to InOutMoveRecord with any movement substeps
-	 * Returns whether any movement was possible.
+	 * Returns the percent of time applied, with 0.0 meaning no movement would occur.
 	 */
-	static UE_API bool TestMoveToStepOver(const FMovingComponentSet& MovingComps, const FVector& GravDir, float MaxStepHeight, float MaxWalkSlopeCosine, bool bUseFlatBaseForFloorChecks, float FloorSweepDistance, const FVector& MoveDelta, const FQuat& Rotation, const FHitResult& MoveHitResult, const FFloorCheckResult& CurrentFloor, bool bIsFalling, FOptionalFloorCheckResult* OutFloorTestResult, FVector& OutFinalLocation, FMovementRecord& InOutMoveRecord);
+	static UE_API float TestMoveToStepOver(const FMovingComponentSet& MovingComps, const FVector& GravDir, float MaxStepHeight, float MaxWalkSlopeCosine, bool bUseFlatBaseForFloorChecks, float FloorSweepDistance, const FVector& MoveDelta, const FQuat& Rotation, const FHitResult& MoveHitResult, const FFloorCheckResult& CurrentFloor, bool bIsFalling, FOptionalFloorCheckResult* OutFloorTestResult, FVector& OutFinalLocation, FMovementRecord& InOutMoveRecord);
 
 
 };
