Mover: Remove unnecessary falling check that will never get hit in movement utils

#Synced-CL 51093404

diff --git a/Engine/Plugins/Experimental/Mover/Source/Mover/Private/MoveLibrary/GroundMovementUtils.cpp b/Engine/Plugins/Experimental/Mover/Source/Mover/Private/MoveLibrary/GroundMovementUtils.cpp
--- a/Engine/Plugins/Experimental/Mover/Source/Mover/Private/MoveLibrary/GroundMovementUtils.cpp
+++ b/Engine/Plugins/Experimental/Mover/Source/Mover/Private/MoveLibrary/GroundMovementUtils.cpp
@@ -227,13 +227,6 @@ float UGroundMovementUtils::TryMoveToStepUp(const FMovingComponentSet& MovingCom
 
 		PercentMovementUsed += ForwardSlideAmount * (1 - PercentMovementUsed);
 		
-		if (bIsFalling)
-		{
-			UE_LOG(LogMover, VeryVerbose, TEXT("Reverting step-fwd attempt during step-up, because we could not adjust without falling"));
-			ScopedStepUpMovement.RevertMove();
-			return 0;
-		}
-
 		// If both the forward hit and the deflection got us nowhere, there is no point in this step up.
 		if (ForwardHitTime == 0.f && ForwardSlideAmount == 0.f)
 		{
