Mover: fixed query that referenced the wrong modifiers array

#rnx

#Synced-CL 51257655

diff --git a/Engine/Plugins/Experimental/Mover/Source/Mover/Private/MovementModifier.cpp b/Engine/Plugins/Experimental/Mover/Source/Mover/Private/MovementModifier.cpp
--- a/Engine/Plugins/Experimental/Mover/Source/Mover/Private/MovementModifier.cpp
+++ b/Engine/Plugins/Experimental/Mover/Source/Mover/Private/MovementModifier.cpp
@@ -346,7 +346,7 @@ TArray<TSharedPtr<FMovementModifierBase>>::TConstIterator FMovementModifierGroup
 
 TArray<TSharedPtr<FMovementModifierBase>>::TConstIterator FMovementModifierGroup::GetQueuedModifiersIterator() const
 {
-	return ActiveModifiers.CreateConstIterator();
+	return QueuedModifiers.CreateConstIterator();
 }
 
 void FMovementModifierGroup::FlushModifierArrays(UMoverComponent* MoverComp, const FMoverTimeStep& TimeStep, const FMoverSyncState& SyncState, const FMoverAuxStateContext& AuxState)
