Mover: updating UE_LOG to UE_LOGF, per updated best practices guidance (better perf)

#rnx

#Synced-CL 51296837

diff --git a/Engine/Plugins/Experimental/Mover/Source/Mover/Private/Backends/MoverNetworkPhysicsLiaison.cpp b/Engine/Plugins/Experimental/Mover/Source/Mover/Private/Backends/MoverNetworkPhysicsLiaison.cpp
--- a/Engine/Plugins/Experimental/Mover/Source/Mover/Private/Backends/MoverNetworkPhysicsLiaison.cpp
+++ b/Engine/Plugins/Experimental/Mover/Source/Mover/Private/Backends/MoverNetworkPhysicsLiaison.cpp
@@ -319,7 +319,7 @@ bool UMoverNetworkPhysicsLiaisonComponent::CanSimulate_Internal(const FPhysicsMo
 		if (!Cast<IPhysicsCharacterMovementModeInterface>(GetMoverComponent().MovementModes.FindChecked(Input.SyncState.MovementMode)))
 		{
 			ensureMsgf(false, TEXT("Attempting to run non-character physics movement mode %s in physics mover update. Only modes that implement IPhysicsCharacterMovementModeInterface can be used with the CharacterPhysicsLiaison."), *Input.SyncState.MovementMode.ToString());
-			UE_LOG(LogMover, Verbose, TEXT("Attempting to run non-character physics movement mode %s in physics mover update. Only modes that implement IPhysicsCharacterMovementModeInterface can be used with the CharacterPhysicsLiaison."),
+			UE_LOGF(LogMover, Verbose, "Attempting to run non-character physics movement mode %ls in physics mover update. Only modes that implement IPhysicsCharacterMovementModeInterface can be used with the CharacterPhysicsLiaison.",
 				*Input.SyncState.MovementMode.ToString());
 			return false;
 		}
diff --git a/Engine/Plugins/Experimental/Mover/Source/Mover/Private/Backends/MoverNetworkPhysicsLiaisonBase.cpp b/Engine/Plugins/Experimental/Mover/Source/Mover/Private/Backends/MoverNetworkPhysicsLiaisonBase.cpp
--- a/Engine/Plugins/Experimental/Mover/Source/Mover/Private/Backends/MoverNetworkPhysicsLiaisonBase.cpp
+++ b/Engine/Plugins/Experimental/Mover/Source/Mover/Private/Backends/MoverNetworkPhysicsLiaisonBase.cpp
@@ -411,13 +411,13 @@ bool UMoverNetworkPhysicsLiaisonComponentBase::CanCreatePhysics() const
 	const UMoverComponent& MoverComp = GetMoverComponent();
 	if (!IsValid(MoverComp.UpdatedComponent))
 	{
-		UE_LOG(LogMover, Warning, TEXT("Can't create physics %s (%s). UpdatedComponent is not set."), *ActorName, *GetPathName());
+		UE_LOGF(LogMover, Warning, "Can't create physics %ls (%ls). UpdatedComponent is not set.", *ActorName, *GetPathName());
 		return false;
 	}
 
 	if (!IsValid(MoverComp.UpdatedCompAsPrimitive))
 	{
-		UE_LOG(LogMover, Warning, TEXT("Can't create physics %s (%s). UpdatedComponent is not a PrimitiveComponent."), *ActorName, *GetPathName());
+		UE_LOGF(LogMover, Warning, "Can't create physics %ls (%ls). UpdatedComponent is not a PrimitiveComponent.", *ActorName, *GetPathName());
 		return false;
 	}
 
diff --git a/Engine/Plugins/Experimental/Mover/Source/Mover/Private/Backends/MoverStandaloneLiaison.cpp b/Engine/Plugins/Experimental/Mover/Source/Mover/Private/Backends/MoverStandaloneLiaison.cpp
--- a/Engine/Plugins/Experimental/Mover/Source/Mover/Private/Backends/MoverStandaloneLiaison.cpp
+++ b/Engine/Plugins/Experimental/Mover/Source/Mover/Private/Backends/MoverStandaloneLiaison.cpp
@@ -117,7 +117,7 @@ bool UMoverStandaloneLiaisonComponent::WritePendingSyncState(const FMoverSyncSta
 {
 	if (bIsInApplySimulationState)
 	{
-		UE_LOG(LogMover, Log, TEXT("Attempted to WritePendingSyncState during ApplySimulationState. Skipping. (%s)"), *GetNameSafe(GetOwner()));
+		UE_LOGF(LogMover, Log, "Attempted to WritePendingSyncState during ApplySimulationState. Skipping. (%ls)", *GetNameSafe(GetOwner()));
 		return false;	// can't write to the sync state while it's being applied. Any changes from this will be overwritten anyway.
 	}
 
diff --git a/Engine/Plugins/Experimental/Mover/Source/Mover/Private/Debug/MoverDebugComponent.cpp b/Engine/Plugins/Experimental/Mover/Source/Mover/Private/Debug/MoverDebugComponent.cpp
--- a/Engine/Plugins/Experimental/Mover/Source/Mover/Private/Debug/MoverDebugComponent.cpp
+++ b/Engine/Plugins/Experimental/Mover/Source/Mover/Private/Debug/MoverDebugComponent.cpp
@@ -278,7 +278,7 @@ void UMoverDebugComponent::DrawCorrections()
 {
 	if (ClientLocations.Num() != CorrectedLocations.Num())
 	{
-		UE_LOG(LogMover, Warning, TEXT("Correction arrays differed in size!"));
+		UE_LOGF(LogMover, Warning, "Correction arrays differed in size!");
 		CorrectedLocations.Empty();
 		ClientLocations.Empty();
 		return;
diff --git a/Engine/Plugins/Experimental/Mover/Source/Mover/Private/DefaultMovementSet/LayeredMoves/RootMotionAttributeLayeredMove.cpp b/Engine/Plugins/Experimental/Mover/Source/Mover/Private/DefaultMovementSet/LayeredMoves/RootMotionAttributeLayeredMove.cpp
--- a/Engine/Plugins/Experimental/Mover/Source/Mover/Private/DefaultMovementSet/LayeredMoves/RootMotionAttributeLayeredMove.cpp
+++ b/Engine/Plugins/Experimental/Mover/Source/Mover/Private/DefaultMovementSet/LayeredMoves/RootMotionAttributeLayeredMove.cpp
@@ -180,7 +180,7 @@ bool FLayeredMove_RootMotionAttribute::GenerateMove(const FMoverTickStartData& S
 
 		if (ExcessiveLinearVelocitySquaredThreshold > 0.f && ExcessiveLinearVelocitySquaredThreshold < OutProposedMove.LinearVelocity.SquaredLength())
 		{
-			UE_LOG(LogMover, Warning, TEXT("RootMotionAttr on %s has excessive speed.  LocalTrans: %s (DT: %.5f) -> XY Vel: %.6f  Z Vel: %.6f"),
+			UE_LOGF(LogMover, Warning, "RootMotionAttr on %ls has excessive speed.  LocalTrans: %ls (DT: %.5f) -> XY Vel: %.6f  Z Vel: %.6f",
 				*GetNameSafe(MoverComp->GetOwner()),
 				*LocalRootMotion.GetTranslation().ToString(), DeltaSeconds, OutProposedMove.LinearVelocity.Size2D(), OutProposedMove.LinearVelocity.Z);
 		}
diff --git a/Engine/Plugins/Experimental/Mover/Source/Mover/Private/DefaultMovementSet/Modes/AsyncNavWalkingMode.cpp b/Engine/Plugins/Experimental/Mover/Source/Mover/Private/DefaultMovementSet/Modes/AsyncNavWalkingMode.cpp
--- a/Engine/Plugins/Experimental/Mover/Source/Mover/Private/DefaultMovementSet/Modes/AsyncNavWalkingMode.cpp
+++ b/Engine/Plugins/Experimental/Mover/Source/Mover/Private/DefaultMovementSet/Modes/AsyncNavWalkingMode.cpp
@@ -239,7 +239,7 @@ void UAsyncNavWalkingMode::SimulationTick_Implementation(const FSimulationTickPa
 	if (bDeltaMoveNearlyZero && bSameNavLocation)
 	{
 		DestNavLocation = CachedNavLocation;
-		UE_LOG(LogMover, VeryVerbose, TEXT("%s using cached navmesh location! (bProjectNavMeshWalking = %d)"), *GetNameSafe(GetMoverComponent()->GetOwner()), bProjectNavMeshWalking);
+		UE_LOGF(LogMover, VeryVerbose, "%ls using cached navmesh location! (bProjectNavMeshWalking = %d)", *GetNameSafe(GetMoverComponent()->GetOwner()), bProjectNavMeshWalking);
 	}
 	else
 	{
@@ -297,7 +297,7 @@ void UAsyncNavWalkingMode::SimulationTick_Implementation(const FSimulationTickPa
 			{
 				default:	// fall through
 				case EOffNavMeshBehavior::SwitchToWalking:
-					UE_LOG(LogMover, Verbose, TEXT("%s could not find valid navigation data at location %s. Switching to walking mode."), *GetNameSafe(MoverComp->GetOwner()), *TargetFeetLocation.ToCompactString());
+					UE_LOGF(LogMover, Verbose, "%ls could not find valid navigation data at location %ls. Switching to walking mode.", *GetNameSafe(MoverComp->GetOwner()), *TargetFeetLocation.ToCompactString());
 					OutputState.MovementEndState.NextModeName = DefaultModeNames::Walking;
 					OutputState.MovementEndState.RemainingMs = Params.TimeStep.StepMs;
 					MoveRecord.SetDeltaSeconds(0.0f);
@@ -311,7 +311,7 @@ void UAsyncNavWalkingMode::SimulationTick_Implementation(const FSimulationTickPa
 					break;
 
 				case EOffNavMeshBehavior::DoNotMove:
-					UE_LOG(LogMover, Verbose, TEXT("%s could not find valid navigation data at location %s. Cannot move."), *GetNameSafe(MoverComp->GetOwner()), *TargetFeetLocation.ToCompactString());
+					UE_LOGF(LogMover, Verbose, "%ls could not find valid navigation data at location %ls. Cannot move.", *GetNameSafe(MoverComp->GetOwner()), *TargetFeetLocation.ToCompactString());
 					// nothing to be done
 					break;
 
@@ -546,7 +546,7 @@ FVector UAsyncNavWalkingMode::ProjectLocationFromNavMesh(float DeltaSeconds, con
 	{
 		if (!bCachedLocationStillValid)
 		{
-			UE_LOG(LogMover, VeryVerbose, TEXT("ProjectLocationFromNavMesh(): %s interval: %.3f [SKIP TRACE]"), *GetNameSafe(GetMoverComponent()->GetOwner()), NavMeshProjectionInterval);
+			UE_LOGF(LogMover, VeryVerbose, "ProjectLocationFromNavMesh(): %ls interval: %.3f [SKIP TRACE]", *GetNameSafe(GetMoverComponent()->GetOwner()), NavMeshProjectionInterval);
 
 			FHitResult HitResult;
 			FindBestNavMeshLocation(TraceStart, TraceEnd, CurrentFeetLocation, TargetNavLocation, HitResult);
@@ -574,7 +574,7 @@ FVector UAsyncNavWalkingMode::ProjectLocationFromNavMesh(float DeltaSeconds, con
 		}
 		else
 		{
-			UE_LOG(LogMover, VeryVerbose, TEXT("ProjectLocationFromNavMesh(): %s interval: %.3f [SKIP TRACE]"), *GetNameSafe(GetMoverComponent()->GetOwner()), NavMeshProjectionInterval);
+			UE_LOGF(LogMover, VeryVerbose, "ProjectLocationFromNavMesh(): %ls interval: %.3f [SKIP TRACE]", *GetNameSafe(GetMoverComponent()->GetOwner()), NavMeshProjectionInterval);
 		}
 
 		// Wrap around to maintain same relative offset to tick time changes.
@@ -638,7 +638,7 @@ void UAsyncNavWalkingMode::OnRegistered(const FName ModeName)
 
 	if (!NavMoverComponent.IsValid())
 	{
-		UE_LOG(LogMover, Warning, TEXT("NavWalkingMode on %s could not find a valid NavMoverComponent and will not function properly."), *GetNameSafe(GetMoverComponent()->GetOwner()));
+		UE_LOGF(LogMover, Warning, "NavWalkingMode on %ls could not find a valid NavMoverComponent and will not function properly.", *GetNameSafe(GetMoverComponent()->GetOwner()));
 	}
 }
 
diff --git a/Engine/Plugins/Experimental/Mover/Source/Mover/Private/DefaultMovementSet/Modes/NavWalkingMode.cpp b/Engine/Plugins/Experimental/Mover/Source/Mover/Private/DefaultMovementSet/Modes/NavWalkingMode.cpp
--- a/Engine/Plugins/Experimental/Mover/Source/Mover/Private/DefaultMovementSet/Modes/NavWalkingMode.cpp
+++ b/Engine/Plugins/Experimental/Mover/Source/Mover/Private/DefaultMovementSet/Modes/NavWalkingMode.cpp
@@ -214,7 +214,7 @@ void UNavWalkingMode::SimulationTick_Implementation(const FSimulationTickParams&
 	if (bDeltaMoveNearlyZero && bSameNavLocation)
 	{
 		DestNavLocation = CachedNavLocation;
-		UE_LOG(LogMover, VeryVerbose, TEXT("%s using cached navmesh location! (bProjectNavMeshWalking = %d)"), *GetNameSafe(GetMoverComponent()->GetOwner()), bProjectNavMeshWalking);
+		UE_LOGF(LogMover, VeryVerbose, "%ls using cached navmesh location! (bProjectNavMeshWalking = %d)", *GetNameSafe(GetMoverComponent()->GetOwner()), bProjectNavMeshWalking);
 	}
 	else
 	{
@@ -273,7 +273,7 @@ void UNavWalkingMode::SimulationTick_Implementation(const FSimulationTickParams&
 			{
 				default:	// fall through
 				case EOffNavMeshBehavior::SwitchToWalking:
-					UE_LOG(LogMover, Verbose, TEXT("%s could not find valid navigation data at location %s. Switching to walking mode."), *GetNameSafe(MoverComp->GetOwner()), *AdjustedDest.ToCompactString());
+					UE_LOGF(LogMover, Verbose, "%ls could not find valid navigation data at location %ls. Switching to walking mode.", *GetNameSafe(MoverComp->GetOwner()), *AdjustedDest.ToCompactString());
 					OutputState.MovementEndState.NextModeName = DefaultModeNames::Walking;
 					OutputState.MovementEndState.RemainingMs = Params.TimeStep.StepMs;
 					MoveRecord.SetDeltaSeconds(0.0f);
@@ -286,7 +286,7 @@ void UNavWalkingMode::SimulationTick_Implementation(const FSimulationTickParams&
 					break;
 
 				case EOffNavMeshBehavior::DoNotMove:
-					UE_LOG(LogMover, Verbose, TEXT("%s could not find valid navigation data at location %s. Cannot move."), *GetNameSafe(MoverComp->GetOwner()), *AdjustedDest.ToCompactString());
+					UE_LOGF(LogMover, Verbose, "%ls could not find valid navigation data at location %ls. Cannot move.", *GetNameSafe(MoverComp->GetOwner()), *AdjustedDest.ToCompactString());
 					// nothing to be done
 					break;
 
@@ -539,7 +539,7 @@ FVector UNavWalkingMode::ProjectLocationFromNavMesh(float DeltaSeconds, const FV
 	{
 		if (!bCachedLocationStillValid)
 		{
-			UE_LOG(LogMover, VeryVerbose, TEXT("ProjectLocationFromNavMesh(): %s interval: %.3f [SKIP TRACE]"), *GetNameSafe(GetMoverComponent()->GetOwner()), NavMeshProjectionInterval);
+			UE_LOGF(LogMover, VeryVerbose, "ProjectLocationFromNavMesh(): %ls interval: %.3f [SKIP TRACE]", *GetNameSafe(GetMoverComponent()->GetOwner()), NavMeshProjectionInterval);
 
 			FHitResult HitResult;
 			FindBestNavMeshLocation(TraceStart, TraceEnd, CurrentFeetLocation, TargetNavLocation, HitResult);
@@ -567,7 +567,7 @@ FVector UNavWalkingMode::ProjectLocationFromNavMesh(float DeltaSeconds, const FV
 		}
 		else
 		{
-			UE_LOG(LogMover, VeryVerbose, TEXT("ProjectLocationFromNavMesh(): %s interval: %.3f [SKIP TRACE]"), *GetNameSafe(GetMoverComponent()->GetOwner()), NavMeshProjectionInterval);
+			UE_LOGF(LogMover, VeryVerbose, "ProjectLocationFromNavMesh(): %ls interval: %.3f [SKIP TRACE]", *GetNameSafe(GetMoverComponent()->GetOwner()), NavMeshProjectionInterval);
 		}
 
 		// Wrap around to maintain same relative offset to tick time changes.
@@ -629,7 +629,7 @@ void UNavWalkingMode::OnRegistered(const FName ModeName)
 	
 	if (!NavMoverComponent)
 	{
-		UE_LOG(LogMover, Warning, TEXT("NavWalkingMode on %s could not find a valid NavMoverComponent and will not function properly."), *GetNameSafe(GetMoverComponent()->GetOwner()));
+		UE_LOGF(LogMover, Warning, "NavWalkingMode on %ls could not find a valid NavMoverComponent and will not function properly.", *GetNameSafe(GetMoverComponent()->GetOwner()));
 	}
 }
 
diff --git a/Engine/Plugins/Experimental/Mover/Source/Mover/Private/DefaultMovementSet/Modes/SimpleFlyingMode.cpp b/Engine/Plugins/Experimental/Mover/Source/Mover/Private/DefaultMovementSet/Modes/SimpleFlyingMode.cpp
--- a/Engine/Plugins/Experimental/Mover/Source/Mover/Private/DefaultMovementSet/Modes/SimpleFlyingMode.cpp
+++ b/Engine/Plugins/Experimental/Mover/Source/Mover/Private/DefaultMovementSet/Modes/SimpleFlyingMode.cpp
@@ -65,7 +65,7 @@ void USimpleFlyingMode::GenerateMove_Implementation(const FMoverTickStartData& S
 	case EMoveInputType::Invalid:
 	default:
 		{
-			UE_LOG(LogMover, Warning, TEXT("Unhandled MoveInputType %i in USimpleFlyingMode"), MoveInputType);
+			UE_LOGF(LogMover, Warning, "Unhandled MoveInputType %i in USimpleFlyingMode", MoveInputType);
 			DesiredVelocity = FVector::ZeroVector;
 			OutProposedMove.DirectionIntent = FVector::ZeroVector;
 		}
diff --git a/Engine/Plugins/Experimental/Mover/Source/Mover/Private/DefaultMovementSet/Modes/SimpleWalkingMode.cpp b/Engine/Plugins/Experimental/Mover/Source/Mover/Private/DefaultMovementSet/Modes/SimpleWalkingMode.cpp
--- a/Engine/Plugins/Experimental/Mover/Source/Mover/Private/DefaultMovementSet/Modes/SimpleWalkingMode.cpp
+++ b/Engine/Plugins/Experimental/Mover/Source/Mover/Private/DefaultMovementSet/Modes/SimpleWalkingMode.cpp
@@ -79,7 +79,7 @@ void USimpleWalkingMode::GenerateMove_Implementation(const FMoverTickStartData&
 	case EMoveInputType::Invalid:
 	default:
 		{
-			UE_LOG(LogMover, Warning, TEXT("Unhandled MoveInputType %i in USimpleWalkingMode"), MoveInputType);
+			UE_LOGF(LogMover, Warning, "Unhandled MoveInputType %i in USimpleWalkingMode", MoveInputType);
 			DesiredVelocity = FVector::ZeroVector;
 			OutProposedMove.DirectionIntent = FVector::ZeroVector;
 		}
diff --git a/Engine/Plugins/Experimental/Mover/Source/Mover/Private/DefaultMovementSet/MovementModifiers/StanceModifier.cpp b/Engine/Plugins/Experimental/Mover/Source/Mover/Private/DefaultMovementSet/MovementModifiers/StanceModifier.cpp
--- a/Engine/Plugins/Experimental/Mover/Source/Mover/Private/DefaultMovementSet/MovementModifiers/StanceModifier.cpp
+++ b/Engine/Plugins/Experimental/Mover/Source/Mover/Private/DefaultMovementSet/MovementModifiers/StanceModifier.cpp
@@ -51,7 +51,7 @@ void FStanceModifier::OnStart(UMoverComponent* MoverComp, const FMoverTimeStep&
 
 			// Prone isn't currently implemented
 			case EStanceMode::Prone:
-				UE_LOG(LogMover, Warning, TEXT("Stance got into prone stance - That stance is not currently implemented."));
+				UE_LOGF(LogMover, Warning, "Stance got into prone stance - That stance is not currently implemented.");
 				// TODO: returning here so we don't apply any bad state to actor in case prone was set. Eventually, the return should be removed once prone is implemented properly
 				DurationMs = 0;
 				return;
@@ -87,7 +87,7 @@ void FStanceModifier::OnPreMovement(UMoverComponent* MoverComp, const FMoverTime
 	// TODO: Prone isn't currently implemented - so we're just going to cancel the modifier if we got into that state
 	if (ActiveStance == EStanceMode::Prone)
 	{
-		UE_LOG(LogMover, Warning, TEXT("Stance got into prone stance - That stance is not currently implemented."));
+		UE_LOGF(LogMover, Warning, "Stance got into prone stance - That stance is not currently implemented.");
 		DurationMs = 0;
 	}
 }
@@ -247,7 +247,7 @@ void FStanceModifier::ApplyMovementSettings(UMoverComponent* MoverComp)
 
 	// Prone isn't currently implemented properly so we're just doing nothing for now
 	case EStanceMode::Prone:
-		UE_LOG(LogMover, Warning, TEXT("Stance got into prone stance - That mode is not currently implemented fully."));
+		UE_LOGF(LogMover, Warning, "Stance got into prone stance - That mode is not currently implemented fully.");
 		break;
 	}
 }
diff --git a/Engine/Plugins/Experimental/Mover/Source/Mover/Private/DefaultMovementSet/NavMoverComponent.cpp b/Engine/Plugins/Experimental/Mover/Source/Mover/Private/DefaultMovementSet/NavMoverComponent.cpp
--- a/Engine/Plugins/Experimental/Mover/Source/Mover/Private/DefaultMovementSet/NavMoverComponent.cpp
+++ b/Engine/Plugins/Experimental/Mover/Source/Mover/Private/DefaultMovementSet/NavMoverComponent.cpp
@@ -28,7 +28,7 @@ void UNavMoverComponent::InitializeComponent()
 	
 	if (!MoverComponent.IsValid())
 	{
-		UE_LOG(LogMover, Warning, TEXT("NavMoverComponent on %s could not find a valid MoverComponent and will not function properly."), *GetNameSafe(GetOwner()));
+		UE_LOGF(LogMover, Warning, "NavMoverComponent on %ls could not find a valid MoverComponent and will not function properly.", *GetNameSafe(GetOwner()));
 	}
 }
 
@@ -90,8 +90,8 @@ bool UNavMoverComponent::ConsumeNavMovementData(FVector& OutMoveInputIntent, FVe
 		OutMoveInputVelocity = CachedNavMoveInputVelocity;
 		bHasNavMovement = true;
 		
-		UE_LOG(LogMover, VeryVerbose, TEXT("Applying %s as NavMoveInputIntent."), *CachedNavMoveInputIntent.ToString());
-		UE_LOG(LogMover, VeryVerbose, TEXT("Applying %s as NavMoveInputVelocity."), *CachedNavMoveInputVelocity.ToString());
+		UE_LOGF(LogMover, VeryVerbose, "Applying %ls as NavMoveInputIntent.", *CachedNavMoveInputIntent.ToString());
+		UE_LOGF(LogMover, VeryVerbose, "Applying %ls as NavMoveInputVelocity.", *CachedNavMoveInputVelocity.ToString());
 	}
 
 	GameFrameNavMovementConsumed = GFrameCounter;
diff --git a/Engine/Plugins/Experimental/Mover/Source/Mover/Private/LayeredMove.cpp b/Engine/Plugins/Experimental/Mover/Source/Mover/Private/LayeredMove.cpp
--- a/Engine/Plugins/Experimental/Mover/Source/Mover/Private/LayeredMove.cpp
+++ b/Engine/Plugins/Experimental/Mover/Source/Mover/Private/LayeredMove.cpp
@@ -131,7 +131,7 @@ void FLayeredMoveGroup::QueueLayeredMove(TSharedPtr<FLayeredMoveBase> Move)
 	if (ensure(Move.IsValid()))
 	{
 		QueuedLayeredMoves.Add(Move);
-		UE_LOG(LogMover, VeryVerbose, TEXT("LayeredMove queued move (%s)"), *Move->ToSimpleString());
+		UE_LOGF(LogMover, VeryVerbose, "LayeredMove queued move (%ls)", *Move->ToSimpleString());
 	}
 }
 
@@ -223,7 +223,7 @@ static void CopyLayeredMoveArray(TArray<TSharedPtr<FLayeredMoveBase>>& Dest, con
 			}
 			else
 			{
-				UE_LOG(LogMover, Warning, TEXT("CopyLayeredMoveArray trying to copy invalid Other Layered Move"));
+				UE_LOGF(LogMover, Warning, "CopyLayeredMoveArray trying to copy invalid Other Layered Move");
 			}
 		}
 	}
@@ -474,7 +474,7 @@ void FLayeredMoveGroup::GatherResidualVelocitySettings(const TSharedPtr<FLayered
 		{
 			if (bResidualVelocityOverridden)
 			{
-				UE_LOG(LogMover, Log, TEXT("Multiple LayeredMove residual settings have a MixMode that overrides. Only one will take effect."));
+				UE_LOGF(LogMover, Log, "Multiple LayeredMove residual settings have a MixMode that overrides. Only one will take effect.");
 			}
 
 			bResidualVelocityOverridden = true;
@@ -488,7 +488,7 @@ void FLayeredMoveGroup::GatherResidualVelocitySettings(const TSharedPtr<FLayered
 		{
 			if (bResidualVelocityOverridden)
 			{
-				UE_LOG(LogMover, Log, TEXT("Multiple LayeredMove residual settings have a MixMode that overrides. Only one will take effect."));
+				UE_LOGF(LogMover, Log, "Multiple LayeredMove residual settings have a MixMode that overrides. Only one will take effect.");
 			}
 
 			bResidualVelocityOverridden = true;
@@ -506,7 +506,7 @@ void FLayeredMoveGroup::GatherResidualVelocitySettings(const TSharedPtr<FLayered
 		{
 			if (bClampVelocityOverridden)
 			{
-				UE_LOG(LogMover, Log, TEXT("Multiple LayeredMove residual settings have a MixMode that overrides. Only one will take effect."));
+				UE_LOGF(LogMover, Log, "Multiple LayeredMove residual settings have a MixMode that overrides. Only one will take effect.");
 			}
 
 			bClampVelocityOverridden = true;
@@ -528,7 +528,7 @@ void FLayeredMoveGroup::GatherResidualVelocitySettings(const TSharedPtr<FLayered
 		{
 			if (bClampVelocityOverridden)
 			{
-				UE_LOG(LogMover, Log, TEXT("Multiple LayeredMove residual settings have a MixMode that overrides. Only one will take effect."));
+				UE_LOGF(LogMover, Log, "Multiple LayeredMove residual settings have a MixMode that overrides. Only one will take effect.");
 			}
 
 			bClampVelocityOverridden = true;
@@ -622,14 +622,14 @@ struct FLayeredMoveDeleter
 			}
 			else
 			{
-				UE_LOG(LogMover, Error, TEXT("FLayeredMoveGroup::NetSerialize: ScriptStruct not derived from FLayeredMoveBase attempted to serialize."));
+				UE_LOGF(LogMover, Error, "FLayeredMoveGroup::NetSerialize: ScriptStruct not derived from FLayeredMoveBase attempted to serialize.");
 				Ar.SetError();
 				break;
 			}
 		}
 		else if (ScriptStruct.IsError())
 		{
-			UE_LOG(LogMover, Error, TEXT("FLayeredMoveGroup::NetSerialize: Invalid ScriptStruct serialized."));
+			UE_LOGF(LogMover, Error, "FLayeredMoveGroup::NetSerialize: Invalid ScriptStruct serialized.");
 			Ar.SetError();
 			break;
 		}
diff --git a/Engine/Plugins/Experimental/Mover/Source/Mover/Private/LayeredMoveBase.cpp b/Engine/Plugins/Experimental/Mover/Source/Mover/Private/LayeredMoveBase.cpp
--- a/Engine/Plugins/Experimental/Mover/Source/Mover/Private/LayeredMoveBase.cpp
+++ b/Engine/Plugins/Experimental/Mover/Source/Mover/Private/LayeredMoveBase.cpp
@@ -366,14 +366,14 @@ void FLayeredMoveInstance::NetSerialize(FArchive& Ar)
 		}
 		else
 		{
-			UE_LOG(LogMover, Error, TEXT("FLayeredMoveInstanceGroup::NetSerialize: ScriptStruct [%s] not derived from FLayeredMoveInstancedData attempted to serialize."),
+			UE_LOGF(LogMover, Error, "FLayeredMoveInstanceGroup::NetSerialize: ScriptStruct [%ls] not derived from FLayeredMoveInstancedData attempted to serialize.",
 				*DataStructType->GetName());
 			Ar.SetError();
 		}
 	}
 	else if (DataStructType.IsError())
 	{
-		UE_LOG(LogMover, Error, TEXT("FLayeredMoveInstanceGroup::NetSerialize: Invalid ScriptStruct serialized."));
+		UE_LOGF(LogMover, Error, "FLayeredMoveInstanceGroup::NetSerialize: Invalid ScriptStruct serialized.");
 		Ar.SetError();
 	}
 }
@@ -404,12 +404,12 @@ bool FLayeredMoveInstance::PopulateMissingActiveMoveLogic(const TArray<TObjectPt
 
 			if (!bFoundMissingLogic)
 			{
-				UE_LOG(LogMover, Warning, TEXT("Active Layered Move couldn't find it's serialized MoveLogicClass (%s) among registered MoveLogic"), *MoveLogicClassType->GetName());
+				UE_LOGF(LogMover, Warning, "Active Layered Move couldn't find it's serialized MoveLogicClass (%ls) among registered MoveLogic", *MoveLogicClassType->GetName());
 			}
 		}
 		else
 		{
-			UE_LOG(LogMover, Warning, TEXT("Active Layered Move didn't have a valid logic class or class type to search for"));
+			UE_LOGF(LogMover, Warning, "Active Layered Move didn't have a valid logic class or class type to search for");
 		}
 
 		return bFoundMissingLogic;
diff --git a/Engine/Plugins/Experimental/Mover/Source/Mover/Private/LayeredMoveGroup.cpp b/Engine/Plugins/Experimental/Mover/Source/Mover/Private/LayeredMoveGroup.cpp
--- a/Engine/Plugins/Experimental/Mover/Source/Mover/Private/LayeredMoveGroup.cpp
+++ b/Engine/Plugins/Experimental/Mover/Source/Mover/Private/LayeredMoveGroup.cpp
@@ -56,7 +56,7 @@ void FLayeredMoveInstanceGroup::QueueLayeredMove(const TSharedPtr<FLayeredMoveIn
 	{
 		QueuedMoves.Add(Move);
 		// TODO NS: Add simple string support and re add this log
-		//UE_LOG(LogMover, VeryVerbose, TEXT("LayeredMove queued move (%s)"), *Move->ToSimpleString());
+		//UE_LOGF(LogMover, VeryVerbose, "LayeredMove queued move (%ls)", *Move->ToSimpleString());
 	}
 }
 bool FLayeredMoveInstanceGroup::HasSameContents(const FLayeredMoveInstanceGroup& Other) const
@@ -335,7 +335,7 @@ void FLayeredMoveInstanceGroup::FlushMoveArrays(const FMoverTimeStep& TimeStep,
 			else
 			{
 				// We should've populated missing logic before this so let's just clear this and warn about it
-				UE_LOG(LogMover, Warning, TEXT("Queued Active Move (%s) logic was not present. Move will not be activated."), *QueuedMove->GetDataStructType()->GetName());
+				UE_LOGF(LogMover, Warning, "Queued Active Move (%ls) logic was not present. Move will not be activated.", *QueuedMove->GetDataStructType()->GetName());
 			}
 		}
 	}
diff --git a/Engine/Plugins/Experimental/Mover/Source/Mover/Private/MoveLibrary/AsyncMovementUtils.cpp b/Engine/Plugins/Experimental/Mover/Source/Mover/Private/MoveLibrary/AsyncMovementUtils.cpp
--- a/Engine/Plugins/Experimental/Mover/Source/Mover/Private/MoveLibrary/AsyncMovementUtils.cpp
+++ b/Engine/Plugins/Experimental/Mover/Source/Mover/Private/MoveLibrary/AsyncMovementUtils.cpp
@@ -52,7 +52,7 @@ bool UAsyncMovementUtils::TestDepenetratingMove(const FMovingComponentSet& Movin
 
 		CurrentLocation = ResolvedStartLocation + ((TargetLocation-ResolvedStartLocation) * OutHit.Time);
 
-		UE_LOG(LogMover, VeryVerbose, TEXT("TestDepenetratingMove: %s (role %i) Delta=%s DidMove=%i"),
+		UE_LOGF(LogMover, VeryVerbose, "TestDepenetratingMove: %ls (role %i) Delta=%ls DidMove=%i",
 			*GetNameSafe(UpdatedComponent->GetOwner()), UpdatedComponent->GetOwnerRole(), *(TargetLocation- ResolvedStartLocation).ToCompactString(), bDidMove);
 	}
 
@@ -70,7 +70,7 @@ bool UAsyncMovementUtils::TestDepenetratingMove(const FMovingComponentSet& Movin
 			// Retry original move
 			bDidMove |= TestMoveComponent_Internal(MovingComps, ResolvedStartLocation, TargetLocation, StartRotation, TargetRotation, bShouldSweep, CollisionParams, OutHit);
 
-			UE_LOG(LogMover, VeryVerbose, TEXT("TestDepenetratingMove retry: %s (role %i) Delta=%s DidMove=%i"),
+			UE_LOGF(LogMover, VeryVerbose, "TestDepenetratingMove retry: %ls (role %i) Delta=%ls DidMove=%i",
 				*GetNameSafe(UpdatedComponent->GetOwner()), UpdatedComponent->GetOwnerRole(), *(TargetLocation - ResolvedStartLocation).ToCompactString(), bDidMove);
 
 			if (bDidMove)
@@ -199,7 +199,7 @@ bool UAsyncMovementUtils::FindMoveToResolveInitialPenetration_Internal(const FMo
 
 			bool bMoved = TestMoveComponent_Internal(MovingComps, StartLocation, StartLocation + ProposedAdjustment, StartRotation, StartRotation, /* bShouldSweep */ true, CollisionParams, SweepOutHit);
 
-			UE_LOG(LogMover, VeryVerbose, TEXT("FindMoveToResolveInitialPenetration_Internal: %s (role %i) Adjustment=%s DidMove=%i"),
+			UE_LOGF(LogMover, VeryVerbose, "FindMoveToResolveInitialPenetration_Internal: %ls (role %i) Adjustment=%ls DidMove=%i",
 				*GetNameSafe(UpdatedComponent->GetOwner()), UpdatedComponent->GetOwnerRole(), *ProposedAdjustment.ToCompactString(), bMoved);
 
 			// Still stuck in penetration?
@@ -213,7 +213,7 @@ bool UAsyncMovementUtils::FindMoveToResolveInitialPenetration_Internal(const FMo
 				{
 					bMoved = TestMoveComponent_Internal(MovingComps, StartLocation, StartLocation+CombinedMTD, StartRotation, StartRotation, /*bShouldSweep*/ true, CollisionParams, SweepOutHit);
 
-					UE_LOG(LogMover, VeryVerbose, TEXT("FindMoveToResolveInitialPenetration_Internal combined: %s (role %i) CombinedAdjustment=%s DidMove=%i"),
+					UE_LOGF(LogMover, VeryVerbose, "FindMoveToResolveInitialPenetration_Internal combined: %ls (role %i) CombinedAdjustment=%ls DidMove=%i",
 						*GetNameSafe(UpdatedComponent->GetOwner()), UpdatedComponent->GetOwnerRole(), *CombinedMTD.ToCompactString(), bMoved);
 				}
 
@@ -231,7 +231,7 @@ bool UAsyncMovementUtils::FindMoveToResolveInitialPenetration_Internal(const FMo
 					const FVector TargetLocationWithAdjustment = StartLocation + ProposedAdjustment + FullMoveDelta;
 					bMoved = TestMoveComponent_Internal(MovingComps, StartLocation, TargetLocationWithAdjustment, StartRotation, StartRotation, /*bShouldSweep*/ true, CollisionParams, SweepOutHit);
 
-					UE_LOG(LogMover, VeryVerbose, TEXT("FindMoveToResolveInitialPenetration_Internal multiple: %s (role %i) TargetLocationWithAdjustment=%s DidMove=%i"),
+					UE_LOGF(LogMover, VeryVerbose, "FindMoveToResolveInitialPenetration_Internal multiple: %ls (role %i) TargetLocationWithAdjustment=%ls DidMove=%i",
 						*GetNameSafe(UpdatedComponent->GetOwner()), UpdatedComponent->GetOwnerRole(), *TargetLocationWithAdjustment.ToCompactString(), bMoved);
 				}
 			}
diff --git a/Engine/Plugins/Experimental/Mover/Source/Mover/Private/MoveLibrary/BasedMovementUtils.cpp b/Engine/Plugins/Experimental/Mover/Source/Mover/Private/MoveLibrary/BasedMovementUtils.cpp
--- a/Engine/Plugins/Experimental/Mover/Source/Mover/Private/MoveLibrary/BasedMovementUtils.cpp
+++ b/Engine/Plugins/Experimental/Mover/Source/Mover/Private/MoveLibrary/BasedMovementUtils.cpp
@@ -130,7 +130,7 @@ bool UBasedMovementUtils::GetMovementBaseTransform(const UPrimitiveComponent* Mo
 			}
 
 			bBoneNameIsInvalid = true;
-			UE_LOG(LogMover, Warning, TEXT("GetMovementBaseTransform(): Invalid bone or socket '%s' for PrimitiveComponent base %s. Using component's root transform instead."), *BoneName.ToString(), *GetPathNameSafe(MovementBase));
+			UE_LOGF(LogMover, Warning, "GetMovementBaseTransform(): Invalid bone or socket '%ls' for PrimitiveComponent base %ls. Using component's root transform instead.", *BoneName.ToString(), *GetPathNameSafe(MovementBase));
 		}
 
 		OutLocation = MovementBase->GetComponentLocation();
@@ -292,7 +292,7 @@ void UBasedMovementUtils::AddTickDependency(FTickFunction& BasedObjectTick, UPri
 	}
 	else
 	{
-		UE_LOG(LogMover, Warning, TEXT("Attempted to AddTickDependency on an invalid or non-dynamic base: %s"), *GetNameSafe(NewBase));
+		UE_LOGF(LogMover, Warning, "Attempted to AddTickDependency on an invalid or non-dynamic base: %ls", *GetNameSafe(NewBase));
 	}
 }
 
diff --git a/Engine/Plugins/Experimental/Mover/Source/Mover/Private/MoveLibrary/ConstrainedMoveUtils.cpp b/Engine/Plugins/Experimental/Mover/Source/Mover/Private/MoveLibrary/ConstrainedMoveUtils.cpp
--- a/Engine/Plugins/Experimental/Mover/Source/Mover/Private/MoveLibrary/ConstrainedMoveUtils.cpp
+++ b/Engine/Plugins/Experimental/Mover/Source/Mover/Private/MoveLibrary/ConstrainedMoveUtils.cpp
@@ -18,7 +18,7 @@ void UPlanarConstraintUtils::SetPlanarConstraintNormal(UPARAM(ref) FPlanarConstr
 
 	if (PlaneNormal.IsNearlyZero())
 	{
-		UE_LOG(LogMover, Warning, TEXT("Can't use SetPlanarConstraintNormal with a zero-length normal. Leaving normal as %s"), *Constraint.PlaneConstraintNormal.ToCompactString());
+		UE_LOGF(LogMover, Warning, "Can't use SetPlanarConstraintNormal with a zero-length normal. Leaving normal as %ls", *Constraint.PlaneConstraintNormal.ToCompactString());
 	}
 	else
 	{
diff --git a/Engine/Plugins/Experimental/Mover/Source/Mover/Private/MoveLibrary/GroundMovementUtils.cpp b/Engine/Plugins/Experimental/Mover/Source/Mover/Private/MoveLibrary/GroundMovementUtils.cpp
--- a/Engine/Plugins/Experimental/Mover/Source/Mover/Private/MoveLibrary/GroundMovementUtils.cpp
+++ b/Engine/Plugins/Experimental/Mover/Source/Mover/Private/MoveLibrary/GroundMovementUtils.cpp
@@ -88,14 +88,14 @@ float UGroundMovementUtils::TryMoveToStepUp(const FMovingComponentSet& MovingCom
 	const float OldLocationDot = OldLocation.Dot(UpDirection);
 	if (InitialImpactDot > OldLocationDot + (PawnHalfHeight - PawnRadius))
 	{
-		UE_LOG(LogMover, VeryVerbose, TEXT("Not stepping up due to top of shape hitting something"));
+		UE_LOGF(LogMover, VeryVerbose, "Not stepping up due to top of shape hitting something");
 		return 0;
 	}
 
 	// TODO: We should rely on movement plane normal, rather than gravity direction
 	if (GravDir.IsZero())
 	{
-		UE_LOG(LogMover, VeryVerbose, TEXT("Not stepping up because there's no gravity"));
+		UE_LOGF(LogMover, VeryVerbose, "Not stepping up because there's no gravity");
 		return 0;
 	}
 
@@ -132,7 +132,7 @@ float UGroundMovementUtils::TryMoveToStepUp(const FMovingComponentSet& MovingCom
 	// Don't step up if the impact is below us, accounting for distance from floor.
 	if (InitialImpactDot <= PawnInitialFloorBaseDot)
 	{
-		UE_LOG(LogMover, VeryVerbose, TEXT("Not stepping up because the impact is below us"));
+		UE_LOGF(LogMover, VeryVerbose, "Not stepping up because the impact is below us");
 		return 0;
 	}
 
@@ -146,13 +146,13 @@ float UGroundMovementUtils::TryMoveToStepUp(const FMovingComponentSet& MovingCom
 	const FVector UpAdjustment = -GravDir * StepTravelUpHeight;
 	const bool bDidStepUp = UMovementUtils::TryMoveUpdatedComponent_Internal(MovingComps, UpAdjustment, PawnRotation, true, MOVECOMP_NoFlags, &SweepUpHit, ETeleportType::None);
 
-	UE_LOG(LogMover, VeryVerbose, TEXT("TryMoveToStepUp Up: %s (role %i) UpAdjustment=%s DidMove=%i"),
+	UE_LOGF(LogMover, VeryVerbose, "TryMoveToStepUp Up: %ls (role %i) UpAdjustment=%ls DidMove=%i",
 		*GetNameSafe(RootComponent->GetOwner()), RootComponent->GetOwnerRole(), *UpAdjustment.ToCompactString(), bDidStepUp);
 
 	if (SweepUpHit.bStartPenetrating)
 	{
 		// Undo movement
-		UE_LOG(LogMover, VeryVerbose, TEXT("Reverting step-up attempt because we started in a penetrating state"));
+		UE_LOGF(LogMover, VeryVerbose, "Reverting step-up attempt because we started in a penetrating state");
 		ScopedStepUpMovement.RevertMove();
 		return 0;
 	}
@@ -165,7 +165,7 @@ float UGroundMovementUtils::TryMoveToStepUp(const FMovingComponentSet& MovingCom
 	FHitResult StepFwdHit(1.f);
 	const bool bDidStepFwd = UMovementUtils::TryMoveUpdatedComponent_Internal(MovingComps, MoveDelta, PawnRotation, true, MOVECOMP_NoFlags, &StepFwdHit, ETeleportType::None);
 
-	UE_LOG(LogMover, VeryVerbose, TEXT("TryMoveToStepUp Fwd: %s (role %i) MoveDelta=%s DidMove=%i"),
+	UE_LOGF(LogMover, VeryVerbose, "TryMoveToStepUp Fwd: %ls (role %i) MoveDelta=%ls DidMove=%i",
 		*GetNameSafe(RootComponent->GetOwner()), RootComponent->GetOwnerRole(), *MoveDelta.ToCompactString(), bDidStepFwd);
 
 	PercentMovementUsed += StepFwdHit.Time;
@@ -176,7 +176,7 @@ float UGroundMovementUtils::TryMoveToStepUp(const FMovingComponentSet& MovingCom
 		if (StepFwdHit.bStartPenetrating)
 		{
 			// Undo movement
-			UE_LOG(LogMover, VeryVerbose, TEXT("Reverting step-fwd attempt during step-up, because we started in a penetrating state"));
+			UE_LOGF(LogMover, VeryVerbose, "Reverting step-fwd attempt during step-up, because we started in a penetrating state");
 			ScopedStepUpMovement.RevertMove();
 			return 0;
 		}
@@ -230,7 +230,7 @@ float UGroundMovementUtils::TryMoveToStepUp(const FMovingComponentSet& MovingCom
 		// If both the forward hit and the deflection got us nowhere, there is no point in this step up.
 		if (ForwardHitTime == 0.f && ForwardSlideAmount == 0.f)
 		{
-			UE_LOG(LogMover, VeryVerbose, TEXT("Reverting step-fwd attempt during step-up, because no movement differences occurred"));
+			UE_LOGF(LogMover, VeryVerbose, "Reverting step-fwd attempt during step-up, because no movement differences occurred");
 			ScopedStepUpMovement.RevertMove();
 			return 0;
 		}
@@ -249,14 +249,14 @@ float UGroundMovementUtils::TryMoveToStepUp(const FMovingComponentSet& MovingCom
 	const FVector StepDownAdjustment = GravDir * StepTravelDownHeight;
 	const bool bDidStepDown = UMovementUtils::TryMoveUpdatedComponent_Internal(MovingComps, StepDownAdjustment, RootComponent->GetComponentQuat(), true, MOVECOMP_NoFlags, &StepDownHit, ETeleportType::None);
 
-	UE_LOG(LogMover, VeryVerbose, TEXT("TryMoveToStepUp Down: %s (role %i) StepDownAdjustment=%s DidMove=%i"),
+	UE_LOGF(LogMover, VeryVerbose, "TryMoveToStepUp Down: %ls (role %i) StepDownAdjustment=%ls DidMove=%i",
 		*GetNameSafe(RootComponent->GetOwner()), RootComponent->GetOwnerRole(), *StepDownAdjustment.ToCompactString(), bDidStepDown);
 
 
 	// If step down was initially penetrating abort the step up
 	if (StepDownHit.bStartPenetrating)
 	{
-		UE_LOG(LogMover, VeryVerbose, TEXT("Reverting step-down attempt during step-up/step-fwd, because we started in a penetrating state"));
+		UE_LOGF(LogMover, VeryVerbose, "Reverting step-down attempt during step-up/step-fwd, because we started in a penetrating state");
 		ScopedStepUpMovement.RevertMove();
 		return 0;
 	}
@@ -268,7 +268,7 @@ float UGroundMovementUtils::TryMoveToStepUp(const FMovingComponentSet& MovingCom
 		const float DeltaDot = (StepDownHit.ImpactPoint.Dot(UpDirection)) - PawnFloorPointDot;
 		if (DeltaDot > MaxStepHeight)
 		{
-			UE_LOG(LogMover, VeryVerbose, TEXT("Reject step-down attempt during step-up/step-fwd, because it made us travel too high (too high Height %.3f) up from floor base %f to %f"), DeltaDot, PawnInitialFloorBaseDot, StepDownHit.ImpactPoint.Z);
+			UE_LOGF(LogMover, VeryVerbose, "Reject step-down attempt during step-up/step-fwd, because it made us travel too high (too high Height %.3f) up from floor base %f to %f", DeltaDot, PawnInitialFloorBaseDot, StepDownHit.ImpactPoint.Z);
 			ScopedStepUpMovement.RevertMove();
 			return 0;
 		}
@@ -280,7 +280,7 @@ float UGroundMovementUtils::TryMoveToStepUp(const FMovingComponentSet& MovingCom
 			const bool bNormalTowardsMe = (MoveDelta | StepDownHit.ImpactNormal) < 0.f;
 			if (bNormalTowardsMe)
 			{
-				UE_LOG(LogMover, VeryVerbose, TEXT("Reject step-down attempt during step-up/step-fwd, due to unwalkable normal %s opposed to movement"), *StepDownHit.ImpactNormal.ToString());
+				UE_LOGF(LogMover, VeryVerbose, "Reject step-down attempt during step-up/step-fwd, due to unwalkable normal %ls opposed to movement", *StepDownHit.ImpactNormal.ToString());
 				ScopedStepUpMovement.RevertMove();
 				return 0;
 			}
@@ -289,7 +289,7 @@ float UGroundMovementUtils::TryMoveToStepUp(const FMovingComponentSet& MovingCom
 			// It's fine to step down onto an unwalkable normal below us, we will just slide off. Rejecting those moves would prevent us from being able to walk off the edge.
 			if (StepDownHit.Location.Dot(UpDirection) > OldLocationDot)
 			{
-				UE_LOG(LogMover, VeryVerbose, TEXT("Reject step-down attempt during step-up/step-fwd, due to unwalkable normal %s above old position)"), *StepDownHit.ImpactNormal.ToString());
+				UE_LOGF(LogMover, VeryVerbose, "Reject step-down attempt during step-up/step-fwd, due to unwalkable normal %ls above old position)", *StepDownHit.ImpactNormal.ToString());
 				ScopedStepUpMovement.RevertMove();
 				return 0;
 			}
@@ -298,7 +298,7 @@ float UGroundMovementUtils::TryMoveToStepUp(const FMovingComponentSet& MovingCom
 		// Reject moves where the downward sweep hit something very close to the edge of the shape. This maintains consistency with FindFloor as well.
 		if (!UFloorQueryUtils::IsWithinEdgeTolerance(StepDownHit.Location, StepDownHit.ImpactPoint, PawnRadius, UpDirection))
 		{
-			UE_LOG(LogMover, VeryVerbose, TEXT("Reject step-down attempt during step-up/step-fwd, due to being outside edge tolerance"));
+			UE_LOGF(LogMover, VeryVerbose, "Reject step-down attempt during step-up/step-fwd, due to being outside edge tolerance");
 			ScopedStepUpMovement.RevertMove();
 			return 0;
 		}
@@ -306,7 +306,7 @@ float UGroundMovementUtils::TryMoveToStepUp(const FMovingComponentSet& MovingCom
 		// Don't step up onto invalid surfaces if traveling higher.
 		if (DeltaDot > 0.f && !CanStepUpOnHitSurface(StepDownHit))
 		{
-			UE_LOG(LogMover, VeryVerbose, TEXT("Reject step-down attempt during step-up/step-fwd, due to being up onto surface with !CanStepUpOnHitSurface")); 
+			UE_LOGF(LogMover, VeryVerbose, "Reject step-down attempt during step-up/step-fwd, due to being up onto surface with !CanStepUpOnHitSurface"); 
 			ScopedStepUpMovement.RevertMove();
 			return 0;
 		}
@@ -326,7 +326,7 @@ float UGroundMovementUtils::TryMoveToStepUp(const FMovingComponentSet& MovingCom
 				const float MAX_STEP_SIDE_DOT = 0.08f; // TODO: Move magic numbers elsewhere
 				if (!StepDownResult.FloorTestResult.bBlockingHit && StepSideDot < MAX_STEP_SIDE_DOT)
 				{
-					UE_LOG(LogMover, VeryVerbose, TEXT("Reject step-down attempt during step-up/step-fwd, due to it being an unperchable step")); 
+					UE_LOGF(LogMover, VeryVerbose, "Reject step-down attempt during step-up/step-fwd, due to it being an unperchable step"); 
 					ScopedStepUpMovement.RevertMove();
 					return 0;
 				}
@@ -799,14 +799,14 @@ float UGroundMovementUtils::TestMoveToStepOver(const FMovingComponentSet& Moving
 	const float OldLocationDot = OldLocation.Dot(UpDirection);
 	if (InitialImpactDot > OldLocationDot + (PawnHalfHeight - PawnRadius))
 	{
-		UE_LOG(LogMover, VeryVerbose, TEXT("Not stepping up due to top of shape hitting something"));
+		UE_LOGF(LogMover, VeryVerbose, "Not stepping up due to top of shape hitting something");
 		return 0;
 	}
 
 	// TODO: We should rely on movement plane normal, rather than gravity direction
 	if (GravDir.IsZero())
 	{
-		UE_LOG(LogMover, VeryVerbose, TEXT("Not stepping up because there's no gravity"));
+		UE_LOGF(LogMover, VeryVerbose, "Not stepping up because there's no gravity");
 		return 0;
 	}
 
@@ -842,7 +842,7 @@ float UGroundMovementUtils::TestMoveToStepOver(const FMovingComponentSet& Moving
 	// Don't step up if the impact is below us, accounting for distance from floor.
 	if (InitialImpactDot <= PawnInitialFloorBaseDot)
 	{
-		UE_LOG(LogMover, VeryVerbose, TEXT("Not stepping up because the impact is below us"));
+		UE_LOGF(LogMover, VeryVerbose, "Not stepping up because the impact is below us");
 		return 0;
 	}
 
@@ -854,13 +854,13 @@ float UGroundMovementUtils::TestMoveToStepOver(const FMovingComponentSet& Moving
 
 	const bool bDidStepUp = UAsyncMovementUtils::TestMoveComponent_Internal(MovingComps, LocationInProgress, LocationInProgress + UpAdjustment, Rotation, Rotation, /*bShouldSweep=*/ true, CollisionParams, SweepUpHit);
 
-	UE_LOG(LogMover, VeryVerbose, TEXT("TestMoveToStepOver Up: %s (role %i) UpAdjustment=%s DidMove=%i"),
+	UE_LOGF(LogMover, VeryVerbose, "TestMoveToStepOver Up: %ls (role %i) UpAdjustment=%ls DidMove=%i",
 		*GetNameSafe(RootComponent->GetOwner()), RootComponent->GetOwnerRole(), *UpAdjustment.ToCompactString(), bDidStepUp);
 
 	if (SweepUpHit.bStartPenetrating)
 	{
 		// Undo movement
-		UE_LOG(LogMover, VeryVerbose, TEXT("Reverting step-up attempt because we started in a penetrating state"));
+		UE_LOGF(LogMover, VeryVerbose, "Reverting step-up attempt because we started in a penetrating state");
 		return 0;
 	}
 
@@ -876,7 +876,7 @@ float UGroundMovementUtils::TestMoveToStepOver(const FMovingComponentSet& Moving
 
 	FVector FwdStepDelta = (StepFwdHit.TraceStart + ((StepFwdHit.TraceEnd - StepFwdHit.TraceStart) * StepFwdHit.Time)) - LocationInProgress;
 
-	UE_LOG(LogMover, VeryVerbose, TEXT("TestMoveToStepOver Fwd: %s (role %i) MoveDelta=%s DidMove=%i"),
+	UE_LOGF(LogMover, VeryVerbose, "TestMoveToStepOver Fwd: %ls (role %i) MoveDelta=%ls DidMove=%i",
 		*GetNameSafe(RootComponent->GetOwner()), RootComponent->GetOwnerRole(), *MoveDelta.ToCompactString(), bDidStepFwd);
 
 	// Check result of forward movement
@@ -885,7 +885,7 @@ float UGroundMovementUtils::TestMoveToStepOver(const FMovingComponentSet& Moving
 		if (StepFwdHit.bStartPenetrating)
 		{
 			// Undo movement
-			UE_LOG(LogMover, VeryVerbose, TEXT("Reverting step-fwd attempt during step-up, because we started in a penetrating state"));
+			UE_LOGF(LogMover, VeryVerbose, "Reverting step-fwd attempt during step-up, because we started in a penetrating state");
 			return 0;
 		}
 
@@ -945,14 +945,14 @@ float UGroundMovementUtils::TestMoveToStepOver(const FMovingComponentSet& Moving
 
 		if (bIsFalling)
 		{
-			UE_LOG(LogMover, VeryVerbose, TEXT("Reverting step-fwd attempt during step-up, because we could not adjust without falling"));
+			UE_LOGF(LogMover, VeryVerbose, "Reverting step-fwd attempt during step-up, because we could not adjust without falling");
 			return 0;
 		}
 
 		// If both the forward hit and the deflection got us nowhere, there is no point in this step up.
 		if (ForwardHitTime == 0.f && ForwardSlideAmount == 0.f)
 		{
-			UE_LOG(LogMover, VeryVerbose, TEXT("Reverting step-fwd attempt during step-up, because no movement differences occurred"));
+			UE_LOGF(LogMover, VeryVerbose, "Reverting step-fwd attempt during step-up, because no movement differences occurred");
 			return 0;
 		}
 	}
@@ -971,14 +971,14 @@ float UGroundMovementUtils::TestMoveToStepOver(const FMovingComponentSet& Moving
 	const bool bDidStepDown = UAsyncMovementUtils::TestMoveComponent_Internal(MovingComps, LocationInProgress, LocationInProgress + StepDownAdjustment, Rotation, Rotation, /*bShouldSweep=*/ true, CollisionParams, StepFwdHit);
 
 
-	UE_LOG(LogMover, VeryVerbose, TEXT("TestMoveToStepOver Down: %s (role %i) StepDownAdjustment=%s DidMove=%i"),
+	UE_LOGF(LogMover, VeryVerbose, "TestMoveToStepOver Down: %ls (role %i) StepDownAdjustment=%ls DidMove=%i",
 		*GetNameSafe(RootComponent->GetOwner()), RootComponent->GetOwnerRole(), *StepDownAdjustment.ToCompactString(), bDidStepDown);
 
 
 	// If step down was initially penetrating abort the step up
 	if (StepFwdHit.bStartPenetrating)
 	{
-		UE_LOG(LogMover, VeryVerbose, TEXT("Reverting step-down attempt during step-up/step-fwd, because we started in a penetrating state"));
+		UE_LOGF(LogMover, VeryVerbose, "Reverting step-down attempt during step-up/step-fwd, because we started in a penetrating state");
 		return 0;
 	}
 
@@ -989,7 +989,7 @@ float UGroundMovementUtils::TestMoveToStepOver(const FMovingComponentSet& Moving
 		const float DeltaDot = (StepFwdHit.ImpactPoint.Dot(UpDirection)) - PawnFloorPointDot;
 		if (DeltaDot > MaxStepHeight)
 		{
-			UE_LOG(LogMover, VeryVerbose, TEXT("Reject step-down attempt during step-up/step-fwd, because it made us travel too high (too high Height %.3f) up from floor base %f to %f"), DeltaDot, PawnInitialFloorBaseDot, StepFwdHit.ImpactPoint.Z);
+			UE_LOGF(LogMover, VeryVerbose, "Reject step-down attempt during step-up/step-fwd, because it made us travel too high (too high Height %.3f) up from floor base %f to %f", DeltaDot, PawnInitialFloorBaseDot, StepFwdHit.ImpactPoint.Z);
 			return 0;
 		}
 
@@ -1000,7 +1000,7 @@ float UGroundMovementUtils::TestMoveToStepOver(const FMovingComponentSet& Moving
 			const bool bNormalTowardsMe = (MoveDelta | StepFwdHit.ImpactNormal) < 0.f;
 			if (bNormalTowardsMe)
 			{
-				UE_LOG(LogMover, VeryVerbose, TEXT("Reject step-down attempt during step-up/step-fwd, due to unwalkable normal %s opposed to movement"), *StepFwdHit.ImpactNormal.ToString());
+				UE_LOGF(LogMover, VeryVerbose, "Reject step-down attempt during step-up/step-fwd, due to unwalkable normal %ls opposed to movement", *StepFwdHit.ImpactNormal.ToString());
 				return 0;
 			}
 
@@ -1008,7 +1008,7 @@ float UGroundMovementUtils::TestMoveToStepOver(const FMovingComponentSet& Moving
 			// It's fine to step down onto an unwalkable normal below us, we will just slide off. Rejecting those moves would prevent us from being able to walk off the edge.
 			if (StepFwdHit.Location.Dot(UpDirection) > OldLocationDot)
 			{
-				UE_LOG(LogMover, VeryVerbose, TEXT("Reject step-down attempt during step-up/step-fwd, due to unwalkable normal %s above old position)"), *StepFwdHit.ImpactNormal.ToString());
+				UE_LOGF(LogMover, VeryVerbose, "Reject step-down attempt during step-up/step-fwd, due to unwalkable normal %ls above old position)", *StepFwdHit.ImpactNormal.ToString());
 				return 0;
 			}
 		}
@@ -1016,14 +1016,14 @@ float UGroundMovementUtils::TestMoveToStepOver(const FMovingComponentSet& Moving
 		// Reject moves where the downward sweep hit something very close to the edge of the shape. This maintains consistency with FindFloor as well.
 		if (!UFloorQueryUtils::IsWithinEdgeTolerance(StepFwdHit.Location, StepFwdHit.ImpactPoint, PawnRadius, UpDirection))
 		{
-			UE_LOG(LogMover, VeryVerbose, TEXT("Reject step-down attempt during step-up/step-fwd, due to being outside edge tolerance"));
+			UE_LOGF(LogMover, VeryVerbose, "Reject step-down attempt during step-up/step-fwd, due to being outside edge tolerance");
 			return 0;
 		}
 
 		// Don't step up onto invalid surfaces if traveling higher.
 		if (DeltaDot > 0.f && !UGroundMovementUtils::CanStepUpOnHitSurface(StepFwdHit))
 		{
-			UE_LOG(LogMover, VeryVerbose, TEXT("Reject step-down attempt during step-up/step-fwd, due to being up onto surface with !CanStepUpOnHitSurface"));
+			UE_LOGF(LogMover, VeryVerbose, "Reject step-down attempt during step-up/step-fwd, due to being up onto surface with !CanStepUpOnHitSurface");
 			return 0;
 		}
 
@@ -1042,7 +1042,7 @@ float UGroundMovementUtils::TestMoveToStepOver(const FMovingComponentSet& Moving
 				const float MAX_STEP_SIDE_DOT = 0.08f; // TODO: Move magic numbers elsewhere
 				if (!StepDownResult.FloorTestResult.bBlockingHit && StepSideDot < MAX_STEP_SIDE_DOT)
 				{
-					UE_LOG(LogMover, VeryVerbose, TEXT("Reject step-down attempt during step-up/step-fwd, due to it being an unperchable step"));
+					UE_LOGF(LogMover, VeryVerbose, "Reject step-down attempt during step-up/step-fwd, due to it being an unperchable step");
 					return 0;
 				}
 			}
diff --git a/Engine/Plugins/Experimental/Mover/Source/Mover/Private/MoveLibrary/MovementMixer.cpp b/Engine/Plugins/Experimental/Mover/Source/Mover/Private/MoveLibrary/MovementMixer.cpp
--- a/Engine/Plugins/Experimental/Mover/Source/Mover/Private/MoveLibrary/MovementMixer.cpp
+++ b/Engine/Plugins/Experimental/Mover/Source/Mover/Private/MoveLibrary/MovementMixer.cpp
@@ -17,7 +17,7 @@ void UMovementMixer::MixLayeredMove(const FLayeredMoveBase& ActiveMove, const FP
 {
 	if (OutCumulativeMove.PreferredMode != MoveStep.PreferredMode && !OutCumulativeMove.PreferredMode.IsNone() && !MoveStep.PreferredMode.IsNone())
 	{
-		UE_LOG(LogMover, Log, TEXT("Multiple LayeredMoves are conflicting with preferred moves. %s will override %s"),
+		UE_LOGF(LogMover, Log, "Multiple LayeredMoves are conflicting with preferred moves. %ls will override %ls",
 			*MoveStep.PreferredMode.ToString(), *OutCumulativeMove.PreferredMode.ToString());
 	}
 
@@ -25,7 +25,7 @@ void UMovementMixer::MixLayeredMove(const FLayeredMoveBase& ActiveMove, const FP
 	{
 		if (OutCumulativeMove.bHasDirIntent)
 		{
-			UE_LOG(LogMover, Log, TEXT("Multiple LayeredMoves are setting direction intent and the layered move with highest priority will be used."));
+			UE_LOGF(LogMover, Log, "Multiple LayeredMoves are setting direction intent and the layered move with highest priority will be used.");
 		}
 				
 		OutCumulativeMove.bHasDirIntent = MoveStep.bHasDirIntent;
@@ -38,7 +38,7 @@ void UMovementMixer::MixLayeredMove(const FLayeredMoveBase& ActiveMove, const FP
 		{
 			if (OutCumulativeMove.MixMode == EMoveMixMode::OverrideVelocity || OutCumulativeMove.MixMode == EMoveMixMode::OverrideAll)
 			{
-				UE_LOG(LogMover, Log, TEXT("Multiple LayeredMoves with Override mix mode are active simultaneously. Layered move with the highest priority will take effect."));
+				UE_LOGF(LogMover, Log, "Multiple LayeredMoves with Override mix mode are active simultaneously. Layered move with the highest priority will take effect.");
 			}
 
 			if (!MoveStep.PreferredMode.IsNone() && OutCumulativeMove.MixMode != EMoveMixMode::OverrideAll)
@@ -70,7 +70,7 @@ void UMovementMixer::MixLayeredMove(const FLayeredMoveBase& ActiveMove, const FP
 		{
 			if (OutCumulativeMove.MixMode == EMoveMixMode::OverrideVelocity || OutCumulativeMove.MixMode == EMoveMixMode::OverrideAll)
 			{
-				UE_LOG(LogMover, Log, TEXT("Multiple LayeredMoves with Override mix mode are active simultaneously. Layered move with the highest priority will take effect."));
+				UE_LOGF(LogMover, Log, "Multiple LayeredMoves with Override mix mode are active simultaneously. Layered move with the highest priority will take effect.");
 			}
 				
 			OutCumulativeMove = MoveStep;
@@ -83,7 +83,7 @@ void UMovementMixer::MixLayeredMove(const FLayeredMoveBase& ActiveMove, const FP
 		{
 			if (OutCumulativeMove.MixMode == EMoveMixMode::OverrideVelocity || OutCumulativeMove.MixMode == EMoveMixMode::OverrideAll || OutCumulativeMove.MixMode == EMoveMixMode::OverrideAllExceptVerticalVelocity)
 			{
-				UE_LOG(LogMover, Log, TEXT("Multiple LayeredMoves with Override mix mode are active simultaneously. Layered move with the highest priority will take effect."));
+				UE_LOGF(LogMover, Log, "Multiple LayeredMoves with Override mix mode are active simultaneously. Layered move with the highest priority will take effect.");
 			}
 
 			OutCumulativeMove = MoveStep;
@@ -100,7 +100,7 @@ void UMovementMixer::MixLayeredMove(const FLayeredMoveInstance& ActiveMove, cons
 {
 	if (OutCumulativeMove.PreferredMode != MoveStep.PreferredMode && !OutCumulativeMove.PreferredMode.IsNone() && !MoveStep.PreferredMode.IsNone())
 	{
-		UE_LOG(LogMover, Log, TEXT("Multiple LayeredMoves are conflicting with preferred moves. %s will override %s"),
+		UE_LOGF(LogMover, Log, "Multiple LayeredMoves are conflicting with preferred moves. %ls will override %ls",
 			*MoveStep.PreferredMode.ToString(), *OutCumulativeMove.PreferredMode.ToString());
 	}
 
@@ -108,7 +108,7 @@ void UMovementMixer::MixLayeredMove(const FLayeredMoveInstance& ActiveMove, cons
 	{
 		if (OutCumulativeMove.bHasDirIntent)
 		{
-			UE_LOG(LogMover, Log, TEXT("Multiple LayeredMoves are setting direction intent and the layered move with highest priority will be used."));
+			UE_LOGF(LogMover, Log, "Multiple LayeredMoves are setting direction intent and the layered move with highest priority will be used.");
 		}
 				
 		OutCumulativeMove.bHasDirIntent = MoveStep.bHasDirIntent;
@@ -121,7 +121,7 @@ void UMovementMixer::MixLayeredMove(const FLayeredMoveInstance& ActiveMove, cons
 		{
 			if (OutCumulativeMove.MixMode == EMoveMixMode::OverrideVelocity || OutCumulativeMove.MixMode == EMoveMixMode::OverrideAll)
 			{
-				UE_LOG(LogMover, Log, TEXT("Multiple LayeredMoves with Override mix mode are active simultaneously. Layered move with the highest priority will take effect."));
+				UE_LOGF(LogMover, Log, "Multiple LayeredMoves with Override mix mode are active simultaneously. Layered move with the highest priority will take effect.");
 			}
 
 			if (!MoveStep.PreferredMode.IsNone() && OutCumulativeMove.MixMode != EMoveMixMode::OverrideAll)
@@ -153,7 +153,7 @@ void UMovementMixer::MixLayeredMove(const FLayeredMoveInstance& ActiveMove, cons
 		{
 			if (OutCumulativeMove.MixMode == EMoveMixMode::OverrideVelocity || OutCumulativeMove.MixMode == EMoveMixMode::OverrideAll)
 			{
-				UE_LOG(LogMover, Log, TEXT("Multiple LayeredMoves with Override mix mode are active simultaneously. Layered move with the highest priority will take effect."));
+				UE_LOGF(LogMover, Log, "Multiple LayeredMoves with Override mix mode are active simultaneously. Layered move with the highest priority will take effect.");
 			}
 				
 			OutCumulativeMove = MoveStep;
@@ -166,7 +166,7 @@ void UMovementMixer::MixLayeredMove(const FLayeredMoveInstance& ActiveMove, cons
 		{
 			if (OutCumulativeMove.MixMode == EMoveMixMode::OverrideVelocity || OutCumulativeMove.MixMode == EMoveMixMode::OverrideAll || OutCumulativeMove.MixMode == EMoveMixMode::OverrideAllExceptVerticalVelocity)
 			{
-				UE_LOG(LogMover, Log, TEXT("Multiple LayeredMoves with Override mix mode are active simultaneously. Layered move with the highest priority will take effect."));
+				UE_LOGF(LogMover, Log, "Multiple LayeredMoves with Override mix mode are active simultaneously. Layered move with the highest priority will take effect.");
 			}
 
 			OutCumulativeMove = MoveStep;
diff --git a/Engine/Plugins/Experimental/Mover/Source/Mover/Private/MoveLibrary/MovementUtils.cpp b/Engine/Plugins/Experimental/Mover/Source/Mover/Private/MoveLibrary/MovementUtils.cpp
--- a/Engine/Plugins/Experimental/Mover/Source/Mover/Private/MoveLibrary/MovementUtils.cpp
+++ b/Engine/Plugins/Experimental/Mover/Source/Mover/Private/MoveLibrary/MovementUtils.cpp
@@ -391,7 +391,7 @@ bool UMovementUtils::TrySafeMoveUpdatedComponent(const FMovingComponentSet& Movi
 
 		if (UpdatedComponent)
 		{
-			UE_LOG(LogMover, VeryVerbose, TEXT("TrySafeMove: %s (role %i) Delta=%s DidMove=%i"), 
+			UE_LOGF(LogMover, VeryVerbose, "TrySafeMove: %ls (role %i) Delta=%ls DidMove=%i", 
 				*GetNameSafe(UpdatedComponent->GetOwner()), UpdatedComponent->GetOwnerRole(), *Delta.ToCompactString(), bMoveResult);
 		}
 	}
@@ -407,7 +407,7 @@ bool UMovementUtils::TrySafeMoveUpdatedComponent(const FMovingComponentSet& Movi
 			// Retry original move
 			bMoveResult = TryMoveUpdatedComponent_Internal(MovingComps, Delta, NewRotation, bSweep, MoveComponentFlags, &OutHit, Teleport);
 
-			UE_LOG(LogMover, VeryVerbose, TEXT("TrySafeMove retry: %s (role %i) Delta=%s DidMove=%i"),
+			UE_LOGF(LogMover, VeryVerbose, "TrySafeMove retry: %ls (role %i) Delta=%ls DidMove=%i",
 				*GetNameSafe(UpdatedComponent->GetOwner()), UpdatedComponent->GetOwnerRole(), *Delta.ToCompactString(), bMoveResult);
 		}
 	}
@@ -475,7 +475,7 @@ bool UMovementUtils::TryMoveToResolvePenetration(const FMovingComponentSet& Movi
 			// Move without sweeping.
 			const bool bDidMove = TryMoveUpdatedComponent_Internal(UpdatedComponent, Adjustment, NewRotationQuat, false, MoveComponentFlags, nullptr, ETeleportType::TeleportPhysics);
 
-			UE_LOG(LogMover, VeryVerbose, TEXT("TryMoveToResolvePenetration unencroached: %s (role %i) Adjustment=%s DidMove=%i"),
+			UE_LOGF(LogMover, VeryVerbose, "TryMoveToResolvePenetration unencroached: %ls (role %i) Adjustment=%ls DidMove=%i",
 				*GetNameSafe(UpdatedComponent->GetOwner()), UpdatedComponent->GetOwnerRole(), *Adjustment.ToCompactString(), bDidMove);
 
 			MoveRecord.Append( FMovementSubstep(PenetrationResolutionSubstepName, UpdatedComponent->GetComponentLocation()-OriginalCompPos, false) );
@@ -490,7 +490,7 @@ bool UMovementUtils::TryMoveToResolvePenetration(const FMovingComponentSet& Movi
 			FHitResult SweepOutHit(1.f);
 			bool bMoved = TryMoveUpdatedComponent_Internal(MovingComps, Adjustment, NewRotationQuat, true, MoveComponentFlags, &SweepOutHit, ETeleportType::TeleportPhysics);
 
-			UE_LOG(LogMover, VeryVerbose, TEXT("TryMoveToResolvePenetration: %s (role %i) Adjustment=%s DidMove=%i"),
+			UE_LOGF(LogMover, VeryVerbose, "TryMoveToResolvePenetration: %ls (role %i) Adjustment=%ls DidMove=%i",
 				*GetNameSafe(UpdatedComponent->GetOwner()), UpdatedComponent->GetOwnerRole(), *Adjustment.ToCompactString(), bMoved);
 
 			// Still stuck?
@@ -503,7 +503,7 @@ bool UMovementUtils::TryMoveToResolvePenetration(const FMovingComponentSet& Movi
 				{
 					bMoved = TryMoveUpdatedComponent_Internal(MovingComps, CombinedMTD, NewRotationQuat, true, MoveComponentFlags, nullptr, ETeleportType::TeleportPhysics);
 
-					UE_LOG(LogMover, VeryVerbose, TEXT("TryMoveToResolvePenetration combined: %s (role %i) CombinedAdjustment=%s DidMove=%i"),
+					UE_LOGF(LogMover, VeryVerbose, "TryMoveToResolvePenetration combined: %ls (role %i) CombinedAdjustment=%ls DidMove=%i",
 						*GetNameSafe(UpdatedComponent->GetOwner()), UpdatedComponent->GetOwnerRole(), *CombinedMTD.ToCompactString(), bMoved);
 				}
 			}
@@ -518,7 +518,7 @@ bool UMovementUtils::TryMoveToResolvePenetration(const FMovingComponentSet& Movi
 					const FVector AdjustAndMoveDelta = Adjustment + MoveDelta;
 					bMoved = TryMoveUpdatedComponent_Internal(MovingComps, AdjustAndMoveDelta, NewRotationQuat, true, MoveComponentFlags, nullptr, ETeleportType::TeleportPhysics);
 
-					UE_LOG(LogMover, VeryVerbose, TEXT("TryMoveToResolvePenetration multiple: %s (role %i) AdjustAndMoveDelta=%s DidMove=%i"),
+					UE_LOGF(LogMover, VeryVerbose, "TryMoveToResolvePenetration multiple: %ls (role %i) AdjustAndMoveDelta=%ls DidMove=%i",
 						*GetNameSafe(UpdatedComponent->GetOwner()), UpdatedComponent->GetOwnerRole(), *AdjustAndMoveDelta.ToCompactString(), bMoved);
 				}
 			}
@@ -899,7 +899,7 @@ bool UMovementUtils::TestEncroachment_Internal(const UWorld* World, const AActor
 			}
 			else
 			{
-				UE_LOG(LogMover, Log, TEXT("Components must be registered in order to be used in a ComponentOverlapMulti call. PriComp: %s TestActor: %s"), *PrimComp->GetName(), *TestActor->GetName());
+				UE_LOGF(LogMover, Log, "Components must be registered in order to be used in a ComponentOverlapMulti call. PriComp: %ls TestActor: %ls", *PrimComp->GetName(), *TestActor->GetName());
 				return false;
 			}
 		}
@@ -952,7 +952,7 @@ bool UMovementUtils::TestEncroachmentWithAdjustment_Internal(const UWorld* World
 			}
 			else
 			{
-				UE_LOG(LogMover, Log, TEXT("Components must be registered in order to be used in a ComponentOverlapMulti call. PriComp: %s TestActor: %s"), *PrimComp->GetName(), *TestActor->GetName());
+				UE_LOGF(LogMover, Log, "Components must be registered in order to be used in a ComponentOverlapMulti call. PriComp: %ls TestActor: %ls", *PrimComp->GetName(), *TestActor->GetName());
 			}
 		}
 		else
@@ -987,7 +987,7 @@ bool UMovementUtils::TestEncroachmentWithAdjustment_Internal(const UWorld* World
 					}
 					else
 					{
-						UE_LOG(LogMover, Log, TEXT("OverlapTest says we are overlapping, yet MTD says we're not. Something is wrong"));
+						UE_LOGF(LogMover, Log, "OverlapTest says we are overlapping, yet MTD says we're not. Something is wrong");
 						// It's not safe to use a partial result, that could push us out to an invalid location (like the other side of a wall).
 						OutProposedAdjustment = FVector::ZeroVector;
 						return true;
@@ -1007,7 +1007,7 @@ bool UMovementUtils::TestEncroachmentWithAdjustment_Internal(const UWorld* World
 						else
 						{
 							// Ignore this overlap.
-							UE_LOG(LogMover, Log, TEXT("OverlapTest says we are overlapping, yet MTD says we're not (with smaller shape). Ignoring this overlap."));
+							UE_LOGF(LogMover, Log, "OverlapTest says we are overlapping, yet MTD says we're not (with smaller shape). Ignoring this overlap.");
 							NumBlockingHits--;
 							continue;
 						}
diff --git a/Engine/Plugins/Experimental/Mover/Source/Mover/Private/MoveLibrary/RollbackBlackboard.cpp b/Engine/Plugins/Experimental/Mover/Source/Mover/Private/MoveLibrary/RollbackBlackboard.cpp
--- a/Engine/Plugins/Experimental/Mover/Source/Mover/Private/MoveLibrary/RollbackBlackboard.cpp
+++ b/Engine/Plugins/Experimental/Mover/Source/Mover/Private/MoveLibrary/RollbackBlackboard.cpp
@@ -170,7 +170,7 @@ void URollbackBlackboard::BeginRollback(const FMoverTimeStep& NewBaseTimeStep)
 	InRollbackThreadId = FPlatformTLS::GetCurrentThreadId();
 	bIsRollbackInProgress = true;
 
-	UE_LOG(LogMover, Verbose, TEXT("Blackboard begin rollback. From Sim F %i / T %.3f -> F %i / T %.3f"), 
+	UE_LOGF(LogMover, Verbose, "Blackboard begin rollback. From Sim F %i / T %.3f -> F %i / T %.3f", 
 		CurrentSimTimeStamp.Frame, CurrentSimTimeStamp.TimeMs, NewBaseTimeStep.ServerFrame, NewBaseTimeStep.BaseSimTimeMs);
 
 	// TODO: Need locking mechanism
diff --git a/Engine/Plugins/Experimental/Mover/Source/Mover/Private/MovementMode.cpp b/Engine/Plugins/Experimental/Mover/Source/Mover/Private/MovementMode.cpp
--- a/Engine/Plugins/Experimental/Mover/Source/Mover/Private/MovementMode.cpp
+++ b/Engine/Plugins/Experimental/Mover/Source/Mover/Private/MovementMode.cpp
@@ -37,7 +37,7 @@ void UBaseMovementMode::OnRegistered(const FName ModeName)
 		}
 		else
 		{
-			UE_LOG(LogMover, Error, TEXT("Invalid or missing transition object on mode of type %s of component %s for actor %s"), *GetPathNameSafe(this), *GetNameSafe(GetOuter()), *GetNameSafe(GetOutermost()));
+			UE_LOGF(LogMover, Error, "Invalid or missing transition object on mode of type %ls of component %ls for actor %ls", *GetPathNameSafe(this), *GetNameSafe(GetOuter()), *GetNameSafe(GetOutermost()));
 		}
 	}
 	
@@ -54,7 +54,7 @@ void UBaseMovementMode::OnUnregistered()
 		}
 		else
 		{
-			UE_LOG(LogMover, Error, TEXT("Invalid or missing transition object on mode of type %s of component %s for actor %s"), *GetPathNameSafe(this), *GetNameSafe(GetOuter()), *GetNameSafe(GetOutermost()));
+			UE_LOGF(LogMover, Error, "Invalid or missing transition object on mode of type %ls of component %ls for actor %ls", *GetPathNameSafe(this), *GetNameSafe(GetOuter()), *GetNameSafe(GetOutermost()));
 		}
 	}
 
diff --git a/Engine/Plugins/Experimental/Mover/Source/Mover/Private/MovementModeStateMachine.cpp b/Engine/Plugins/Experimental/Mover/Source/Mover/Private/MovementModeStateMachine.cpp
--- a/Engine/Plugins/Experimental/Mover/Source/Mover/Private/MovementModeStateMachine.cpp
+++ b/Engine/Plugins/Experimental/Mover/Source/Mover/Private/MovementModeStateMachine.cpp
@@ -143,7 +143,7 @@ void UMovementModeStateMachine::QueueNextModeInternal(FName DesiredNextModeName,
 		if ((NextModeName != NAME_None) &&
 		    (NextModeName != DesiredNextModeName || bShouldReenter != bShouldNextModeReenter))
 		{
-			UE_LOG(LogMover, Log, TEXT("%s (%s) Overwriting of queued mode change (%s, reenter: %i) with (%s, reenter: %i)"), *GetNameSafe(GetOwnerActor()), *UEnum::GetValueAsString(GetOwnerActor()->GetLocalRole()), *NextModeName.ToString(), bShouldNextModeReenter, *DesiredNextModeName.ToString(), bShouldReenter);
+			UE_LOGF(LogMover, Log, "%ls (%ls) Overwriting of queued mode change (%ls, reenter: %i) with (%ls, reenter: %i)", *GetNameSafe(GetOwnerActor()), *UEnum::GetValueAsString(GetOwnerActor()->GetLocalRole()), *NextModeName.ToString(), bShouldNextModeReenter, *DesiredNextModeName.ToString(), bShouldReenter);
 		}
 
 		if (Modes.Contains(DesiredNextModeName))
@@ -152,7 +152,7 @@ void UMovementModeStateMachine::QueueNextModeInternal(FName DesiredNextModeName,
 		}
 		else
 		{
-			UE_LOG(LogMover, Warning, TEXT("Attempted to queue an unregistered movement mode: %s on owner %s"), *DesiredNextModeName.ToString(), *GetNameSafe(GetOwnerActor()));
+			UE_LOGF(LogMover, Warning, "Attempted to queue an unregistered movement mode: %ls on owner %ls", *DesiredNextModeName.ToString(), *GetNameSafe(GetOwnerActor()));
 		}
 	}
 }
@@ -429,7 +429,7 @@ void UMovementModeStateMachine::OnSimulationTick(USceneComponent* UpdatedCompone
 				// if we've done this sub step a lot before go ahead and just advance time to avoid freezing editor
 				if (NumConsecutiveFullRefundedSubsteps >= MaxConsecutiveFullRefundedSubsteps)
 				{
-					UE_LOG(LogMover, Warning, TEXT("Movement mode %s and %s on %s are stuck giving time back to each other. Overriding to advance to next substep."), *CurrentModeName.ToString(), *OutputState.MovementEndState.NextModeName.ToString(), *MoverComp->GetOwner()->GetName());
+					UE_LOGF(LogMover, Warning, "Movement mode %ls and %ls on %ls are stuck giving time back to each other. Overriding to advance to next substep.", *CurrentModeName.ToString(), *OutputState.MovementEndState.NextModeName.ToString(), *MoverComp->GetOwner()->GetName());
 					TotalUsedMs += SubTimeStep.StepMs;
 				}
 			}
@@ -747,7 +747,7 @@ void UMovementModeStateMachine::AdvanceToNextMode()
 		(CurrentModeName != NextModeName || bShouldNextModeReenter))
 	{
 		const AActor* OwnerActor = GetOwnerActor();
-		UE_LOG(LogMover, Verbose, TEXT("AdvanceToNextMode: %s (%s) from %s to %s"), 
+		UE_LOGF(LogMover, Verbose, "AdvanceToNextMode: %ls (%ls) from %ls to %ls", 
 			*GetNameSafe(OwnerActor), *UEnum::GetValueAsString(OwnerActor->GetLocalRole()), *CurrentModeName.ToString(), *NextModeName.ToString());
 
 		const FName PreviousModeName = CurrentModeName;
@@ -881,7 +881,7 @@ void UMovementModeStateMachine::RollbackModifiers(const FMoverSyncState* Invalid
 		
 			if (!bContainsModifier)
 			{
-				UE_LOG(LogMover, Log, TEXT("Modifier(%s) was started on %s after a rollback."), *ModifierFromRollbackIt->Get()->ToSimpleString(), *GetNameSafe(MoverComp->GetOwner()));
+				UE_LOGF(LogMover, Log, "Modifier(%ls) was started on %ls after a rollback.", *ModifierFromRollbackIt->Get()->ToSimpleString(), *GetNameSafe(MoverComp->GetOwner()));
 				ModifierFromRollbackIt->Get()->OnStart(MoverComp, MoverComp->GetLastTimeStep(), *SyncState, *AuxState);
 			}
 		}
@@ -900,7 +900,7 @@ void UMovementModeStateMachine::RollbackModifiers(const FMoverSyncState* Invalid
 	
 			if (!bContainsModifier)
 			{
-				UE_LOG(LogMover, Log, TEXT("Modifier(%s) was ended on %s after a rollback."), *ModifierFromCacheIt->Get()->ToSimpleString(), *GetNameSafe(MoverComp->GetOwner()));
+				UE_LOGF(LogMover, Log, "Modifier(%ls) was ended on %ls after a rollback.", *ModifierFromCacheIt->Get()->ToSimpleString(), *GetNameSafe(MoverComp->GetOwner()));
 				ModifierFromCacheIt->Get()->OnEnd(MoverComp, MoverComp->GetLastTimeStep(), *SyncState, *AuxState);
 			}
 		}
diff --git a/Engine/Plugins/Experimental/Mover/Source/Mover/Private/MovementModifier.cpp b/Engine/Plugins/Experimental/Mover/Source/Mover/Private/MovementModifier.cpp
--- a/Engine/Plugins/Experimental/Mover/Source/Mover/Private/MovementModifier.cpp
+++ b/Engine/Plugins/Experimental/Mover/Source/Mover/Private/MovementModifier.cpp
@@ -121,7 +121,7 @@ void FMovementModifierGroup::QueueMovementModifier(TSharedPtr<FMovementModifierB
 	if (ensure(Modifier.IsValid()))
 	{
 		QueuedModifiers.Add(Modifier);
-		UE_LOG(LogMover, VeryVerbose, TEXT("Queued Movement Modifier (%s)"), *Modifier->ToSimpleString());
+		UE_LOGF(LogMover, VeryVerbose, "Queued Movement Modifier (%ls)", *Modifier->ToSimpleString());
 	}
 }
 
@@ -217,7 +217,7 @@ static void CopyModifierArray(TArray<TSharedPtr<FMovementModifierBase>>& Dest, c
 			}
 			else
 			{
-				UE_LOG(LogMover, Warning, TEXT("CopyModifierArray trying to copy invalid Other Modifier"));
+				UE_LOGF(LogMover, Warning, "CopyModifierArray trying to copy invalid Other Modifier");
 			}
 		}
 	}
@@ -531,14 +531,14 @@ struct FMovementModifierDeleter
 			}
 			else
 			{
-				UE_LOG(LogMover, Error, TEXT("FMovementModifierGroup::NetSerialize: ScriptStruct not derived from FMovementModifierBase attempted to serialize."));
+				UE_LOGF(LogMover, Error, "FMovementModifierGroup::NetSerialize: ScriptStruct not derived from FMovementModifierBase attempted to serialize.");
 				Ar.SetError();
 				break;
 			}
 		}
 		else if (ScriptStruct.IsError())
 		{
-			UE_LOG(LogMover, Error, TEXT("FMovementModifierGroup::NetSerialize: Invalid ScriptStruct serialized."));
+			UE_LOGF(LogMover, Error, "FMovementModifierGroup::NetSerialize: Invalid ScriptStruct serialized.");
 			Ar.SetError();
 			break;
 		}
diff --git a/Engine/Plugins/Experimental/Mover/Source/Mover/Private/MoverComponent.cpp b/Engine/Plugins/Experimental/Mover/Source/Mover/Private/MoverComponent.cpp
--- a/Engine/Plugins/Experimental/Mover/Source/Mover/Private/MoverComponent.cpp
+++ b/Engine/Plugins/Experimental/Mover/Source/Mover/Private/MoverComponent.cpp
@@ -123,7 +123,7 @@ void UMoverComponent::InitializeComponent()
 		}
 		else
 		{
-			UE_LOG(LogMover, Error, TEXT("No backend class set on %s. Mover actor will not function."), *GetNameSafe(GetOwner()));
+			UE_LOGF(LogMover, Error, "No backend class set on %ls. Mover actor will not function.", *GetNameSafe(GetOwner()));
 		}
 
 
@@ -139,7 +139,7 @@ void UMoverComponent::InitializeComponent()
 		{
 			if (Element.Value.Get() == nullptr)
 			{
-				UE_LOG(LogMover, Warning, TEXT("Invalid Movement Mode type '%s' detected on %s. Mover actor will not function correctly."),
+				UE_LOGF(LogMover, Warning, "Invalid Movement Mode type '%ls' detected on %ls. Mover actor will not function correctly.",
 					*Element.Key.ToString(), *GetNameSafe(GetOwner()));
 				continue;
 			}
@@ -646,7 +646,7 @@ void UMoverComponent::SimulationTick(const FMoverTimeStep& InTimeStep, const FMo
 			if (UpdatedComponent->GetComponentLocation().Equals(OutputSyncState->GetLocation_WorldSpace()) == false ||
 				UpdatedComponent->GetComponentQuat().Equals(OutputSyncState->GetOrientation_WorldSpace().Quaternion(), UE_KINDA_SMALL_NUMBER) == false)
 			{
-				UE_LOG(LogMover, Warning, TEXT("Detected pos/rot difference between Mover actor (%s) sync state and scene component after sim ticking. This indicates a movement mode may not be authoring the final state correctly."), *GetNameSafe(UpdatedComponent->GetOwner()));
+				UE_LOGF(LogMover, Warning, "Detected pos/rot difference between Mover actor (%ls) sync state and scene component after sim ticking. This indicates a movement mode may not be authoring the final state correctly.", *GetNameSafe(UpdatedComponent->GetOwner()));
 			}
 		}
 	}
@@ -655,7 +655,7 @@ void UMoverComponent::SimulationTick(const FMoverTimeStep& InTimeStep, const FMo
 
 	if (!SimOutput.MoveRecord.GetTotalMoveDelta().IsZero())
 	{
-		UE_LOG(LogMover, VeryVerbose, TEXT("KinematicSimTick: %s (role %i) frame %d: %s"),
+		UE_LOGF(LogMover, VeryVerbose, "KinematicSimTick: %ls (role %i) frame %d: %ls",
 			*GetNameSafe(UpdatedComponent->GetOwner()), UpdatedComponent->GetOwnerRole(), MoverTimeStep.ServerFrame, *SimOutput.MoveRecord.ToString());
 	}
 
@@ -1456,7 +1456,7 @@ void UMoverComponent::RefreshSharedSettings()
 			{
 				if (SharedSettingsType.Get() == nullptr)
 				{
-					UE_LOG(LogMover, Warning, TEXT("Invalid shared setting class detected on Movement Mode %s."), *Mode->GetName());
+					UE_LOGF(LogMover, Warning, "Invalid shared setting class detected on Movement Mode %ls.", *Mode->GetName());
 					continue;
 				}
 
@@ -1489,7 +1489,7 @@ void UMoverComponent::RefreshSharedSettings()
 				{
 					if (SharedSettingsType.Get() == nullptr)
 					{
-						UE_LOG(LogMover, Warning, TEXT("Invalid shared setting class detected on Transition on Movement Mode %s."), *Mode->GetName());
+						UE_LOGF(LogMover, Warning, "Invalid shared setting class detected on Transition on Movement Mode %ls.", *Mode->GetName());
 						continue;
 					}
 
@@ -1525,7 +1525,7 @@ void UMoverComponent::RefreshSharedSettings()
 		{
 			if (SharedSettingsType.Get() == nullptr)
 			{
-				UE_LOG(LogMover, Warning, TEXT("Invalid shared setting class detected on Transition."));
+				UE_LOGF(LogMover, Warning, "Invalid shared setting class detected on Transition.");
 				continue;
 			}
 
@@ -1857,7 +1857,7 @@ void UMoverComponent::QueueInstantMovementEffect(const FScheduledInstantMovement
 
 #if !defined(BUILD_SHIPPING) || !BUILD_SHIPPING
 	ENetMode NetMode = GetWorld() ? GetWorld()->GetNetMode() : NM_MAX;
-	UE_LOG(LogMover, Verbose, TEXT("(%s) UMoverComponent::QueueInstantMovementEffect: Game Thread queueing an instant movement effect scheduled for frame %d: %s."),
+	UE_LOGF(LogMover, Verbose, "(%ls) UMoverComponent::QueueInstantMovementEffect: Game Thread queueing an instant movement effect scheduled for frame %d: %ls.",
 		*ToString(NetMode), InstantMovementEffect.ExecutionServerFrame, InstantMovementEffect.Effect.IsValid() ? *InstantMovementEffect.Effect->ToSimpleString() : TEXT("INVALID INSTANT EFFECT"));
 #endif // !defined(BUILD_SHIPPING) || !BUILD_SHIPPING
 }
@@ -1966,12 +1966,12 @@ UBaseMovementMode* UMoverComponent::AddMovementModeFromClass(FName ModeName, TSu
 {
 	if (!MovementMode)
 	{
-		UE_LOG(LogMover, Warning, TEXT("Attempted to add a movement mode that wasn't valid. AddMovementModeFromClass will not add anything. (%s)"), *GetNameSafe(GetOwner()));
+		UE_LOGF(LogMover, Warning, "Attempted to add a movement mode that wasn't valid. AddMovementModeFromClass will not add anything. (%ls)", *GetNameSafe(GetOwner()));
 		return nullptr;
 	}
 	if (MovementMode->HasAnyClassFlags(CLASS_Abstract))
 	{
-		UE_LOG(LogMover, Warning, TEXT("The Movement Mode class (%s) is abstract and is not a valid class to instantiate. AddMovementModeFromClass will not do anything. (%s)"), *GetNameSafe(MovementMode), *GetNameSafe(GetOwner()));
+		UE_LOGF(LogMover, Warning, "The Movement Mode class (%ls) is abstract and is not a valid class to instantiate. AddMovementModeFromClass will not do anything. (%ls)", *GetNameSafe(MovementMode), *GetNameSafe(GetOwner()));
 		return nullptr;
 	}
 
@@ -1983,7 +1983,7 @@ bool UMoverComponent::AddMovementModeFromObject(FName ModeName, UBaseMovementMod
 {
 	if (!ModeFSM)
 	{
-		UE_LOG(LogMover, Warning, TEXT("Attempted to add %s movement mode before Movement Mode State Machine created %s. AddMovementModeFromObject did nothing. Check if mover component is initialized? (%s)"), *GetNameSafe(MovementMode), *GetNameSafe(this), *GetNameSafe(GetOwner()));
+		UE_LOGF(LogMover, Warning, "Attempted to add %ls movement mode before Movement Mode State Machine created %ls. AddMovementModeFromObject did nothing. Check if mover component is initialized? (%ls)", *GetNameSafe(MovementMode), *GetNameSafe(this), *GetNameSafe(GetOwner()));
 		return false;
 	}
 
@@ -1991,7 +1991,7 @@ bool UMoverComponent::AddMovementModeFromObject(FName ModeName, UBaseMovementMod
 	{
 		if (MovementMode->GetClass()->HasAnyClassFlags(CLASS_Abstract))
 		{
-			UE_LOG(LogMover, Warning, TEXT("The Movement Mode class (%s) is abstract and is not a valid class to instantiate. AddMovementModeFromObject will not do anything. (%s)"), *GetNameSafe(MovementMode), *GetNameSafe(GetOwner()));
+			UE_LOGF(LogMover, Warning, "The Movement Mode class (%ls) is abstract and is not a valid class to instantiate. AddMovementModeFromObject will not do anything. (%ls)", *GetNameSafe(MovementMode), *GetNameSafe(GetOwner()));
 			return false;
 		}
 		
@@ -1999,7 +1999,7 @@ bool UMoverComponent::AddMovementModeFromObject(FName ModeName, UBaseMovementMod
 		{
 			if (FoundMovementMode->Get()->GetClass() == MovementMode->GetClass())
 			{
-				UE_LOG(LogMover, Warning, TEXT("Added the same movement mode (%s) for a movement mode name (%s). AddMovementModeFromObject will add the mode but is likely unwanted/unnecessary behavior. (%s)"), *GetNameSafe(MovementMode), *ModeName.ToString(), *GetNameSafe(GetOwner()));
+				UE_LOGF(LogMover, Warning, "Added the same movement mode (%ls) for a movement mode name (%ls). AddMovementModeFromObject will add the mode but is likely unwanted/unnecessary behavior. (%ls)", *GetNameSafe(MovementMode), *ModeName.ToString(), *GetNameSafe(GetOwner()));
 			}
 
 			RemoveMovementMode(ModeName);
@@ -2007,7 +2007,7 @@ bool UMoverComponent::AddMovementModeFromObject(FName ModeName, UBaseMovementMod
 		
 		if (MovementMode->GetOuter() != this)
 		{
-			UE_LOG(LogMover, Verbose, TEXT("Movement modes are expected to be parented to the MoverComponent. The %s movement mode was reparented to %s! (%s)"), *GetNameSafe(MovementMode), *GetNameSafe(this), *GetNameSafe(GetOwner()));
+			UE_LOGF(LogMover, Verbose, "Movement modes are expected to be parented to the MoverComponent. The %ls movement mode was reparented to %ls! (%ls)", *GetNameSafe(MovementMode), *GetNameSafe(this), *GetNameSafe(GetOwner()));
 			MovementMode->Rename(nullptr, this, REN_DoNotDirty | REN_NonTransactional);
 		}
 		
@@ -2016,7 +2016,7 @@ bool UMoverComponent::AddMovementModeFromObject(FName ModeName, UBaseMovementMod
 	}
 	else
 	{
-		UE_LOG(LogMover, Warning, TEXT("Attempted to add %s movement mode that wasn't valid to %s. AddMovementModeFromObject did not add anything. (%s)"), *GetNameSafe(MovementMode), *GetNameSafe(this), *GetNameSafe(GetOwner()));
+		UE_LOGF(LogMover, Warning, "Attempted to add %ls movement mode that wasn't valid to %ls. AddMovementModeFromObject did not add anything. (%ls)", *GetNameSafe(MovementMode), *GetNameSafe(this), *GetNameSafe(GetOwner()));
 		return false;
 	}
 
@@ -2027,13 +2027,13 @@ bool UMoverComponent::RemoveMovementMode(FName ModeName)
 {
 	if (!ModeFSM)
 	{
-		UE_LOG(LogMover, Warning, TEXT("Attempted to remove %s movement mode before Movement Mode State Machine created %s. RemoveMovementMode did nothing. Check if mover component is initialized? (%s)"), *ModeName.ToString(), *GetNameSafe(this), *GetNameSafe(GetOwner()));
+		UE_LOGF(LogMover, Warning, "Attempted to remove %ls movement mode before Movement Mode State Machine created %ls. RemoveMovementMode did nothing. Check if mover component is initialized? (%ls)", *ModeName.ToString(), *GetNameSafe(this), *GetNameSafe(GetOwner()));
 		return false;
 	}
 
 	if (ModeFSM->GetCurrentModeName() == ModeName)
 	{
-		UE_LOG(LogMover, Warning, TEXT("The mode being removed (%s Movement Mode) is the mode this actor (%s) is currently in. It was removed but may cause issues. Consider waiting to remove the mode or queueing a different valid mode to avoid issues."), *ModeName.ToString(), *GetNameSafe(GetOwner()));
+		UE_LOGF(LogMover, Warning, "The mode being removed (%ls Movement Mode) is the mode this actor (%ls) is currently in. It was removed but may cause issues. Consider waiting to remove the mode or queueing a different valid mode to avoid issues.", *ModeName.ToString(), *GetNameSafe(GetOwner()));
 	}
 	
 	TObjectPtr<UBaseMovementMode>* ModeToRemove = MovementModes.Find(ModeName);
@@ -2266,7 +2266,7 @@ void UMoverComponent::SetUpDirectionOverride(bool bOverrideUpDirection, FVector
 	{
 		if (UpDirection.IsNearlyZero())
 		{
-			UE_LOG(LogMover, Warning, TEXT("Ignoring the provided UpDirection (%s) override because it is a zero vector. (%s)"), *UpDirection.ToString(), *GetNameSafe(GetOwner()));
+			UE_LOGF(LogMover, Warning, "Ignoring the provided UpDirection (%ls) override because it is a zero vector. (%ls)", *UpDirection.ToString(), *GetNameSafe(GetOwner()));
 			bHasGravityOverride = false;
 			return;
 		}
@@ -2447,7 +2447,7 @@ TArray<FTrajectorySampleInfo> UMoverComponent::GetPredictedTrajectory(FMoverPred
 		}
 	}
 
-	UE_LOG(LogMover, Warning, TEXT("Cannot predict trajectory on a MoverComponent that is not initialized. Returning defaulted samples. Actor: %s  Initialized? %i BegunPlay? %i BeingDestroyed? %i"), *GetNameSafe(GetOwner()), HasBeenInitialized(), HasBegunPlay(), IsBeingDestroyed());
+	UE_LOGF(LogMover, Warning, "Cannot predict trajectory on a MoverComponent that is not initialized. Returning defaulted samples. Actor: %ls  Initialized? %i BegunPlay? %i BeingDestroyed? %i", *GetNameSafe(GetOwner()), HasBeenInitialized(), HasBegunPlay(), IsBeingDestroyed());
 
 	TArray<FTrajectorySampleInfo> BlankDefaultSamples;
 	BlankDefaultSamples.AddDefaulted(PredictionParams.NumPredictionSamples);
@@ -2656,7 +2656,7 @@ bool UMoverComponent::MakeAndQueueLayeredMove(const TSubclassOf<ULayeredMoveLogi
 		}
 		else
 		{
-			UE_LOG(LogMover, Warning, TEXT("%s activation was queued on %s but the move was NOT queued since it did not have valid data. InstancedDataStructType on Move Logic needs to be a FLayeredMoveInstancedData or child struct of."),
+			UE_LOGF(LogMover, Warning, "%ls activation was queued on %ls but the move was NOT queued since it did not have valid data. InstancedDataStructType on Move Logic needs to be a FLayeredMoveInstancedData or child struct of.",
 				*MoveLogicClass->GetName(),
 				*GetOwner()->GetName());
 			
@@ -2665,7 +2665,7 @@ bool UMoverComponent::MakeAndQueueLayeredMove(const TSubclassOf<ULayeredMoveLogi
 	}
 	else
 	{
-		UE_LOG(LogMover, Warning, TEXT("%s activation was queued on %s and the move was not registered. Any move activated on a MoverComponent Needs to be Registered with the MoverCompoent. The layered move will not be queued for activation."),
+		UE_LOGF(LogMover, Warning, "%ls activation was queued on %ls and the move was not registered. Any move activated on a MoverComponent Needs to be Registered with the MoverCompoent. The layered move will not be queued for activation.",
 			*MoveLogicClass->GetName(),
 			*GetOwner()->GetName());
 			
@@ -2786,7 +2786,7 @@ void UMoverComponent::CheckForExternalMovement(const FMoverTickStartData& SimSta
 		{
 			if (bWarnOnExternalMovement)
 			{
-				UE_LOG(LogMover, Warning, TEXT("%s %s: Simulation start location (%s) disagrees with actual mover component location (%s). This indicates movement of the component out-of-band with the simulation, and may cause poor quality motion."),
+				UE_LOGF(LogMover, Warning, "%ls %ls: Simulation start location (%ls) disagrees with actual mover component location (%ls). This indicates movement of the component out-of-band with the simulation, and may cause poor quality motion.",
 					*GetNameSafe(GetOwner()),
 					*StaticEnum<ENetRole>()->GetValueAsString(GetOwnerRole()),
 					*StartingSyncState->GetLocation_WorldSpace().ToCompactString(),
diff --git a/Engine/Plugins/Experimental/Mover/Source/Mover/Private/MoverDataModelTypes.cpp b/Engine/Plugins/Experimental/Mover/Source/Mover/Private/MoverDataModelTypes.cpp
--- a/Engine/Plugins/Experimental/Mover/Source/Mover/Private/MoverDataModelTypes.cpp
+++ b/Engine/Plugins/Experimental/Mover/Source/Mover/Private/MoverDataModelTypes.cpp
@@ -331,7 +331,7 @@ void FMoverDefaultSyncState::SetTransforms_WorldSpace(FVector WorldLocation, FRo
 	{
 		if (Base)
 		{
-			UE_LOG(LogMover, Warning, TEXT("Failed to set base as %s. Falling back to world space movement"), *GetNameSafe(Base->GetOwner()));
+			UE_LOGF(LogMover, Warning, "Failed to set base as %ls. Falling back to world space movement", *GetNameSafe(Base->GetOwner()));
 		}
 
 		Location = WorldLocation;
diff --git a/Engine/Plugins/Experimental/Mover/Source/Mover/Private/MoverPoseSearchTrajectoryPredictor.cpp b/Engine/Plugins/Experimental/Mover/Source/Mover/Private/MoverPoseSearchTrajectoryPredictor.cpp
--- a/Engine/Plugins/Experimental/Mover/Source/Mover/Private/MoverPoseSearchTrajectoryPredictor.cpp
+++ b/Engine/Plugins/Experimental/Mover/Source/Mover/Private/MoverPoseSearchTrajectoryPredictor.cpp
@@ -12,7 +12,7 @@ void UMoverTrajectoryPredictor::Predict(FTransformTrajectory& InOutTrajectory,
 {
 	if (!MoverComponent)
 	{
-		UE_LOG(LogMover, Log, TEXT("Calling Predict without a Mover Component. This is invalid and the trajectory will not be modified."));
+		UE_LOGF(LogMover, Log, "Calling Predict without a Mover Component. This is invalid and the trajectory will not be modified.");
 		return;
 	}
 
@@ -37,7 +37,7 @@ void UMoverTrajectoryPredictor::Predict(UMoverComponent& MoverComponent, FTransf
 
 	if (InOutTrajectory.Samples.Num() < (NumHistorySamples + NumPredictionSamples))
 	{
-		UE_LOG(LogMover, Warning, TEXT("InOutTrajectory Samples array does not have enough space for %i predicted samples"), MoverPredictionSamples.Num());
+		UE_LOGF(LogMover, Warning, "InOutTrajectory Samples array does not have enough space for %i predicted samples", MoverPredictionSamples.Num());
 	}
 	else
 	{
@@ -86,7 +86,7 @@ void UMoverTrajectoryPredictor::GetGravity(FVector& OutGravityAccel)
 {
 	if (!MoverComponent)
 	{
-		UE_LOG(LogMover, Log, TEXT("Calling GetGravity without a Mover Component. Return value will be defaulted."));
+		UE_LOGF(LogMover, Log, "Calling GetGravity without a Mover Component. Return value will be defaulted.");
 		OutGravityAccel = FVector::ZeroVector;
 		return;
 	}
@@ -99,7 +99,7 @@ void UMoverTrajectoryPredictor::GetCurrentState(FVector& OutPosition, FQuat& Out
 {
 	if (!MoverComponent)
 	{
-		UE_LOG(LogMover, Log, TEXT("Calling GetCurrentState without a Mover Component. Return values will be defaulted."));
+		UE_LOGF(LogMover, Log, "Calling GetCurrentState without a Mover Component. Return values will be defaulted.");
 		OutPosition = OutVelocity = FVector::ZeroVector;
 		OutFacing = FQuat::Identity;
 		return;
@@ -148,7 +148,7 @@ void UMoverTrajectoryPredictor::GetVelocity(FVector& OutVelocity)
 {
 	if (!MoverComponent)
 	{
-		UE_LOG(LogMover, Log, TEXT("Calling GetVelocity without a Mover Component. Return value will be defaulted."));
+		UE_LOGF(LogMover, Log, "Calling GetVelocity without a Mover Component. Return value will be defaulted.");
 		OutVelocity = FVector::ZeroVector;
 		return;
 	}
diff --git a/Engine/Plugins/Experimental/Mover/Source/Mover/Private/MoverTypes.cpp b/Engine/Plugins/Experimental/Mover/Source/Mover/Private/MoverTypes.cpp
--- a/Engine/Plugins/Experimental/Mover/Source/Mover/Private/MoverTypes.cpp
+++ b/Engine/Plugins/Experimental/Mover/Source/Mover/Private/MoverTypes.cpp
@@ -219,7 +219,7 @@ FMoverDataCollection& FMoverDataCollection::operator=(const FMoverDataCollection
 				}
 				else
 				{
-					UE_LOG(LogMover, Warning, TEXT("FMoverDataCollection::operator= trying to copy invalid Other DataArray element"));
+					UE_LOGF(LogMover, Warning, "FMoverDataCollection::operator= trying to copy invalid Other DataArray element");
 				}
 			}
 		}
@@ -581,21 +581,21 @@ void FMoverDataCollection::NetSerializeDataArray(FArchive& Ar, UPackageMap* Map,
 
 				if (!bArrayElementSuccess)
 				{
-					UE_LOG(LogMover, Error, TEXT("FMoverDataCollection::NetSerialize: Failed to serialize ScriptStruct %s"), *ScriptStruct->GetName());
+					UE_LOGF(LogMover, Error, "FMoverDataCollection::NetSerialize: Failed to serialize ScriptStruct %ls", *ScriptStruct->GetName());
 					Ar.SetError();
 					break;
 				}
 			}
 			else
 			{
-				UE_LOG(LogMover, Error, TEXT("FMoverDataCollection::NetSerialize: ScriptStruct not derived from FMoverDataStructBase attempted to serialize."));
+				UE_LOGF(LogMover, Error, "FMoverDataCollection::NetSerialize: ScriptStruct not derived from FMoverDataStructBase attempted to serialize.");
 				Ar.SetError();
 				break;
 			}
 		}
 		else if (ScriptStruct.IsError())
 		{
-			UE_LOG(LogMover, Error, TEXT("FMoverDataCollection::NetSerialize: Invalid ScriptStruct serialized."));
+			UE_LOGF(LogMover, Error, "FMoverDataCollection::NetSerialize: Invalid ScriptStruct serialized.");
 			Ar.SetError();
 			break;
 		}
@@ -655,7 +655,7 @@ DEFINE_FUNCTION(UMoverDataCollectionLibrary::execK2_AddDataToCollection)
 			}
 			else
 			{
-				UE_LOG(LogMover, Warning, TEXT("AddDataToCollection: invalid struct type submitted: %s"), *SourceStructProp->Struct->GetName());
+				UE_LOGF(LogMover, Warning, "AddDataToCollection: invalid struct type submitted: %ls", *SourceStructProp->Struct->GetName());
 			}
 		}
 
diff --git a/Engine/Plugins/Experimental/Mover/Source/Mover/Private/PhysicsMover/PathedMovement/PathedMovementMode.cpp b/Engine/Plugins/Experimental/Mover/Source/Mover/Private/PhysicsMover/PathedMovement/PathedMovementMode.cpp
--- a/Engine/Plugins/Experimental/Mover/Source/Mover/Private/PhysicsMover/PathedMovement/PathedMovementMode.cpp
+++ b/Engine/Plugins/Experimental/Mover/Source/Mover/Private/PhysicsMover/PathedMovement/PathedMovementMode.cpp
@@ -149,7 +149,7 @@ void UPathedPhysicsMovementMode::OnPreSimulate_Internal(const FSimulationTickPar
 	if (CVarEnablePathedPhysicsMovementDebug->GetBool())
 	{
 		const Chaos::FRigidTransform3 TargetWorldTransform = Chaos::FRigidTransform3::MultiplyNoScale(TargetRelativeTransform, OutputMoveState.MutableProps.PathOrigin);
-		UE_LOG(LogMover, VeryVerbose, TEXT("Pathed mover [%s (%s)]: Setting target:\n\tRelative [%s]\n\tWorld [%s]"),
+		UE_LOGF(LogMover, VeryVerbose, "Pathed mover [%ls (%ls)]: Setting target:\n\tRelative [%ls]\n\tWorld [%ls]",
 			*GetPathedMoverComp().GetOwner()->GetName(),
 			GetPathedMoverComp().GetOwner()->GetNetMode() == NM_DedicatedServer ? TEXT("Server") : TEXT("Client"),
 			*TargetRelativeTransform.ToString(),
diff --git a/Engine/Plugins/Experimental/Mover/Source/Mover/Private/PhysicsMover/PhysicsCharacterMoverComponent.cpp b/Engine/Plugins/Experimental/Mover/Source/Mover/Private/PhysicsMover/PhysicsCharacterMoverComponent.cpp
--- a/Engine/Plugins/Experimental/Mover/Source/Mover/Private/PhysicsMover/PhysicsCharacterMoverComponent.cpp
+++ b/Engine/Plugins/Experimental/Mover/Source/Mover/Private/PhysicsMover/PhysicsCharacterMoverComponent.cpp
@@ -31,7 +31,7 @@ void UPhysicsCharacterMoverComponent::BeginPlay()
 
 	if (bHandleJump)
 	{
-		UE_LOG(LogMover, Warning, TEXT("Handle Jump flag is ignored for Physics Character Mover Component; jumps are handled via the Physics Jump Check Transition."));
+		UE_LOGF(LogMover, Warning, "Handle Jump flag is ignored for Physics Character Mover Component; jumps are handled via the Physics Jump Check Transition.");
 	}
 }
 
diff --git a/Engine/Plugins/Experimental/Mover/Source/Mover/Private/UserDefinedStructSupport.cpp b/Engine/Plugins/Experimental/Mover/Source/Mover/Private/UserDefinedStructSupport.cpp
--- a/Engine/Plugins/Experimental/Mover/Source/Mover/Private/UserDefinedStructSupport.cpp
+++ b/Engine/Plugins/Experimental/Mover/Source/Mover/Private/UserDefinedStructSupport.cpp
@@ -114,7 +114,7 @@ bool FMoverUserDefinedDataStruct::NetSerializeUserDefinedStruct(FInstancedStruct
 		if (!UDSStructInstance.IsValid())
 		{
 			NonConstScriptStruct = nullptr;
-			UE_LOG(LogMover, Error, TEXT("FMoverUserDefinedDataStruct::NetSerializeUserDefinedStruct: Bad script struct serialized when loading, cannot recover."));
+			UE_LOGF(LogMover, Error, "FMoverUserDefinedDataStruct::NetSerializeUserDefinedStruct: Bad script struct serialized when loading, cannot recover.");
 			Ar.SetError();
 		}
 	}
diff --git a/Engine/Plugins/Experimental/Mover/Source/Mover/Public/LayeredMoveBase.h b/Engine/Plugins/Experimental/Mover/Source/Mover/Public/LayeredMoveBase.h
--- a/Engine/Plugins/Experimental/Mover/Source/Mover/Public/LayeredMoveBase.h
+++ b/Engine/Plugins/Experimental/Mover/Source/Mover/Public/LayeredMoveBase.h
@@ -169,7 +169,7 @@ protected:
 #if UE_BUILD_SHIPPING || UE_BUILD_TEST
 		if (!ensure(InstancedDataStructType))
         {
-			UE_LOG(LogMover, Error, TEXT("Move Logic needs an active data struct type. If no data is needed consider using the default move data type."))
+			UE_LOGF(LogMover, Error, "Move Logic needs an active data struct type. If no data is needed consider using the default move data type.")
         	static MoveDataT GarbageData;
         	return GarbageData;
         }
diff --git a/Engine/Plugins/Experimental/Mover/Source/Mover/Public/MoveLibrary/RollbackBlackboard.h b/Engine/Plugins/Experimental/Mover/Source/Mover/Public/MoveLibrary/RollbackBlackboard.h
--- a/Engine/Plugins/Experimental/Mover/Source/Mover/Public/MoveLibrary/RollbackBlackboard.h
+++ b/Engine/Plugins/Experimental/Mover/Source/Mover/Public/MoveLibrary/RollbackBlackboard.h
@@ -342,7 +342,7 @@ public:
 	{
 		if (EntryMap.Contains(EntryName))
 		{
-			UE_LOG(LogMover, Warning, TEXT("Skipping attempt to create a new blackboard entry named %s since it already exists. Policy settings from the first entry will be retained."), *EntryName.ToString());
+			UE_LOGF(LogMover, Warning, "Skipping attempt to create a new blackboard entry named %ls since it already exists. Policy settings from the first entry will be retained.", *EntryName.ToString());
 			return false;
 		}
 
