Mover: fix bug where an invalid up-direction override would result in a gravity override being turned off instead

#Synced-CL 51296837
#Synced-CL 51296837

diff --git a/Engine/Plugins/Experimental/Mover/Source/Mover/Private/MoverComponent.cpp b/Engine/Plugins/Experimental/Mover/Source/Mover/Private/MoverComponent.cpp
--- a/Engine/Plugins/Experimental/Mover/Source/Mover/Private/MoverComponent.cpp
+++ b/Engine/Plugins/Experimental/Mover/Source/Mover/Private/MoverComponent.cpp
@@ -2267,7 +2267,7 @@ void UMoverComponent::SetUpDirectionOverride(bool bOverrideUpDirection, FVector
 		if (UpDirection.IsNearlyZero())
 		{
 			UE_LOGF(LogMover, Warning, "Ignoring the provided UpDirection (%ls) override because it is a zero vector. (%ls)", *UpDirection.ToString(), *GetNameSafe(GetOwner()));
-			bHasGravityOverride = false;
+			bHasUpDirectionOverride = false;
 			return;
 		}
 		UpDirectionOverride = UpDirection.GetSafeNormal();
