Mover: fixed case where a crash could occur if there was a rollback before BeginPlay, due to the ModeFSM not having an associated simulation yet

#Synced-CL 51339572

diff --git a/Engine/Plugins/Experimental/Mover/Source/Mover/Private/MoverComponent.cpp b/Engine/Plugins/Experimental/Mover/Source/Mover/Private/MoverComponent.cpp
--- a/Engine/Plugins/Experimental/Mover/Source/Mover/Private/MoverComponent.cpp
+++ b/Engine/Plugins/Experimental/Mover/Source/Mover/Private/MoverComponent.cpp
@@ -100,34 +100,6 @@ void UMoverComponent::InitializeComponent()
 		FindDefaultUpdatedComponent();
 
 
-		// Instantiate our sister backend component that will actually talk to the system driving the simulation
-		if (BackendClass)
-		{
-			UActorComponent* NewLiaisonComp = Cast<UActorComponent>(GetOwner()->FindComponentByInterface<IMoverBackendLiaisonInterface>());
-			if (!NewLiaisonComp)
-			{
-				NewLiaisonComp = NewObject<UActorComponent>(this, BackendClass, TEXT("BackendLiaisonComponent"));
-			}
-
-			BackendLiaisonComp.SetObject(NewLiaisonComp);
-			BackendLiaisonComp.SetInterface(CastChecked<IMoverBackendLiaisonInterface>(NewLiaisonComp));
-			if (BackendLiaisonComp)
-			{
-				NewLiaisonComp->RegisterComponent();
-				if (!NewLiaisonComp->HasBeenInitialized())
-				{
-					NewLiaisonComp->InitializeComponent();
-				}
-				NewLiaisonComp->SetNetAddressable();
-			}
-		}
-		else
-		{
-			UE_LOGF(LogMover, Error, "No backend class set on %ls. Mover actor will not function.", *GetNameSafe(GetOwner()));
-		}
-
-
-
 		// Set up FSM and initial movement states
 		ModeFSM = NewObject<UMovementModeStateMachine>(this, TEXT("MoverStateMachine"), RF_Transient);
 		ModeFSM->ClearAllMovementModes();
@@ -162,6 +134,33 @@ void UMoverComponent::InitializeComponent()
 			ModeFSM->SetDefaultMode(StartingMovementMode);
 			ModeFSM->QueueNextMode(StartingMovementMode);
 		}
+
+		// Instantiate our sister backend component that will actually talk to the system driving the simulation
+		if (BackendClass)
+		{
+			UActorComponent* NewLiaisonComp = Cast<UActorComponent>(GetOwner()->FindComponentByInterface<IMoverBackendLiaisonInterface>());
+			if (!NewLiaisonComp)
+			{
+				NewLiaisonComp = NewObject<UActorComponent>(this, BackendClass, TEXT("BackendLiaisonComponent"));
+			}
+
+			BackendLiaisonComp.SetObject(NewLiaisonComp);
+			BackendLiaisonComp.SetInterface(CastChecked<IMoverBackendLiaisonInterface>(NewLiaisonComp));
+			if (BackendLiaisonComp)
+			{
+				NewLiaisonComp->RegisterComponent();
+				if (!NewLiaisonComp->HasBeenInitialized())
+				{
+					NewLiaisonComp->InitializeComponent();
+				}
+				NewLiaisonComp->SetNetAddressable();
+			}
+		}
+		else
+		{
+			UE_LOGF(LogMover, Error, "No backend class set on %ls. Mover actor will not function.", *GetNameSafe(GetOwner()));
+		}
+
 	}
 
 	// Gather initial state to fulfill queries
@@ -252,17 +251,14 @@ void UMoverComponent::BeginPlay()
 {
 	Super::BeginPlay();
 
-	if (ensureMsgf(Simulation != nullptr, TEXT("MoverComponent requires a Simulation to be set during initialization! The backend should have done this already.")))
-	{
-		ModeFSM->SetSimulation(*Simulation);
+	check(Simulation != nullptr); // MoverComponent requires a Simulation to be set during initialization. The backend should have done this already.
 
-		static URollbackBlackboard::EntrySettings ModeChangeRecordSettings;
-		ModeChangeRecordSettings.SizingPolicy = EBlackboardSizingPolicy::FixedDeclaredSize;
-		ModeChangeRecordSettings.FixedSize = 4;
-		ModeChangeRecordSettings.PersistencePolicy = EBlackboardPersistencePolicy::Forever;
-		ModeChangeRecordSettings.RollbackPolicy = EBlackboardRollbackPolicy::InvalidatedOnRollback;
-		RollbackBlackboardExternalWrapper.CreateEntry<FMovementModeChangeRecord>(CommonBlackboard::LastModeChangeRecord, ModeChangeRecordSettings);
-	}
+	static URollbackBlackboard::EntrySettings ModeChangeRecordSettings;
+	ModeChangeRecordSettings.SizingPolicy = EBlackboardSizingPolicy::FixedDeclaredSize;
+	ModeChangeRecordSettings.FixedSize = 4;
+	ModeChangeRecordSettings.PersistencePolicy = EBlackboardPersistencePolicy::Forever;
+	ModeChangeRecordSettings.RollbackPolicy = EBlackboardRollbackPolicy::InvalidatedOnRollback;
+	RollbackBlackboardExternalWrapper.CreateEntry<FMovementModeChangeRecord>(CommonBlackboard::LastModeChangeRecord, ModeChangeRecordSettings);
 
 	FindDefaultUpdatedComponent();
 	ensureMsgf(UpdatedComponent != nullptr, TEXT("No root component found on %s. Simulation initialization will most likely fail."), *GetPathNameSafe(GetOwner()));
@@ -1044,6 +1040,8 @@ void UMoverComponent::SetSimulation(UMoverSimulation& InSimulation)
 	
 	RollbackBlackboard = Simulation->RollbackBlackboard;	// TODO: fetching the true blackboard is temporary: this will be removed once simulations are used consistently
 	RollbackBlackboardExternalWrapper = FRollbackBlackboardExternalWrapper(*Simulation->RollbackBlackboard);
+
+	ModeFSM->SetSimulation(InSimulation);
 }
 
 void UMoverComponent::HandleImpact(FMoverOnImpactParams& ImpactParams)
