Mover: Adding RVO interface to NavMoverComponent to support detour crowd avoidance. Note, this is not force based RVO avoidance. Also adding some example content to mover examples to show this funcionality.

#virtualized
#Synced-CL 51340830

diff --git a/Engine/Plugins/Experimental/Mover/Source/Mover/Private/DefaultMovementSet/NavMoverComponent.cpp b/Engine/Plugins/Experimental/Mover/Source/Mover/Private/DefaultMovementSet/NavMoverComponent.cpp
--- a/Engine/Plugins/Experimental/Mover/Source/Mover/Private/DefaultMovementSet/NavMoverComponent.cpp
+++ b/Engine/Plugins/Experimental/Mover/Source/Mover/Private/DefaultMovementSet/NavMoverComponent.cpp
@@ -12,9 +12,31 @@
 #include UE_INLINE_GENERATED_CPP_BY_NAME(NavMoverComponent)
 
 UNavMoverComponent::UNavMoverComponent()
+	: AvoidanceUID(0)
+	, AvoidanceWeight(0.5)
+	, AvoidanceConsiderationRadius(500)
 {
 	bWantsInitializeComponent = true;
 	bAutoActivate = true;
+	
+	AvoidanceGroup.bGroup0 = true;
+	GroupsToAvoid.Packed = 0xFFFFFFFF;
+	GroupsToIgnore.Packed = 0;
+}
+
+void UNavMoverComponent::SetAvoidanceGroupMask(const FNavAvoidanceMask& GroupMask)
+{
+	SetAvoidanceGroupMask(GroupMask.Packed);
+}
+
+void UNavMoverComponent::SetGroupsToAvoidMask(const FNavAvoidanceMask& GroupMask)
+{
+	SetGroupsToAvoidMask(GroupMask.Packed);
+}
+
+void UNavMoverComponent::SetGroupsToIgnoreMask(const FNavAvoidanceMask& GroupMask)
+{
+	SetGroupsToIgnoreMask(GroupMask.Packed);
 }
 
 void UNavMoverComponent::InitializeComponent()
@@ -320,3 +342,80 @@ FVector UNavMoverComponent::GetVelocityForNavMovement() const
 {
 	return MoverComponent.IsValid() ? MoverComponent->GetVelocity() : FVector::ZeroVector;
 }
+
+void UNavMoverComponent::SetRVOAvoidanceUID(int32 UID)
+{
+	AvoidanceUID = UID;
+}
+
+int32 UNavMoverComponent::GetRVOAvoidanceUID()
+{
+	return AvoidanceUID;
+}
+
+void UNavMoverComponent::SetRVOAvoidanceWeight(float Weight)
+{
+	AvoidanceWeight = Weight;
+}
+
+float UNavMoverComponent::GetRVOAvoidanceWeight()
+{
+	return AvoidanceWeight;
+}
+
+FVector UNavMoverComponent::GetRVOAvoidanceOrigin()
+{
+	return GetFeetLocation();
+}
+
+float UNavMoverComponent::GetRVOAvoidanceRadius()
+{
+	const UCapsuleComponent* CapsuleComp = Cast<UCapsuleComponent>(GetUpdatedObject());
+	return CapsuleComp ? CapsuleComp->GetScaledCapsuleRadius() : 0.0f;
+}
+
+float UNavMoverComponent::GetRVOAvoidanceConsiderationRadius()
+{
+	return AvoidanceConsiderationRadius;
+}
+
+float UNavMoverComponent::GetRVOAvoidanceHeight()
+{
+	const UCapsuleComponent* CapsuleComp = Cast<UCapsuleComponent>(GetUpdatedObject());
+	return CapsuleComp ? CapsuleComp->GetScaledCapsuleHalfHeight() : 0.0f;
+}
+
+FVector UNavMoverComponent::GetVelocityForRVOConsideration()
+{
+	return GetVelocityForNavMovement();
+}
+
+void UNavMoverComponent::SetAvoidanceGroupMask(int32 GroupFlags)
+{
+	AvoidanceGroup.SetFlagsDirectly(GroupFlags);
+}
+
+int32 UNavMoverComponent::GetAvoidanceGroupMask()
+{
+	return AvoidanceGroup.Packed;
+}
+
+void UNavMoverComponent::SetGroupsToAvoidMask(int32 GroupFlags)
+{
+	GroupsToAvoid.SetFlagsDirectly(GroupFlags);
+}
+
+int32 UNavMoverComponent::GetGroupsToAvoidMask()
+{
+	return GroupsToAvoid.Packed;
+}
+
+void UNavMoverComponent::SetGroupsToIgnoreMask(int32 GroupFlags)
+{
+	GroupsToIgnore.SetFlagsDirectly(GroupFlags);
+}
+
+int32 UNavMoverComponent::GetGroupsToIgnoreMask()
+{
+	return GroupsToIgnore.Packed;
+}
diff --git a/Engine/Plugins/Experimental/Mover/Source/Mover/Public/DefaultMovementSet/NavMoverComponent.h b/Engine/Plugins/Experimental/Mover/Source/Mover/Public/DefaultMovementSet/NavMoverComponent.h
--- a/Engine/Plugins/Experimental/Mover/Source/Mover/Public/DefaultMovementSet/NavMoverComponent.h
+++ b/Engine/Plugins/Experimental/Mover/Source/Mover/Public/DefaultMovementSet/NavMoverComponent.h
@@ -6,6 +6,8 @@
 #include "GameFramework/Actor.h"
 #include "Components/ActorComponent.h"
 #include "MoverComponent.h"
+#include "AI/RVOAvoidanceInterface.h"
+#include "AI/Navigation/NavigationAvoidanceTypes.h"
 #include "GameFramework/NavMovementInterface.h"
 #include "UObject/WeakInterfacePtr.h"
 
@@ -16,11 +18,13 @@
 /**
  * NavMoverComponent: Responsible for implementing INavMoveInterface with UMoverComponent so pathfinding and other forms of navigation movement work.
  * This component also caches the input given to it that is then consumed by the mover system.
- * Note: This component relies on the parent actor having a MoverComponent as well. By default this component only has a reference to MoverComponent meaning
+ * This component supports a basic version of the RVO avoidance interface to make it usable with Detour Crowd Avoidance. It does not support force based 
+ * RVO avoidance, but this or child component would be a good place to update/evaluate Avoidance information and apply effects to the Mover Simulation.
+ * Note: This component relies on the parent actor having a MoverComponent as well. By default, this component only has a reference to MoverComponent meaning
  * we use other ways (such as gameplay tags for the active movement mode) to check for state rather than calling specific functions on the active MoverComponent.
  */
 UCLASS(MinimalAPI, BlueprintType, meta = (BlueprintSpawnableComponent))
-class UNavMoverComponent : public UActorComponent, public INavMovementInterface
+class UNavMoverComponent : public UActorComponent, public INavMovementInterface, public IRVOAvoidanceInterface
 {
 	GENERATED_BODY()
 
@@ -35,6 +39,38 @@ public:
 	UPROPERTY()
 	FMovementProperties MovementState;
 	
+	/** No default value, for now it's assumed to be valid if GetAvoidanceManager() returns non-NULL. */
+	UPROPERTY(Category="Avoidance", VisibleAnywhere, BlueprintReadOnly, AdvancedDisplay)
+	int32 AvoidanceUID;
+	
+	/** Moving actor's group mask */
+	UPROPERTY(Category="Avoidance", EditAnywhere, BlueprintReadOnly, AdvancedDisplay)
+	FNavAvoidanceMask AvoidanceGroup;
+	
+	UFUNCTION(BlueprintCallable, Category = "Nav Movement|Avoidance")
+	UE_API void SetAvoidanceGroupMask(const FNavAvoidanceMask& GroupMask);
+	
+	/** Will avoid other agents if they are in one of specified groups */
+	UPROPERTY(Category="Avoidance", EditAnywhere, BlueprintReadOnly, AdvancedDisplay)
+	FNavAvoidanceMask GroupsToAvoid;
+	
+	UFUNCTION(BlueprintCallable, Category = "Nav Movement|Avoidance")
+	UE_API void SetGroupsToAvoidMask(const FNavAvoidanceMask& GroupMask);
+	
+	/** Will NOT avoid other agents if they are in one of specified groups, higher priority than GroupsToAvoid */
+	UPROPERTY(Category="Avoidance", EditAnywhere, BlueprintReadOnly, AdvancedDisplay)
+	FNavAvoidanceMask GroupsToIgnore;
+	
+	UFUNCTION(BlueprintCallable, Category = "Nav Movement|Avoidance")
+	UE_API void SetGroupsToIgnoreMask(const FNavAvoidanceMask& GroupMask);
+	
+	/** De facto default value 0.5 (due to that being the default in the avoidance registration function), indicates RVO behavior. */
+	UPROPERTY(Category="Avoidance", EditAnywhere, BlueprintReadOnly)
+	float AvoidanceWeight;
+	
+	UPROPERTY(Category="Avoidance", EditAnywhere, BlueprintReadOnly, meta=(ForceUnits=cm))
+	float AvoidanceConsiderationRadius;
+	
 protected:
 	/** associated properties for nav movement */
 	UPROPERTY(EditAnywhere, BlueprintReadWrite, Category="Nav Movement")
@@ -51,7 +87,6 @@ protected:
 	// Used to store movement input velocity from requested moves
 	FVector CachedNavMoveInputVelocity = FVector::ZeroVector;
 	
-private:
 	/** object implementing IPathFollowingAgentInterface. Is private to control access to it.
 	 *	@see SetPathFollowingAgent, GetPathFollowingAgent */
 	TWeakInterfacePtr<IPathFollowingAgentInterface> PathFollowingComp;
@@ -151,6 +186,24 @@ public:
 
 	/** Returns true if currently flying (moving through a non-fluid volume without resting on the ground) */
 	UE_API virtual bool IsFlying() const override;
+	
+	/** BEGIN IRVOAvoidanceInterface */
+	UE_API virtual void SetRVOAvoidanceUID(int32 UID) override;
+	UE_API virtual int32 GetRVOAvoidanceUID() override;
+	UE_API virtual void SetRVOAvoidanceWeight(float Weight) override;
+	UE_API virtual float GetRVOAvoidanceWeight() override;
+	UE_API virtual FVector GetRVOAvoidanceOrigin() override;
+	UE_API virtual float GetRVOAvoidanceRadius() override;
+	UE_API virtual float GetRVOAvoidanceHeight() override;
+	UE_API virtual float GetRVOAvoidanceConsiderationRadius() override;
+	UE_API virtual FVector GetVelocityForRVOConsideration() override;
+	UE_API virtual void SetAvoidanceGroupMask(int32 GroupFlags) override;
+	UE_API virtual int32 GetAvoidanceGroupMask() override;
+	UE_API virtual void SetGroupsToAvoidMask(int32 GroupFlags) override;
+	UE_API virtual int32 GetGroupsToAvoidMask() override;
+	UE_API virtual void SetGroupsToIgnoreMask(int32 GroupFlags) override;
+	UE_API virtual int32 GetGroupsToIgnoreMask() override;
+	/** END IRVOAvoidanceInterface */
 };
 
 #undef UE_API
