Remove includes guarded by `UE_ENABLE_INCLUDE_ORDER_DEPRECATED_IN_5_4`.

#Synced-CL 51408971

#Synced-CL 51408971

diff --git a/Engine/Plugins/Compositing/OpenColorIO/Source/OpenColorIO/Public/OpenColorIOColorSpace.h b/Engine/Plugins/Compositing/OpenColorIO/Source/OpenColorIO/Public/OpenColorIOColorSpace.h
--- a/Engine/Plugins/Compositing/OpenColorIO/Source/OpenColorIO/Public/OpenColorIOColorSpace.h
+++ b/Engine/Plugins/Compositing/OpenColorIO/Source/OpenColorIO/Public/OpenColorIOColorSpace.h
@@ -2,10 +2,6 @@
 
 #pragma once
 
-#if UE_ENABLE_INCLUDE_ORDER_DEPRECATED_IN_5_4
-#include "CoreMinimal.h"
-#endif
-
 #include "OpenColorIOColorSpace.generated.h"
 
 
diff --git a/Engine/Plugins/Compositing/OpenColorIO/Source/OpenColorIO/Public/OpenColorIOShader.h b/Engine/Plugins/Compositing/OpenColorIO/Source/OpenColorIO/Public/OpenColorIOShader.h
--- a/Engine/Plugins/Compositing/OpenColorIO/Source/OpenColorIO/Public/OpenColorIOShader.h
+++ b/Engine/Plugins/Compositing/OpenColorIO/Source/OpenColorIO/Public/OpenColorIOShader.h
@@ -7,9 +7,6 @@
 #pragma once
 
 #include "Containers/ContainersFwd.h"
-#if UE_ENABLE_INCLUDE_ORDER_DEPRECATED_IN_5_4
-#include "GlobalShader.h"
-#endif
 #include "OpenColorIOShaderType.h"
 #include "OpenColorIOShared.h"
 #include "ShaderParameters.h"
diff --git a/Engine/Plugins/Compositing/OpenColorIO/Source/OpenColorIO/Public/OpenColorIOShared.h b/Engine/Plugins/Compositing/OpenColorIO/Source/OpenColorIO/Public/OpenColorIOShared.h
--- a/Engine/Plugins/Compositing/OpenColorIO/Source/OpenColorIO/Public/OpenColorIOShared.h
+++ b/Engine/Plugins/Compositing/OpenColorIO/Source/OpenColorIO/Public/OpenColorIOShared.h
@@ -17,15 +17,6 @@
 #include "StaticParameterSet.h"
 #include "Templates/RefCounting.h"
 
-#if UE_ENABLE_INCLUDE_ORDER_DEPRECATED_IN_5_4
-#include "CoreMinimal.h"
-#include "Engine/EngineTypes.h"
-#include "Misc/Optional.h"
-#include "ShaderCompiler.h"
-#include "OpenColorIOShaderCompilationManager.h"
-#include "RHI.h"
-#endif
-
 /** Enum used to indicate whether the working color space should be used as a source or destination. */
 enum class EOpenColorIOWorkingColorSpaceTransform : uint8
 {
diff --git a/Engine/Plugins/Experimental/ChaosVehiclesPlugin/Source/ChaosVehicles/Public/ChaosVehicleWheel.h b/Engine/Plugins/Experimental/ChaosVehiclesPlugin/Source/ChaosVehicles/Public/ChaosVehicleWheel.h
--- a/Engine/Plugins/Experimental/ChaosVehiclesPlugin/Source/ChaosVehicles/Public/ChaosVehicleWheel.h
+++ b/Engine/Plugins/Experimental/ChaosVehiclesPlugin/Source/ChaosVehicles/Public/ChaosVehicleWheel.h
@@ -13,9 +13,6 @@
 #include "EngineDefines.h"
 #include "SimpleVehicle.h"
 #include "VehicleUtility.h"
-#if UE_ENABLE_INCLUDE_ORDER_DEPRECATED_IN_5_4
-#include "Engine/HitResult.h"
-#endif // UE_ENABLE_INCLUDE_ORDER_DEPRECATED_IN_5_4
 #include "Curves/CurveFloat.h"
 
 #include "ChaosVehicleWheel.generated.h"
diff --git a/Engine/Plugins/Experimental/LearningCore/Source/Learning/Public/MultiArrayView.h b/Engine/Plugins/Experimental/LearningCore/Source/Learning/Public/MultiArrayView.h
--- a/Engine/Plugins/Experimental/LearningCore/Source/Learning/Public/MultiArrayView.h
+++ b/Engine/Plugins/Experimental/LearningCore/Source/Learning/Public/MultiArrayView.h
@@ -897,9 +897,3 @@ private:
 
 template<uint8 InDimNum, typename InElementType>
 using TConstMultiArrayView = TMultiArrayView<InDimNum, const InElementType>;
-
-#if UE_ENABLE_INCLUDE_ORDER_DEPRECATED_IN_5_4
-#include "Templates/IsConst.h"
-#include "Templates/IsSigned.h"
-#include "Templates/PointerIsConvertibleFromTo.h"
-#endif
\ No newline at end of file
diff --git a/Engine/Plugins/Experimental/Mover/Source/Mover/Public/MoverComponent.h b/Engine/Plugins/Experimental/Mover/Source/Mover/Public/MoverComponent.h
--- a/Engine/Plugins/Experimental/Mover/Source/Mover/Public/MoverComponent.h
+++ b/Engine/Plugins/Experimental/Mover/Source/Mover/Public/MoverComponent.h
@@ -2,9 +2,6 @@
 
 #pragma once
 
-#if UE_ENABLE_INCLUDE_ORDER_DEPRECATED_IN_5_4
-#include "CoreMinimal.h"
-#endif // UE_ENABLE_INCLUDE_ORDER_DEPRECATED_IN_5_4
 #include "Components/ActorComponent.h"
 #include "MotionWarpingAdapter.h"
 #include "MovementMode.h"
@@ -14,9 +11,6 @@
 #include "MoveLibrary/BasedMovementUtils.h"
 #include "MoveLibrary/ConstrainedMoveUtils.h"
 #include "MoveLibrary/RollbackBlackboardWrappers.h"
-#if UE_ENABLE_INCLUDE_ORDER_DEPRECATED_IN_5_4
-#include "Engine/HitResult.h"
-#endif // UE_ENABLE_INCLUDE_ORDER_DEPRECATED_IN_5_4
 #include "MovementModifier.h"
 #include "Backends/MoverBackendLiaison.h"
 #include "UObject/WeakInterfacePtr.h"
diff --git a/Engine/Plugins/Experimental/XRCreativeFramework/Source/XRCreative/Private/ITF/SelectionInteraction.h b/Engine/Plugins/Experimental/XRCreativeFramework/Source/XRCreative/Private/ITF/SelectionInteraction.h
--- a/Engine/Plugins/Experimental/XRCreativeFramework/Source/XRCreative/Private/ITF/SelectionInteraction.h
+++ b/Engine/Plugins/Experimental/XRCreativeFramework/Source/XRCreative/Private/ITF/SelectionInteraction.h
@@ -6,9 +6,6 @@
 #include "BaseBehaviors/SingleClickBehavior.h"
 #include "Elements/Actor/ActorElementData.h"
 #include "Elements/Framework/TypedElementOwnerStore.h"
-#if UE_ENABLE_INCLUDE_ORDER_DEPRECATED_IN_5_4
-#include "Engine/HitResult.h"
-#endif // UE_ENABLE_INCLUDE_ORDER_DEPRECATED_IN_5_4
 #include "InputBehaviorSet.h"
 #include "SelectionInteraction.generated.h"
 
diff --git a/Engine/Plugins/Runtime/GeoReferencing/Source/GeoReferencingEditor/Public/GeoReferencingEditorBPLibrary.h b/Engine/Plugins/Runtime/GeoReferencing/Source/GeoReferencingEditor/Public/GeoReferencingEditorBPLibrary.h
--- a/Engine/Plugins/Runtime/GeoReferencing/Source/GeoReferencingEditor/Public/GeoReferencingEditorBPLibrary.h
+++ b/Engine/Plugins/Runtime/GeoReferencing/Source/GeoReferencingEditor/Public/GeoReferencingEditorBPLibrary.h
@@ -2,10 +2,6 @@
 
 #pragma once
 
-#if UE_ENABLE_INCLUDE_ORDER_DEPRECATED_IN_5_4
-#include "CoreMinimal.h"
-#include "Engine/HitResult.h"
-#endif // UE_ENABLE_INCLUDE_ORDER_DEPRECATED_IN_5_4
 #include "Kismet/BlueprintFunctionLibrary.h"
 
 #include "GeoReferencingEditorBPLibrary.generated.h"
diff --git a/Engine/Plugins/Runtime/MeshModelingToolset/Source/ModelingComponents/Public/Selection/SelectClickedAction.h b/Engine/Plugins/Runtime/MeshModelingToolset/Source/ModelingComponents/Public/Selection/SelectClickedAction.h
--- a/Engine/Plugins/Runtime/MeshModelingToolset/Source/ModelingComponents/Public/Selection/SelectClickedAction.h
+++ b/Engine/Plugins/Runtime/MeshModelingToolset/Source/ModelingComponents/Public/Selection/SelectClickedAction.h
@@ -5,9 +5,6 @@
 #include "CoreMinimal.h"
 #include "BaseBehaviors/BehaviorTargetInterfaces.h"
 #include "ToolSceneQueriesUtil.h"
-#if UE_ENABLE_INCLUDE_ORDER_DEPRECATED_IN_5_4
-#include "Engine/HitResult.h"
-#endif // UE_ENABLE_INCLUDE_ORDER_DEPRECATED_IN_5_4
 
 struct FHitResult;
 
diff --git a/Engine/Plugins/Runtime/NetworkPredictionExtras/Source/NetworkPredictionExtras/Public/BaseMovementComponent.h b/Engine/Plugins/Runtime/NetworkPredictionExtras/Source/NetworkPredictionExtras/Public/BaseMovementComponent.h
--- a/Engine/Plugins/Runtime/NetworkPredictionExtras/Source/NetworkPredictionExtras/Public/BaseMovementComponent.h
+++ b/Engine/Plugins/Runtime/NetworkPredictionExtras/Source/NetworkPredictionExtras/Public/BaseMovementComponent.h
@@ -3,9 +3,6 @@
 #pragma once
 
 #include "NetworkPredictionComponent.h"
-#if UE_ENABLE_INCLUDE_ORDER_DEPRECATED_IN_5_4
-#include "Engine/HitResult.h"
-#endif // UE_ENABLE_INCLUDE_ORDER_DEPRECATED_IN_5_4
 #include "BaseMovementComponent.generated.h"
 
 struct FHitResult;
diff --git a/Engine/Source/Developer/Apple/MetalShaderFormat/Private/MetalShaderCompiler.cpp b/Engine/Source/Developer/Apple/MetalShaderFormat/Private/MetalShaderCompiler.cpp
--- a/Engine/Source/Developer/Apple/MetalShaderFormat/Private/MetalShaderCompiler.cpp
+++ b/Engine/Source/Developer/Apple/MetalShaderFormat/Private/MetalShaderCompiler.cpp
@@ -2,9 +2,6 @@
 
 #include "MetalShaderCompiler.h"
 
-#if UE_ENABLE_INCLUDE_ORDER_DEPRECATED_IN_5_4
-#include "CoreMinimal.h"
-#endif
 #include "Containers/ArrayView.h"
 #include "HAL/PlatformFileManager.h"
 #include "MetalShaderFormat.h"
diff --git a/Engine/Source/Editor/ClothPainter/Private/ClothPainter.h b/Engine/Source/Editor/ClothPainter/Private/ClothPainter.h
--- a/Engine/Source/Editor/ClothPainter/Private/ClothPainter.h
+++ b/Engine/Source/Editor/ClothPainter/Private/ClothPainter.h
@@ -7,9 +7,6 @@
 #include "Containers/Map.h"
 #include "Delegates/IDelegateInstance.h"
 #include "Engine/EngineBaseTypes.h"
-#if UE_ENABLE_INCLUDE_ORDER_DEPRECATED_IN_5_4
-#include "Engine/HitResult.h"
-#endif // UE_ENABLE_INCLUDE_ORDER_DEPRECATED_IN_5_4
 #include "IMeshPainter.h"
 #include "InputCoreTypes.h"
 #include "Internationalization/Text.h"
diff --git a/Engine/Source/Editor/UnrealEd/Private/UnrealEdPrivatePCH.h b/Engine/Source/Editor/UnrealEd/Private/UnrealEdPrivatePCH.h
--- a/Engine/Source/Editor/UnrealEd/Private/UnrealEdPrivatePCH.h
+++ b/Engine/Source/Editor/UnrealEd/Private/UnrealEdPrivatePCH.h
@@ -648,10 +648,6 @@
 #include "Commandlets/GatherTextCommandletBase.h"
 #include "Dialogs/Dialogs.h"
 
-#if UE_ENABLE_INCLUDE_ORDER_DEPRECATED_IN_5_4
-#include "NavFilters/NavigationQueryFilter.h"
-#endif
-
 #if UE_ENABLE_INCLUDE_ORDER_DEPRECATED_IN_5_5
 #include "Engine/MeshMerging.h"
 #include "Templates/IsTriviallyDestructible.h"
diff --git a/Engine/Source/Editor/UnrealEd/Public/UnrealEdSharedPCH.h b/Engine/Source/Editor/UnrealEd/Public/UnrealEdSharedPCH.h
--- a/Engine/Source/Editor/UnrealEd/Public/UnrealEdSharedPCH.h
+++ b/Engine/Source/Editor/UnrealEd/Public/UnrealEdSharedPCH.h
@@ -21,11 +21,6 @@
 #include "GameplayTaskTypes.h"
 #include "GameplayTask.h"
 
-#if UE_ENABLE_INCLUDE_ORDER_DEPRECATED_IN_5_4
-// From NavigationSystem:
-#include "NavFilters/NavigationQueryFilter.h"
-#endif
-
 // From UnrealEd:
 #include "AssetThumbnail.h"
 #include "Editor.h"
diff --git a/Engine/Source/Runtime/AIModule/Classes/AIController.h b/Engine/Source/Runtime/AIModule/Classes/AIController.h
--- a/Engine/Source/Runtime/AIModule/Classes/AIController.h
+++ b/Engine/Source/Runtime/AIModule/Classes/AIController.h
@@ -7,9 +7,6 @@
 #include "UObject/UObjectGlobals.h"
 #include "Templates/SubclassOf.h"
 #include "EngineDefines.h"
-#if UE_ENABLE_INCLUDE_ORDER_DEPRECATED_IN_5_4
-#include "NavFilters/NavigationQueryFilter.h"
-#endif
 #include "AITypes.h"
 #include "GameplayTaskOwnerInterface.h"
 #include "GameplayTask.h"
diff --git a/Engine/Source/Runtime/AIModule/Classes/AITypes.h b/Engine/Source/Runtime/AIModule/Classes/AITypes.h
--- a/Engine/Source/Runtime/AIModule/Classes/AITypes.h
+++ b/Engine/Source/Runtime/AIModule/Classes/AITypes.h
@@ -6,14 +6,7 @@
 #include "Stats/Stats.h"
 #include "UObject/ObjectMacros.h"
 #include "Templates/SubclassOf.h"
-#if UE_ENABLE_INCLUDE_ORDER_DEPRECATED_IN_5_4
-#include "NavFilters/NavigationQueryFilter.h"
-#include "AI/Navigation/NavigationTypes.h"
-#endif
 #include "NavigationSystemTypes.h"
-#if UE_ENABLE_INCLUDE_ORDER_DEPRECATED_IN_5_4
-#include "GameFramework/Actor.h"
-#endif
 #include "AITypes.generated.h"
 
 class AActor;
diff --git a/Engine/Source/Runtime/AIModule/Classes/BehaviorTree/Decorators/BTDecorator_DoesPathExist.h b/Engine/Source/Runtime/AIModule/Classes/BehaviorTree/Decorators/BTDecorator_DoesPathExist.h
--- a/Engine/Source/Runtime/AIModule/Classes/BehaviorTree/Decorators/BTDecorator_DoesPathExist.h
+++ b/Engine/Source/Runtime/AIModule/Classes/BehaviorTree/Decorators/BTDecorator_DoesPathExist.h
@@ -5,9 +5,6 @@
 #include "CoreMinimal.h"
 #include "UObject/ObjectMacros.h"
 #include "Templates/SubclassOf.h"
-#if UE_ENABLE_INCLUDE_ORDER_DEPRECATED_IN_5_4
-#include "NavFilters/NavigationQueryFilter.h"
-#endif
 #include "BehaviorTree/BehaviorTreeTypes.h"
 #include "BehaviorTree/BTDecorator.h"
 #include "BehaviorTree/ValueOrBBKey.h"
diff --git a/Engine/Source/Runtime/AIModule/Classes/BehaviorTree/Tasks/BTTask_MoveTo.h b/Engine/Source/Runtime/AIModule/Classes/BehaviorTree/Tasks/BTTask_MoveTo.h
--- a/Engine/Source/Runtime/AIModule/Classes/BehaviorTree/Tasks/BTTask_MoveTo.h
+++ b/Engine/Source/Runtime/AIModule/Classes/BehaviorTree/Tasks/BTTask_MoveTo.h
@@ -6,9 +6,6 @@
 #include "UObject/ObjectMacros.h"
 #include "InputCoreTypes.h"
 #include "Templates/SubclassOf.h"
-#if UE_ENABLE_INCLUDE_ORDER_DEPRECATED_IN_5_4
-#include "NavFilters/NavigationQueryFilter.h"
-#endif
 #include "AITypes.h"
 #include "BehaviorTree/Tasks/BTTask_BlackboardBase.h"
 #include "BehaviorTree/ValueOrBBKey.h"
diff --git a/Engine/Source/Runtime/AIModule/Classes/EnvironmentQuery/EnvQueryTypes.h b/Engine/Source/Runtime/AIModule/Classes/EnvironmentQuery/EnvQueryTypes.h
--- a/Engine/Source/Runtime/AIModule/Classes/EnvironmentQuery/EnvQueryTypes.h
+++ b/Engine/Source/Runtime/AIModule/Classes/EnvironmentQuery/EnvQueryTypes.h
@@ -8,9 +8,6 @@
 #include "UObject/Object.h"
 #include "Templates/SubclassOf.h"
 #include "Engine/EngineTypes.h"
-#if UE_ENABLE_INCLUDE_ORDER_DEPRECATED_IN_5_4
-#include "NavFilters/NavigationQueryFilter.h"
-#endif
 #include "EnvironmentQuery/Items/EnvQueryItemType.h"
 #include "EnvironmentQuery/EnvQueryContext.h"
 #include "BehaviorTree/BehaviorTreeTypes.h"
diff --git a/Engine/Source/Runtime/AIModule/Classes/EnvironmentQuery/Generators/EnvQueryGenerator_PathingGrid.h b/Engine/Source/Runtime/AIModule/Classes/EnvironmentQuery/Generators/EnvQueryGenerator_PathingGrid.h
--- a/Engine/Source/Runtime/AIModule/Classes/EnvironmentQuery/Generators/EnvQueryGenerator_PathingGrid.h
+++ b/Engine/Source/Runtime/AIModule/Classes/EnvironmentQuery/Generators/EnvQueryGenerator_PathingGrid.h
@@ -5,9 +5,6 @@
 #include "CoreMinimal.h"
 #include "UObject/ObjectMacros.h"
 #include "Templates/SubclassOf.h"
-#if UE_ENABLE_INCLUDE_ORDER_DEPRECATED_IN_5_4
-#include "NavFilters/NavigationQueryFilter.h"
-#endif
 #include "AI/Navigation/NavigationTypes.h"
 #include "DataProviders/AIDataProvider.h"
 #include "EnvironmentQuery/Generators/EnvQueryGenerator_SimpleGrid.h"
diff --git a/Engine/Source/Runtime/AIModule/Classes/EnvironmentQuery/Tests/EnvQueryTest_Pathfinding.h b/Engine/Source/Runtime/AIModule/Classes/EnvironmentQuery/Tests/EnvQueryTest_Pathfinding.h
--- a/Engine/Source/Runtime/AIModule/Classes/EnvironmentQuery/Tests/EnvQueryTest_Pathfinding.h
+++ b/Engine/Source/Runtime/AIModule/Classes/EnvironmentQuery/Tests/EnvQueryTest_Pathfinding.h
@@ -5,9 +5,6 @@
 #include "CoreMinimal.h"
 #include "UObject/ObjectMacros.h"
 #include "Templates/SubclassOf.h"
-#if UE_ENABLE_INCLUDE_ORDER_DEPRECATED_IN_5_4
-#include "NavFilters/NavigationQueryFilter.h"
-#endif
 #include "AI/Navigation/NavigationTypes.h"
 #include "NavigationSystemTypes.h"
 #include "EnvironmentQuery/EnvQueryContext.h"
diff --git a/Engine/Source/Runtime/AIModule/Classes/Navigation/CrowdManager.h b/Engine/Source/Runtime/AIModule/Classes/Navigation/CrowdManager.h
--- a/Engine/Source/Runtime/AIModule/Classes/Navigation/CrowdManager.h
+++ b/Engine/Source/Runtime/AIModule/Classes/Navigation/CrowdManager.h
@@ -10,9 +10,6 @@
 #include "AI/Navigation/NavigationTypes.h"
 #include "Tickable.h"
 #include "DrawDebugHelpers.h"
-#if UE_ENABLE_INCLUDE_ORDER_DEPRECATED_IN_5_4
-#include "NavFilters/NavigationQueryFilter.h"
-#endif
 #include "CrowdManager.generated.h"
 
 class ANavigationData;
diff --git a/Engine/Source/Runtime/Engine/Classes/AI/Navigation/AvoidanceManager.h b/Engine/Source/Runtime/Engine/Classes/AI/Navigation/AvoidanceManager.h
--- a/Engine/Source/Runtime/Engine/Classes/AI/Navigation/AvoidanceManager.h
+++ b/Engine/Source/Runtime/Engine/Classes/AI/Navigation/AvoidanceManager.h
@@ -2,9 +2,6 @@
 
 #pragma once
 
-#if UE_ENABLE_INCLUDE_ORDER_DEPRECATED_IN_5_4
-#include "CoreMinimal.h"
-#endif
 #include "Stats/Stats.h"
 #include "UObject/ObjectMacros.h"
 #include "UObject/Object.h"
diff --git a/Engine/Source/Runtime/Engine/Classes/AI/Navigation/NavCollisionBase.h b/Engine/Source/Runtime/Engine/Classes/AI/Navigation/NavCollisionBase.h
--- a/Engine/Source/Runtime/Engine/Classes/AI/Navigation/NavCollisionBase.h
+++ b/Engine/Source/Runtime/Engine/Classes/AI/Navigation/NavCollisionBase.h
@@ -2,9 +2,6 @@
 
 #pragma once
 
-#if UE_ENABLE_INCLUDE_ORDER_DEPRECATED_IN_5_4
-#include "CoreMinimal.h"
-#endif
 #include "AI/NavigationSystemHelpers.h"
 #include "AI/Navigation/NavigationTypes.h"
 #include "NavCollisionBase.generated.h"
diff --git a/Engine/Source/Runtime/Engine/Classes/AI/Navigation/NavEdgeProviderInterface.h b/Engine/Source/Runtime/Engine/Classes/AI/Navigation/NavEdgeProviderInterface.h
--- a/Engine/Source/Runtime/Engine/Classes/AI/Navigation/NavEdgeProviderInterface.h
+++ b/Engine/Source/Runtime/Engine/Classes/AI/Navigation/NavEdgeProviderInterface.h
@@ -2,9 +2,6 @@
 
 #pragma once
 
-#if UE_ENABLE_INCLUDE_ORDER_DEPRECATED_IN_5_4
-#include "CoreMinimal.h"
-#endif
 #include "UObject/ObjectMacros.h"
 #include "UObject/Interface.h"
 #include "NavEdgeProviderInterface.generated.h"
diff --git a/Engine/Source/Runtime/Engine/Classes/AI/Navigation/NavLinkDefinition.h b/Engine/Source/Runtime/Engine/Classes/AI/Navigation/NavLinkDefinition.h
--- a/Engine/Source/Runtime/Engine/Classes/AI/Navigation/NavLinkDefinition.h
+++ b/Engine/Source/Runtime/Engine/Classes/AI/Navigation/NavLinkDefinition.h
@@ -2,9 +2,6 @@
 
 #pragma once
 
-#if UE_ENABLE_INCLUDE_ORDER_DEPRECATED_IN_5_4
-#include "CoreMinimal.h"
-#endif
 #include "UObject/ObjectMacros.h"
 #include "UObject/Object.h"
 #include "UObject/Class.h"
diff --git a/Engine/Source/Runtime/Engine/Classes/AI/Navigation/NavPathObserverInterface.h b/Engine/Source/Runtime/Engine/Classes/AI/Navigation/NavPathObserverInterface.h
--- a/Engine/Source/Runtime/Engine/Classes/AI/Navigation/NavPathObserverInterface.h
+++ b/Engine/Source/Runtime/Engine/Classes/AI/Navigation/NavPathObserverInterface.h
@@ -2,9 +2,6 @@
 
 #pragma once
 
-#if UE_ENABLE_INCLUDE_ORDER_DEPRECATED_IN_5_4
-#include "CoreMinimal.h"
-#endif
 #include "UObject/ObjectMacros.h"
 #include "UObject/Interface.h"
 #include "NavPathObserverInterface.generated.h"
diff --git a/Engine/Source/Runtime/Engine/Classes/AI/Navigation/NavRelevantInterface.h b/Engine/Source/Runtime/Engine/Classes/AI/Navigation/NavRelevantInterface.h
--- a/Engine/Source/Runtime/Engine/Classes/AI/Navigation/NavRelevantInterface.h
+++ b/Engine/Source/Runtime/Engine/Classes/AI/Navigation/NavRelevantInterface.h
@@ -2,17 +2,9 @@
 
 #pragma once
 
-#if UE_ENABLE_INCLUDE_ORDER_DEPRECATED_IN_5_4
-#include "CoreMinimal.h"
-#endif
 #include "UObject/ObjectMacros.h"
 #include "UObject/Interface.h"
 #include "AI/Navigation/NavDataGatheringMode.h"
-#if UE_ENABLE_INCLUDE_ORDER_DEPRECATED_IN_5_4
-#include "AI/Navigation/NavigationTypes.h"
-#include "AI/NavigationModifier.h"
-#include "AI/Navigation/NavigationRelevantData.h"
-#endif
 #include "NavRelevantInterface.generated.h"
 
 struct FNavigableGeometryExport;
diff --git a/Engine/Source/Runtime/Engine/Classes/AI/Navigation/NavigationAvoidanceTypes.h b/Engine/Source/Runtime/Engine/Classes/AI/Navigation/NavigationAvoidanceTypes.h
--- a/Engine/Source/Runtime/Engine/Classes/AI/Navigation/NavigationAvoidanceTypes.h
+++ b/Engine/Source/Runtime/Engine/Classes/AI/Navigation/NavigationAvoidanceTypes.h
@@ -2,9 +2,6 @@
 
 #pragma once
 
-#if UE_ENABLE_INCLUDE_ORDER_DEPRECATED_IN_5_4
-#include "CoreMinimal.h"
-#endif
 #include "UObject/ObjectMacros.h"
 #include "NavigationAvoidanceTypes.generated.h"
 
diff --git a/Engine/Source/Runtime/Engine/Classes/AI/Navigation/NavigationDataChunk.h b/Engine/Source/Runtime/Engine/Classes/AI/Navigation/NavigationDataChunk.h
--- a/Engine/Source/Runtime/Engine/Classes/AI/Navigation/NavigationDataChunk.h
+++ b/Engine/Source/Runtime/Engine/Classes/AI/Navigation/NavigationDataChunk.h
@@ -2,9 +2,6 @@
 
 #pragma once
 
-#if UE_ENABLE_INCLUDE_ORDER_DEPRECATED_IN_5_4
-#include "CoreMinimal.h"
-#endif
 #include "UObject/ObjectMacros.h"
 #include "UObject/Object.h"
 #include "NavigationDataChunk.generated.h"
diff --git a/Engine/Source/Runtime/Engine/Classes/AI/Navigation/NavigationDataInterface.h b/Engine/Source/Runtime/Engine/Classes/AI/Navigation/NavigationDataInterface.h
--- a/Engine/Source/Runtime/Engine/Classes/AI/Navigation/NavigationDataInterface.h
+++ b/Engine/Source/Runtime/Engine/Classes/AI/Navigation/NavigationDataInterface.h
@@ -2,9 +2,6 @@
 
 #pragma once
 
-#if UE_ENABLE_INCLUDE_ORDER_DEPRECATED_IN_5_4
-#include "CoreMinimal.h"
-#endif
 #include "UObject/ObjectMacros.h"
 #include "UObject/Interface.h"
 #include "AI/Navigation/NavigationTypes.h"
diff --git a/Engine/Source/Runtime/Engine/Classes/AI/Navigation/NavigationTypes.h b/Engine/Source/Runtime/Engine/Classes/AI/Navigation/NavigationTypes.h
--- a/Engine/Source/Runtime/Engine/Classes/AI/Navigation/NavigationTypes.h
+++ b/Engine/Source/Runtime/Engine/Classes/AI/Navigation/NavigationTypes.h
@@ -8,10 +8,6 @@
 #include "Templates/SubclassOf.h"
 #include "Misc/CoreStats.h"
 #include "UObject/SoftObjectPath.h"
-#if UE_ENABLE_INCLUDE_ORDER_DEPRECATED_IN_5_4
-#include "GameFramework/Actor.h"
-#include "AI/Navigation/NavDataGatheringMode.h"
-#endif // UE_ENABLE_INCLUDE_ORDER_DEPRECATED_IN_5_4
 #if UE_ENABLE_INCLUDE_ORDER_DEPRECATED_IN_5_5
 #include "NavigationDirtyArea.h"
 #endif //UE_ENABLE_INCLUDE_ORDER_DEPRECATED_IN_5_5
diff --git a/Engine/Source/Runtime/Engine/Classes/AI/Navigation/PathFollowingAgentInterface.h b/Engine/Source/Runtime/Engine/Classes/AI/Navigation/PathFollowingAgentInterface.h
--- a/Engine/Source/Runtime/Engine/Classes/AI/Navigation/PathFollowingAgentInterface.h
+++ b/Engine/Source/Runtime/Engine/Classes/AI/Navigation/PathFollowingAgentInterface.h
@@ -2,9 +2,6 @@
 
 #pragma once
 
-#if UE_ENABLE_INCLUDE_ORDER_DEPRECATED_IN_5_4
-#include "CoreMinimal.h"
-#endif
 #include "UObject/ObjectMacros.h"
 #include "UObject/Interface.h"
 #include "PathFollowingAgentInterface.generated.h"
diff --git a/Engine/Source/Runtime/Engine/Classes/Components/InterpToMovementComponent.h b/Engine/Source/Runtime/Engine/Classes/Components/InterpToMovementComponent.h
--- a/Engine/Source/Runtime/Engine/Classes/Components/InterpToMovementComponent.h
+++ b/Engine/Source/Runtime/Engine/Classes/Components/InterpToMovementComponent.h
@@ -2,10 +2,6 @@
 
 #pragma once
 
-#if UE_ENABLE_INCLUDE_ORDER_DEPRECATED_IN_5_4
-#include "CoreMinimal.h"
-#include "Engine/HitResult.h"
-#endif // UE_ENABLE_INCLUDE_ORDER_DEPRECATED_IN_5_4
 #include "GameFramework/MovementComponent.h"
 #include "UObject/ObjectMacros.h"
 #include "InterpToMovementComponent.generated.h"
diff --git a/Engine/Source/Runtime/Engine/Classes/Engine/AssetUserData.h b/Engine/Source/Runtime/Engine/Classes/Engine/AssetUserData.h
--- a/Engine/Source/Runtime/Engine/Classes/Engine/AssetUserData.h
+++ b/Engine/Source/Runtime/Engine/Classes/Engine/AssetUserData.h
@@ -2,9 +2,6 @@
 
 #pragma once
 
-#if UE_ENABLE_INCLUDE_ORDER_DEPRECATED_IN_5_4
-#include "CoreMinimal.h"
-#endif //UE_ENABLE_INCLUDE_ORDER_DEPRECATED_IN_5_4
 #include "UObject/ObjectMacros.h"
 #include "UObject/Object.h"
 #include "AssetUserData.generated.h"
diff --git a/Engine/Source/Runtime/Engine/Classes/Engine/DataAsset.h b/Engine/Source/Runtime/Engine/Classes/Engine/DataAsset.h
--- a/Engine/Source/Runtime/Engine/Classes/Engine/DataAsset.h
+++ b/Engine/Source/Runtime/Engine/Classes/Engine/DataAsset.h
@@ -2,9 +2,6 @@
 
 #pragma once
 
-#if UE_ENABLE_INCLUDE_ORDER_DEPRECATED_IN_5_4
-#include "CoreMinimal.h"
-#endif //UE_ENABLE_INCLUDE_ORDER_DEPRECATED_IN_5_4
 #include "UObject/ObjectMacros.h"
 #include "UObject/Object.h"
 #include "AssetRegistry/AssetBundleData.h"
diff --git a/Engine/Source/Runtime/Engine/Classes/Engine/DataTable.h b/Engine/Source/Runtime/Engine/Classes/Engine/DataTable.h
--- a/Engine/Source/Runtime/Engine/Classes/Engine/DataTable.h
+++ b/Engine/Source/Runtime/Engine/Classes/Engine/DataTable.h
@@ -2,9 +2,6 @@
 
 #pragma once
 
-#if UE_ENABLE_INCLUDE_ORDER_DEPRECATED_IN_5_4
-#include "CoreMinimal.h"
-#endif //UE_ENABLE_INCLUDE_ORDER_DEPRECATED_IN_5_4
 #include "UObject/ObjectMacros.h"
 #include "UObject/Object.h"
 #include "UObject/Class.h"
diff --git a/Engine/Source/Runtime/Engine/Classes/Engine/StaticMeshActor.h b/Engine/Source/Runtime/Engine/Classes/Engine/StaticMeshActor.h
--- a/Engine/Source/Runtime/Engine/Classes/Engine/StaticMeshActor.h
+++ b/Engine/Source/Runtime/Engine/Classes/Engine/StaticMeshActor.h
@@ -5,9 +5,6 @@
 #include "CoreMinimal.h"
 #include "UObject/ObjectMacros.h"
 #include "GameFramework/Actor.h"
-#if UE_ENABLE_INCLUDE_ORDER_DEPRECATED_IN_5_4
-#include "AI/Navigation/NavigationTypes.h"
-#endif //UE_ENABLE_INCLUDE_ORDER_DEPRECATED_IN_5_4
 #include "AI/Navigation/NavDataGatheringMode.h"
 #include "StaticMeshActor.generated.h"
 
diff --git a/Engine/Source/Runtime/Engine/Classes/GameFramework/ProjectileMovementComponent.h b/Engine/Source/Runtime/Engine/Classes/GameFramework/ProjectileMovementComponent.h
--- a/Engine/Source/Runtime/Engine/Classes/GameFramework/ProjectileMovementComponent.h
+++ b/Engine/Source/Runtime/Engine/Classes/GameFramework/ProjectileMovementComponent.h
@@ -2,10 +2,6 @@
 
 #pragma once
 
-#if UE_ENABLE_INCLUDE_ORDER_DEPRECATED_IN_5_4
-#include "CoreMinimal.h"
-#include "Engine/HitResult.h"
-#endif // UE_ENABLE_INCLUDE_ORDER_DEPRECATED_IN_5_4
 #include "GameFramework/MovementComponent.h"
 #include "UObject/ObjectMacros.h"
 #include "ProjectileMovementComponent.generated.h"
diff --git a/Engine/Source/Runtime/Engine/Classes/Kismet/KismetSystemLibrary.h b/Engine/Source/Runtime/Engine/Classes/Kismet/KismetSystemLibrary.h
--- a/Engine/Source/Runtime/Engine/Classes/Kismet/KismetSystemLibrary.h
+++ b/Engine/Source/Runtime/Engine/Classes/Kismet/KismetSystemLibrary.h
@@ -2,15 +2,9 @@
 
 #pragma once
 
-#if UE_ENABLE_INCLUDE_ORDER_DEPRECATED_IN_5_4
-#include "CoreMinimal.h"
-#endif // UE_ENABLE_INCLUDE_ORDER_DEPRECATED_IN_5_4
 #include "UObject/ObjectMacros.h"
 #include "UObject/Object.h"
 #include "Templates/SubclassOf.h"
-#if UE_ENABLE_INCLUDE_ORDER_DEPRECATED_IN_5_4
-#include "Engine/HitResult.h"
-#endif // UE_ENABLE_INCLUDE_ORDER_DEPRECATED_IN_5_4
 #include "UObject/UnrealType.h"
 #include "UObject/ScriptMacros.h"
 #include "UObject/Interface.h"
diff --git a/Engine/Source/Runtime/Engine/Public/AI/NavigationModifier.h b/Engine/Source/Runtime/Engine/Public/AI/NavigationModifier.h
--- a/Engine/Source/Runtime/Engine/Public/AI/NavigationModifier.h
+++ b/Engine/Source/Runtime/Engine/Public/AI/NavigationModifier.h
@@ -1,10 +1,7 @@
 // Copyright Epic Games, Inc. All Rights Reserved.
 
-#pragma once 
+#pragma once
 
-#if UE_ENABLE_INCLUDE_ORDER_DEPRECATED_IN_5_4
-#include "CollisionShape.h"
-#endif
 #include "Templates/SubclassOf.h"
 #include "UObject/WeakObjectPtrTemplates.h"
 #include "AI/Navigation/NavLinkDefinition.h"
diff --git a/Engine/Source/Runtime/Engine/Public/ActorDeferredScriptManager.h b/Engine/Source/Runtime/Engine/Public/ActorDeferredScriptManager.h
--- a/Engine/Source/Runtime/Engine/Public/ActorDeferredScriptManager.h
+++ b/Engine/Source/Runtime/Engine/Public/ActorDeferredScriptManager.h
@@ -75,8 +75,3 @@ private:
 };
 
 #endif
-
-#if UE_ENABLE_INCLUDE_ORDER_DEPRECATED_IN_5_4
-#include "AssetCompilingManager.h"
-#include "AsyncCompilationHelpers.h"
-#endif
diff --git a/Engine/Source/Runtime/Engine/Public/AssetCompilingManager.h b/Engine/Source/Runtime/Engine/Public/AssetCompilingManager.h
--- a/Engine/Source/Runtime/Engine/Public/AssetCompilingManager.h
+++ b/Engine/Source/Runtime/Engine/Public/AssetCompilingManager.h
@@ -5,11 +5,6 @@
 #include "UObject/WeakObjectPtr.h"
 #include "UObject/Class.h"
 #include "UObject/ObjectMacros.h"
-#if UE_ENABLE_INCLUDE_ORDER_DEPRECATED_IN_5_4
-#include "CoreMinimal.h"
-#include "Interfaces/Interface_AsyncCompilation.h"
-#include "IAssetCompilingManager.h"
-#endif
 #include "AssetCompilingManager.generated.h"
 
 class FQueuedThreadPool;
diff --git a/Engine/Source/Runtime/Engine/Public/Collision.h b/Engine/Source/Runtime/Engine/Public/Collision.h
--- a/Engine/Source/Runtime/Engine/Public/Collision.h
+++ b/Engine/Source/Runtime/Engine/Public/Collision.h
@@ -6,14 +6,8 @@
 
 #pragma once
 
-#if UE_ENABLE_INCLUDE_ORDER_DEPRECATED_IN_5_4
-#include "CoreMinimal.h"
-#endif // UE_ENABLE_INCLUDE_ORDER_DEPRECATED_IN_5_4
 #include "Math/Vector.h"
 #include "Stats/Stats.h"
-#if UE_ENABLE_INCLUDE_ORDER_DEPRECATED_IN_5_4
-#include "Engine/HitResult.h"
-#endif // UE_ENABLE_INCLUDE_ORDER_DEPRECATED_IN_5_4
 #include "EngineDefines.h"
 
 struct FHitResult;
diff --git a/Engine/Source/Runtime/Engine/Public/Components.h b/Engine/Source/Runtime/Engine/Public/Components.h
--- a/Engine/Source/Runtime/Engine/Public/Components.h
+++ b/Engine/Source/Runtime/Engine/Public/Components.h
@@ -7,9 +7,6 @@
 #include "Containers/StridedView.h"
 #include "MeshUVChannelInfo.h"
 #include "VertexStreamComponent.h"
-#if UE_ENABLE_INCLUDE_ORDER_DEPRECATED_IN_5_4
-#include "NaniteDefinitions.h"
-#endif
 
 /*=============================================================================
 	Components.h: Forward declarations of object components of actors
diff --git a/Engine/Source/Runtime/Engine/Public/DistanceFieldAtlas.h b/Engine/Source/Runtime/Engine/Public/DistanceFieldAtlas.h
--- a/Engine/Source/Runtime/Engine/Public/DistanceFieldAtlas.h
+++ b/Engine/Source/Runtime/Engine/Public/DistanceFieldAtlas.h
@@ -472,9 +472,3 @@ extern ENGINE_API void AppendDistanceFieldConstantsToClassSchema(FAppendToClassS
 
 extern ENGINE_API FString BuildDistanceFieldDerivedDataKey(const FString& InMeshKey);
 #endif
-
-
-#if UE_ENABLE_INCLUDE_ORDER_DEPRECATED_IN_5_4
-#include "AsyncCompilationHelpers.h"
-#include "AssetCompilingManager.h"
-#endif
diff --git a/Engine/Source/Runtime/Engine/Public/Engine/ScopedMovementUpdate.h b/Engine/Source/Runtime/Engine/Public/Engine/ScopedMovementUpdate.h
--- a/Engine/Source/Runtime/Engine/Public/Engine/ScopedMovementUpdate.h
+++ b/Engine/Source/Runtime/Engine/Public/Engine/ScopedMovementUpdate.h
@@ -2,10 +2,6 @@
 
 #pragma once
 
-#if UE_ENABLE_INCLUDE_ORDER_DEPRECATED_IN_5_4
-#include "CoreMinimal.h"
-#include "Engine/HitResult.h"
-#endif // UE_ENABLE_INCLUDE_ORDER_DEPRECATED_IN_5_4
 #include "Engine/OverlapInfo.h"
 
 struct FHitResult;
diff --git a/Engine/Source/Runtime/Engine/Public/MeshCardBuild.h b/Engine/Source/Runtime/Engine/Public/MeshCardBuild.h
--- a/Engine/Source/Runtime/Engine/Public/MeshCardBuild.h
+++ b/Engine/Source/Runtime/Engine/Public/MeshCardBuild.h
@@ -393,8 +393,3 @@ inline void BeginCacheMeshCardRepresentation(
 	BeginCacheMeshCardRepresentation(StaticMeshAsset, RenderData, DistanceFieldKey, OptionalSourceMeshData ? *OptionalSourceMeshData : FSourceMeshDataForDerivedDataTask());
 }
 #endif
-
-#if UE_ENABLE_INCLUDE_ORDER_DEPRECATED_IN_5_4
-#include "AssetCompilingManager.h"
-#include "AsyncCompilationHelpers.h"
-#endif
diff --git a/Engine/Source/Runtime/Engine/Public/PrimitiveSceneProxy.h b/Engine/Source/Runtime/Engine/Public/PrimitiveSceneProxy.h
--- a/Engine/Source/Runtime/Engine/Public/PrimitiveSceneProxy.h
+++ b/Engine/Source/Runtime/Engine/Public/PrimitiveSceneProxy.h
@@ -16,9 +16,6 @@
 #include "Engine/Scene.h"
 #include "UniformBuffer.h"
 #include "PrimitiveUniformShaderParameters.h"
-#if UE_ENABLE_INCLUDE_ORDER_DEPRECATED_IN_5_4
-#include "InstanceUniformShaderParameters.h"
-#endif
 #include "DrawDebugHelpers.h"
 #include "Math/CapsuleShape.h"
 #include "SceneDefinitions.h"
diff --git a/Engine/Source/Runtime/Engine/Public/SkinnedAssetCompiler.h b/Engine/Source/Runtime/Engine/Public/SkinnedAssetCompiler.h
--- a/Engine/Source/Runtime/Engine/Public/SkinnedAssetCompiler.h
+++ b/Engine/Source/Runtime/Engine/Public/SkinnedAssetCompiler.h
@@ -122,9 +122,3 @@ private:
 };
 
 #endif // #if WITH_EDITOR
-
-#if UE_ENABLE_INCLUDE_ORDER_DEPRECATED_IN_5_4
-#include "CoreMinimal.h"
-#include "AssetCompilingManager.h"
-#include "AsyncCompilationHelpers.h"
-#endif
diff --git a/Engine/Source/Runtime/Engine/Public/SoundWaveCompiler.h b/Engine/Source/Runtime/Engine/Public/SoundWaveCompiler.h
--- a/Engine/Source/Runtime/Engine/Public/SoundWaveCompiler.h
+++ b/Engine/Source/Runtime/Engine/Public/SoundWaveCompiler.h
@@ -100,9 +100,3 @@ private:
 };
 
 #endif
-
-#if UE_ENABLE_INCLUDE_ORDER_DEPRECATED_IN_5_4
-#include "CoreMinimal.h"
-#include "AssetCompilingManager.h"
-#include "AsyncCompilationHelpers.h"
-#endif
diff --git a/Engine/Source/Runtime/Engine/Public/StaticMeshCompiler.h b/Engine/Source/Runtime/Engine/Public/StaticMeshCompiler.h
--- a/Engine/Source/Runtime/Engine/Public/StaticMeshCompiler.h
+++ b/Engine/Source/Runtime/Engine/Public/StaticMeshCompiler.h
@@ -130,9 +130,3 @@ private:
 };
 
 #endif // #if WITH_EDITOR
-
-#if UE_ENABLE_INCLUDE_ORDER_DEPRECATED_IN_5_4
-#include "CoreMinimal.h"
-#include "AssetCompilingManager.h"
-#include "AsyncCompilationHelpers.h"
-#endif
diff --git a/Engine/Source/Runtime/Engine/Public/TextureCompiler.h b/Engine/Source/Runtime/Engine/Public/TextureCompiler.h
--- a/Engine/Source/Runtime/Engine/Public/TextureCompiler.h
+++ b/Engine/Source/Runtime/Engine/Public/TextureCompiler.h
@@ -131,9 +131,3 @@ private:
 };
 
 #endif
-
-#if UE_ENABLE_INCLUDE_ORDER_DEPRECATED_IN_5_4
-#include "CoreMinimal.h"
-#include "AssetCompilingManager.h"
-#include "AsyncCompilationHelpers.h"
-#endif
diff --git a/Engine/Source/Runtime/Engine/Public/WorldCollision.h b/Engine/Source/Runtime/Engine/Public/WorldCollision.h
--- a/Engine/Source/Runtime/Engine/Public/WorldCollision.h
+++ b/Engine/Source/Runtime/Engine/Public/WorldCollision.h
@@ -6,15 +6,8 @@
 #pragma once 
 
 #include "Async/TaskGraphFwd.h"
-#if UE_ENABLE_INCLUDE_ORDER_DEPRECATED_IN_5_4
-#include "CoreMinimal.h"
-#endif // UE_ENABLE_INCLUDE_ORDER_DEPRECATED_IN_5_4
 #include "CollisionQueryParams.h"
 #include "CollisionShape.h"
-#if UE_ENABLE_INCLUDE_ORDER_DEPRECATED_IN_5_4
-#include "Engine/OverlapResult.h"
-#include "Engine/HitResult.h"
-#endif // UE_ENABLE_INCLUDE_ORDER_DEPRECATED_IN_5_4
 
 struct FOverlapDatum;
 struct FTraceDatum;
diff --git a/Engine/Source/Runtime/Experimental/Chaos/Public/Chaos/NewtonEvolution.h b/Engine/Source/Runtime/Experimental/Chaos/Public/Chaos/NewtonEvolution.h
--- a/Engine/Source/Runtime/Experimental/Chaos/Public/Chaos/NewtonEvolution.h
+++ b/Engine/Source/Runtime/Experimental/Chaos/Public/Chaos/NewtonEvolution.h
@@ -8,9 +8,6 @@
 #include "Chaos/PBDSoftsSolverParticles.h"
 #include "Chaos/NewtonElasticFEM.h"
 #include "Chaos/NewtonCorotatedCache.h"
-#if UE_ENABLE_INCLUDE_ORDER_DEPRECATED_IN_5_4
-#include "Chaos/KinematicGeometryParticles.h"
-#endif
 #include "Chaos/SoftsSolverCollisionParticles.h"
 #include "Chaos/VelocityField.h"
 
diff --git a/Engine/Source/Runtime/Experimental/Chaos/Public/Chaos/PBDEvolution.h b/Engine/Source/Runtime/Experimental/Chaos/Public/Chaos/PBDEvolution.h
--- a/Engine/Source/Runtime/Experimental/Chaos/Public/Chaos/PBDEvolution.h
+++ b/Engine/Source/Runtime/Experimental/Chaos/Public/Chaos/PBDEvolution.h
@@ -6,9 +6,6 @@
 #include "Chaos/PBDActiveView.h"
 #include "Chaos/PBDSoftsEvolutionFwd.h"
 #include "Chaos/PBDSoftsSolverParticles.h"
-#if UE_ENABLE_INCLUDE_ORDER_DEPRECATED_IN_5_4
-#include "Chaos/KinematicGeometryParticles.h"
-#endif
 #include "Chaos/SoftsSolverCollisionParticles.h"
 #include "Chaos/VelocityField.h"
 
diff --git a/Engine/Source/Runtime/Experimental/Chaos/Public/Chaos/PerParticlePBDCCDCollisionConstraint.h b/Engine/Source/Runtime/Experimental/Chaos/Public/Chaos/PerParticlePBDCCDCollisionConstraint.h
--- a/Engine/Source/Runtime/Experimental/Chaos/Public/Chaos/PerParticlePBDCCDCollisionConstraint.h
+++ b/Engine/Source/Runtime/Experimental/Chaos/Public/Chaos/PerParticlePBDCCDCollisionConstraint.h
@@ -3,9 +3,6 @@
 
 // HEADER_UNIT_SKIP - Internal
 
-#if UE_ENABLE_INCLUDE_ORDER_DEPRECATED_IN_5_4
-#include "Chaos/GeometryParticlesfwd.h"
-#endif
 #include "Chaos/SoftsSolverCollisionParticles.h"
 #include "Chaos/Transform.h"
 #include "Chaos/PBDActiveView.h"
diff --git a/Engine/Source/Runtime/Experimental/Chaos/Public/Chaos/PerParticlePBDCollisionConstraint.h b/Engine/Source/Runtime/Experimental/Chaos/Public/Chaos/PerParticlePBDCollisionConstraint.h
--- a/Engine/Source/Runtime/Experimental/Chaos/Public/Chaos/PerParticlePBDCollisionConstraint.h
+++ b/Engine/Source/Runtime/Experimental/Chaos/Public/Chaos/PerParticlePBDCollisionConstraint.h
@@ -5,9 +5,6 @@
 #include "Chaos/PBDActiveView.h"
 #include "Chaos/PBDSoftsEvolutionFwd.h"
 #include "Chaos/PBDSoftsSolverParticles.h"
-#if UE_ENABLE_INCLUDE_ORDER_DEPRECATED_IN_5_4
-#include "Chaos/KinematicGeometryParticles.h"
-#endif
 #include "Chaos/SoftsSolverCollisionParticles.h"
 #include "Chaos/WeightedLatticeImplicitObject.h"
 #include "Chaos/Levelset.h"
diff --git a/Engine/Source/Runtime/Experimental/ChaosSolverEngine/Public/Chaos/ChaosNotifyHandlerInterface.h b/Engine/Source/Runtime/Experimental/ChaosSolverEngine/Public/Chaos/ChaosNotifyHandlerInterface.h
--- a/Engine/Source/Runtime/Experimental/ChaosSolverEngine/Public/Chaos/ChaosNotifyHandlerInterface.h
+++ b/Engine/Source/Runtime/Experimental/ChaosSolverEngine/Public/Chaos/ChaosNotifyHandlerInterface.h
@@ -4,9 +4,6 @@
 
 #include "Kismet/BlueprintFunctionLibrary.h"
 #include "UObject/Interface.h"
-#if UE_ENABLE_INCLUDE_ORDER_DEPRECATED_IN_5_4
-#include "Engine/HitResult.h"
-#endif // UE_ENABLE_INCLUDE_ORDER_DEPRECATED_IN_5_4
 #include "ChaosNotifyHandlerInterface.generated.h"
 
 struct FHitResult;
diff --git a/Engine/Source/Runtime/GameplayTags/Classes/GameplayTagContainer.h b/Engine/Source/Runtime/GameplayTags/Classes/GameplayTagContainer.h
--- a/Engine/Source/Runtime/GameplayTags/Classes/GameplayTagContainer.h
+++ b/Engine/Source/Runtime/GameplayTags/Classes/GameplayTagContainer.h
@@ -2,9 +2,6 @@
 
 #pragma once
 
-#if UE_ENABLE_INCLUDE_ORDER_DEPRECATED_IN_5_4
-#include "CoreMinimal.h"
-#endif //UE_ENABLE_INCLUDE_ORDER_DEPRECATED_IN_5_4
 #include "Stats/Stats.h"
 #include "UObject/ObjectMacros.h"
 #include "UObject/Object.h"
diff --git a/Engine/Source/Runtime/NavigationSystem/Public/AbstractNavData.h b/Engine/Source/Runtime/NavigationSystem/Public/AbstractNavData.h
--- a/Engine/Source/Runtime/NavigationSystem/Public/AbstractNavData.h
+++ b/Engine/Source/Runtime/NavigationSystem/Public/AbstractNavData.h
@@ -3,18 +3,9 @@
 #pragma once
 
 #include "AI/Navigation/NavQueryFilter.h"
-#if UE_ENABLE_INCLUDE_ORDER_DEPRECATED_IN_5_4
-#include "AI/Navigation/NavigationTypes.h"
-#endif
 #include "Containers/Array.h"
 #include "Containers/ContainersFwd.h"
-#if UE_ENABLE_INCLUDE_ORDER_DEPRECATED_IN_5_4
-#include "CoreMinimal.h"
-#endif
 #include "Math/Box.h"
-#if UE_ENABLE_INCLUDE_ORDER_DEPRECATED_IN_5_4
-#include "Math/UnrealMathSSE.h"
-#endif
 #include "Math/Vector.h"
 #include "NavigationData.h"
 #include "UObject/ObjectMacros.h"
diff --git a/Engine/Source/Runtime/NavigationSystem/Public/NavAreas/NavArea.h b/Engine/Source/Runtime/NavigationSystem/Public/NavAreas/NavArea.h
--- a/Engine/Source/Runtime/NavigationSystem/Public/NavAreas/NavArea.h
+++ b/Engine/Source/Runtime/NavigationSystem/Public/NavAreas/NavArea.h
@@ -2,9 +2,6 @@
 
 #pragma once
 
-#if UE_ENABLE_INCLUDE_ORDER_DEPRECATED_IN_5_4
-#include "CoreMinimal.h"
-#endif
 #include "UObject/ObjectMacros.h"
 #include "UObject/Object.h"
 #include "Templates/SubclassOf.h"
diff --git a/Engine/Source/Runtime/NavigationSystem/Public/NavAreas/NavAreaMeta.h b/Engine/Source/Runtime/NavigationSystem/Public/NavAreas/NavAreaMeta.h
--- a/Engine/Source/Runtime/NavigationSystem/Public/NavAreas/NavAreaMeta.h
+++ b/Engine/Source/Runtime/NavigationSystem/Public/NavAreas/NavAreaMeta.h
@@ -2,9 +2,6 @@
 
 #pragma once
 
-#if UE_ENABLE_INCLUDE_ORDER_DEPRECATED_IN_5_4
-#include "CoreMinimal.h"
-#endif
 #include "NavAreas/NavArea.h"
 #include "Templates/SubclassOf.h"
 #include "UObject/ObjectMacros.h"
diff --git a/Engine/Source/Runtime/NavigationSystem/Public/NavAreas/NavAreaMeta_SwitchByAgent.h b/Engine/Source/Runtime/NavigationSystem/Public/NavAreas/NavAreaMeta_SwitchByAgent.h
--- a/Engine/Source/Runtime/NavigationSystem/Public/NavAreas/NavAreaMeta_SwitchByAgent.h
+++ b/Engine/Source/Runtime/NavigationSystem/Public/NavAreas/NavAreaMeta_SwitchByAgent.h
@@ -2,9 +2,6 @@
 
 #pragma once
 
-#if UE_ENABLE_INCLUDE_ORDER_DEPRECATED_IN_5_4
-#include "CoreMinimal.h"
-#endif
 #include "UObject/ObjectMacros.h"
 #include "Templates/SubclassOf.h"
 #include "NavAreas/NavAreaMeta.h"
diff --git a/Engine/Source/Runtime/NavigationSystem/Public/NavAreas/NavArea_Default.h b/Engine/Source/Runtime/NavigationSystem/Public/NavAreas/NavArea_Default.h
--- a/Engine/Source/Runtime/NavigationSystem/Public/NavAreas/NavArea_Default.h
+++ b/Engine/Source/Runtime/NavigationSystem/Public/NavAreas/NavArea_Default.h
@@ -2,9 +2,6 @@
 
 #pragma once
 
-#if UE_ENABLE_INCLUDE_ORDER_DEPRECATED_IN_5_4
-#include "CoreMinimal.h"
-#endif
 #include "NavAreas/NavArea.h"
 #include "UObject/ObjectMacros.h"
 #include "UObject/UObjectGlobals.h"
diff --git a/Engine/Source/Runtime/NavigationSystem/Public/NavAreas/NavArea_Null.h b/Engine/Source/Runtime/NavigationSystem/Public/NavAreas/NavArea_Null.h
--- a/Engine/Source/Runtime/NavigationSystem/Public/NavAreas/NavArea_Null.h
+++ b/Engine/Source/Runtime/NavigationSystem/Public/NavAreas/NavArea_Null.h
@@ -2,9 +2,6 @@
 
 #pragma once
 
-#if UE_ENABLE_INCLUDE_ORDER_DEPRECATED_IN_5_4
-#include "CoreMinimal.h"
-#endif
 #include "NavAreas/NavArea.h"
 #include "UObject/ObjectMacros.h"
 #include "UObject/UObjectGlobals.h"
diff --git a/Engine/Source/Runtime/NavigationSystem/Public/NavAreas/NavArea_Obstacle.h b/Engine/Source/Runtime/NavigationSystem/Public/NavAreas/NavArea_Obstacle.h
--- a/Engine/Source/Runtime/NavigationSystem/Public/NavAreas/NavArea_Obstacle.h
+++ b/Engine/Source/Runtime/NavigationSystem/Public/NavAreas/NavArea_Obstacle.h
@@ -2,9 +2,6 @@
 
 #pragma once
 
-#if UE_ENABLE_INCLUDE_ORDER_DEPRECATED_IN_5_4
-#include "CoreMinimal.h"
-#endif
 #include "NavAreas/NavArea.h"
 #include "UObject/ObjectMacros.h"
 #include "UObject/UObjectGlobals.h"
diff --git a/Engine/Source/Runtime/NavigationSystem/Public/NavCollision.h b/Engine/Source/Runtime/NavigationSystem/Public/NavCollision.h
--- a/Engine/Source/Runtime/NavigationSystem/Public/NavCollision.h
+++ b/Engine/Source/Runtime/NavigationSystem/Public/NavCollision.h
@@ -2,9 +2,6 @@
 
 #pragma once
 
-#if UE_ENABLE_INCLUDE_ORDER_DEPRECATED_IN_5_4
-#include "CoreMinimal.h"
-#endif
 #include "UObject/ObjectMacros.h"
 #include "UObject/Object.h"
 #include "Misc/Guid.h"
diff --git a/Engine/Source/Runtime/NavigationSystem/Public/NavDebugTypes.h b/Engine/Source/Runtime/NavigationSystem/Public/NavDebugTypes.h
--- a/Engine/Source/Runtime/NavigationSystem/Public/NavDebugTypes.h
+++ b/Engine/Source/Runtime/NavigationSystem/Public/NavDebugTypes.h
@@ -2,9 +2,6 @@
 
 #pragma once
 
-#if UE_ENABLE_INCLUDE_ORDER_DEPRECATED_IN_5_4
-#include "CoreMinimal.h"
-#endif
 #include "EngineDefines.h"
 
 
diff --git a/Engine/Source/Runtime/NavigationSystem/Public/NavFilters/NavigationQueryFilter.h b/Engine/Source/Runtime/NavigationSystem/Public/NavFilters/NavigationQueryFilter.h
--- a/Engine/Source/Runtime/NavigationSystem/Public/NavFilters/NavigationQueryFilter.h
+++ b/Engine/Source/Runtime/NavigationSystem/Public/NavFilters/NavigationQueryFilter.h
@@ -2,9 +2,6 @@
 
 #pragma once
 
-#if UE_ENABLE_INCLUDE_ORDER_DEPRECATED_IN_5_4
-#include "CoreMinimal.h"
-#endif
 #include "UObject/ObjectMacros.h"
 #include "UObject/Object.h"
 #include "Templates/SubclassOf.h"
diff --git a/Engine/Source/Runtime/NavigationSystem/Public/NavFilters/RecastFilter_UseDefaultArea.h b/Engine/Source/Runtime/NavigationSystem/Public/NavFilters/RecastFilter_UseDefaultArea.h
--- a/Engine/Source/Runtime/NavigationSystem/Public/NavFilters/RecastFilter_UseDefaultArea.h
+++ b/Engine/Source/Runtime/NavigationSystem/Public/NavFilters/RecastFilter_UseDefaultArea.h
@@ -2,9 +2,6 @@
 
 #pragma once
 
-#if UE_ENABLE_INCLUDE_ORDER_DEPRECATED_IN_5_4
-#include "CoreMinimal.h"
-#endif
 #include "NavFilters/NavigationQueryFilter.h"
 #include "UObject/ObjectMacros.h"
 #include "UObject/UObjectGlobals.h"
diff --git a/Engine/Source/Runtime/NavigationSystem/Public/NavGraph/NavigationGraph.h b/Engine/Source/Runtime/NavigationSystem/Public/NavGraph/NavigationGraph.h
--- a/Engine/Source/Runtime/NavigationSystem/Public/NavGraph/NavigationGraph.h
+++ b/Engine/Source/Runtime/NavigationSystem/Public/NavGraph/NavigationGraph.h
@@ -2,9 +2,6 @@
 
 #pragma once
 
-#if UE_ENABLE_INCLUDE_ORDER_DEPRECATED_IN_5_4
-#include "CoreMinimal.h"
-#endif
 #include "UObject/ObjectMacros.h"
 #include "NavigationData.h"
 #include "NavigationGraph.generated.h"
diff --git a/Engine/Source/Runtime/NavigationSystem/Public/NavGraph/NavigationGraphNode.h b/Engine/Source/Runtime/NavigationSystem/Public/NavGraph/NavigationGraphNode.h
--- a/Engine/Source/Runtime/NavigationSystem/Public/NavGraph/NavigationGraphNode.h
+++ b/Engine/Source/Runtime/NavigationSystem/Public/NavGraph/NavigationGraphNode.h
@@ -2,9 +2,6 @@
 
 #pragma once
 
-#if UE_ENABLE_INCLUDE_ORDER_DEPRECATED_IN_5_4
-#include "CoreMinimal.h"
-#endif
 #include "UObject/ObjectMacros.h"
 #include "GameFramework/Actor.h"
 
diff --git a/Engine/Source/Runtime/NavigationSystem/Public/NavGraph/NavigationGraphNodeComponent.h b/Engine/Source/Runtime/NavigationSystem/Public/NavGraph/NavigationGraphNodeComponent.h
--- a/Engine/Source/Runtime/NavigationSystem/Public/NavGraph/NavigationGraphNodeComponent.h
+++ b/Engine/Source/Runtime/NavigationSystem/Public/NavGraph/NavigationGraphNodeComponent.h
@@ -2,9 +2,6 @@
 
 #pragma once
 
-#if UE_ENABLE_INCLUDE_ORDER_DEPRECATED_IN_5_4
-#include "CoreMinimal.h"
-#endif
 #include "UObject/ObjectMacros.h"
 #include "Components/SceneComponent.h"
 #include "NavGraph/NavigationGraph.h"
diff --git a/Engine/Source/Runtime/NavigationSystem/Public/NavLinkComponent.h b/Engine/Source/Runtime/NavigationSystem/Public/NavLinkComponent.h
--- a/Engine/Source/Runtime/NavigationSystem/Public/NavLinkComponent.h
+++ b/Engine/Source/Runtime/NavigationSystem/Public/NavLinkComponent.h
@@ -2,9 +2,6 @@
 
 #pragma once
 
-#if UE_ENABLE_INCLUDE_ORDER_DEPRECATED_IN_5_4
-#include "CoreMinimal.h"
-#endif
 #include "UObject/ObjectMacros.h"
 #include "Templates/SubclassOf.h"
 #include "Components/PrimitiveComponent.h"
diff --git a/Engine/Source/Runtime/NavigationSystem/Public/NavLinkCustomComponent.h b/Engine/Source/Runtime/NavigationSystem/Public/NavLinkCustomComponent.h
--- a/Engine/Source/Runtime/NavigationSystem/Public/NavLinkCustomComponent.h
+++ b/Engine/Source/Runtime/NavigationSystem/Public/NavLinkCustomComponent.h
@@ -2,9 +2,6 @@
 
 #pragma once
 
-#if UE_ENABLE_INCLUDE_ORDER_DEPRECATED_IN_5_4
-#include "CoreMinimal.h"
-#endif
 #include "UObject/ObjectMacros.h"
 #include "Templates/SubclassOf.h"
 #include "Engine/EngineTypes.h"
diff --git a/Engine/Source/Runtime/NavigationSystem/Public/NavLinkCustomInterface.h b/Engine/Source/Runtime/NavigationSystem/Public/NavLinkCustomInterface.h
--- a/Engine/Source/Runtime/NavigationSystem/Public/NavLinkCustomInterface.h
+++ b/Engine/Source/Runtime/NavigationSystem/Public/NavLinkCustomInterface.h
@@ -2,17 +2,11 @@
 
 #pragma once
 
-#if UE_ENABLE_INCLUDE_ORDER_DEPRECATED_IN_5_4
-#include "CoreMinimal.h"
-#endif
 #include "UObject/ObjectMacros.h"
 #include "Templates/SubclassOf.h"
 #include "UObject/Interface.h"
 #include "NavAreas/NavArea.h"
 #include "AI/Navigation/NavLinkDefinition.h"
-#if UE_ENABLE_INCLUDE_ORDER_DEPRECATED_IN_5_4
-#include "Engine/World.h"
-#endif
 #include "Engine/WorldInitializationValues.h"
 #include "NavLinkCustomInterface.generated.h"
 
diff --git a/Engine/Source/Runtime/NavigationSystem/Public/NavLinkHostInterface.h b/Engine/Source/Runtime/NavigationSystem/Public/NavLinkHostInterface.h
--- a/Engine/Source/Runtime/NavigationSystem/Public/NavLinkHostInterface.h
+++ b/Engine/Source/Runtime/NavigationSystem/Public/NavLinkHostInterface.h
@@ -2,9 +2,6 @@
 
 #pragma once
 
-#if UE_ENABLE_INCLUDE_ORDER_DEPRECATED_IN_5_4
-#include "CoreMinimal.h"
-#endif
 #include "UObject/ObjectMacros.h"
 #include "Templates/SubclassOf.h"
 #include "UObject/Interface.h"
diff --git a/Engine/Source/Runtime/NavigationSystem/Public/NavLinkRenderingComponent.h b/Engine/Source/Runtime/NavigationSystem/Public/NavLinkRenderingComponent.h
--- a/Engine/Source/Runtime/NavigationSystem/Public/NavLinkRenderingComponent.h
+++ b/Engine/Source/Runtime/NavigationSystem/Public/NavLinkRenderingComponent.h
@@ -2,9 +2,6 @@
 
 #pragma once
 
-#if UE_ENABLE_INCLUDE_ORDER_DEPRECATED_IN_5_4
-#include "CoreMinimal.h"
-#endif
 #include "UObject/ObjectMacros.h"
 #include "Components/PrimitiveComponent.h"
 #include "NavLinkRenderingComponent.generated.h"
diff --git a/Engine/Source/Runtime/NavigationSystem/Public/NavLinkRenderingProxy.h b/Engine/Source/Runtime/NavigationSystem/Public/NavLinkRenderingProxy.h
--- a/Engine/Source/Runtime/NavigationSystem/Public/NavLinkRenderingProxy.h
+++ b/Engine/Source/Runtime/NavigationSystem/Public/NavLinkRenderingProxy.h
@@ -2,9 +2,6 @@
 
 #pragma once
 
-#if UE_ENABLE_INCLUDE_ORDER_DEPRECATED_IN_5_4
-#include "CoreMinimal.h"
-#endif
 #include "NavAreas/NavArea.h"
 #include "PrimitiveViewRelevance.h"
 #include "PrimitiveSceneProxy.h"
diff --git a/Engine/Source/Runtime/NavigationSystem/Public/NavLinkTrivial.h b/Engine/Source/Runtime/NavigationSystem/Public/NavLinkTrivial.h
--- a/Engine/Source/Runtime/NavigationSystem/Public/NavLinkTrivial.h
+++ b/Engine/Source/Runtime/NavigationSystem/Public/NavLinkTrivial.h
@@ -2,9 +2,6 @@
 
 #pragma once
 
-#if UE_ENABLE_INCLUDE_ORDER_DEPRECATED_IN_5_4
-#include "CoreMinimal.h"
-#endif
 #include "UObject/ObjectMacros.h"
 #include "AI/Navigation/NavLinkDefinition.h"
 #include "NavLinkTrivial.generated.h"
diff --git a/Engine/Source/Runtime/NavigationSystem/Public/NavMesh/NavMeshBoundsVolume.h b/Engine/Source/Runtime/NavigationSystem/Public/NavMesh/NavMeshBoundsVolume.h
--- a/Engine/Source/Runtime/NavigationSystem/Public/NavMesh/NavMeshBoundsVolume.h
+++ b/Engine/Source/Runtime/NavigationSystem/Public/NavMesh/NavMeshBoundsVolume.h
@@ -7,13 +7,7 @@
 
 #pragma once
 
-#if UE_ENABLE_INCLUDE_ORDER_DEPRECATED_IN_5_4
-#include "CoreMinimal.h"
-#endif
 #include "UObject/ObjectMacros.h"
-#if UE_ENABLE_INCLUDE_ORDER_DEPRECATED_IN_5_4
-#include "AI/Navigation/NavigationTypes.h"
-#endif
 #include "GameFramework/Volume.h"
 #include "AI/Navigation/NavAgentSelector.h"
 #include "NavMeshBoundsVolume.generated.h"
diff --git a/Engine/Source/Runtime/NavigationSystem/Public/NavMesh/NavMeshPath.h b/Engine/Source/Runtime/NavigationSystem/Public/NavMesh/NavMeshPath.h
--- a/Engine/Source/Runtime/NavigationSystem/Public/NavMesh/NavMeshPath.h
+++ b/Engine/Source/Runtime/NavigationSystem/Public/NavMesh/NavMeshPath.h
@@ -2,9 +2,6 @@
 
 #pragma once
 
-#if UE_ENABLE_INCLUDE_ORDER_DEPRECATED_IN_5_4
-#include "CoreMinimal.h"
-#endif
 #include "UObject/ObjectMacros.h"
 #include "UObject/Object.h"
 #include "NavigationPath.h"
diff --git a/Engine/Source/Runtime/NavigationSystem/Public/NavMesh/NavTestRenderingComponent.h b/Engine/Source/Runtime/NavigationSystem/Public/NavMesh/NavTestRenderingComponent.h
--- a/Engine/Source/Runtime/NavigationSystem/Public/NavMesh/NavTestRenderingComponent.h
+++ b/Engine/Source/Runtime/NavigationSystem/Public/NavMesh/NavTestRenderingComponent.h
@@ -5,12 +5,6 @@
 #include "Debug/DebugDrawComponent.h"
 #include "AI/Navigation/NavigationTypes.h" // NavNodeRef
 
-#if UE_ENABLE_INCLUDE_ORDER_DEPRECATED_IN_5_4
-#include "PrimitiveViewRelevance.h"
-#include "DynamicMeshBuilder.h"
-#include "DebugRenderSceneProxy.h"
-#endif // UE_ENABLE_INCLUDE_ORDER_DEPRECATED_IN_5_4
-
 #include "NavTestRenderingComponent.generated.h"
 
 #if UE_ENABLE_DEBUG_DRAWING
diff --git a/Engine/Source/Runtime/NavigationSystem/Public/NavMesh/PImplRecastNavMesh.h b/Engine/Source/Runtime/NavigationSystem/Public/NavMesh/PImplRecastNavMesh.h
--- a/Engine/Source/Runtime/NavigationSystem/Public/NavMesh/PImplRecastNavMesh.h
+++ b/Engine/Source/Runtime/NavigationSystem/Public/NavMesh/PImplRecastNavMesh.h
@@ -8,13 +8,7 @@
 
 #pragma once 
 
-#if UE_ENABLE_INCLUDE_ORDER_DEPRECATED_IN_5_4
-#include "CoreMinimal.h"
-#endif
 #include "UObject/WeakObjectPtr.h"
-#if UE_ENABLE_INCLUDE_ORDER_DEPRECATED_IN_5_4
-#include "NavFilters/NavigationQueryFilter.h"
-#endif
 #include "AI/Navigation/NavigationTypes.h"
 #include "NavMesh/RecastNavMesh.h"
 #include "NavMesh/RecastQueryFilter.h"
diff --git a/Engine/Source/Runtime/NavigationSystem/Public/NavMesh/RecastHelpers.h b/Engine/Source/Runtime/NavigationSystem/Public/NavMesh/RecastHelpers.h
--- a/Engine/Source/Runtime/NavigationSystem/Public/NavMesh/RecastHelpers.h
+++ b/Engine/Source/Runtime/NavigationSystem/Public/NavMesh/RecastHelpers.h
@@ -10,15 +10,9 @@
 
 #pragma once
 
-#if UE_ENABLE_INCLUDE_ORDER_DEPRECATED_IN_5_4
-#include "CoreMinimal.h"
-#endif
 #include "Math/Box.h"
 #include "Math/Color.h"
 #include "Math/Matrix.h"
-#if UE_ENABLE_INCLUDE_ORDER_DEPRECATED_IN_5_4
-#include "Math/UnrealMathSSE.h"
-#endif
 #include "Math/Vector.h"
 
 
diff --git a/Engine/Source/Runtime/NavigationSystem/Public/NavMesh/RecastNavMesh.h b/Engine/Source/Runtime/NavigationSystem/Public/NavMesh/RecastNavMesh.h
--- a/Engine/Source/Runtime/NavigationSystem/Public/NavMesh/RecastNavMesh.h
+++ b/Engine/Source/Runtime/NavigationSystem/Public/NavMesh/RecastNavMesh.h
@@ -2,9 +2,6 @@
 
 #pragma once
 
-#if UE_ENABLE_INCLUDE_ORDER_DEPRECATED_IN_5_4
-#include "CoreMinimal.h"
-#endif
 #include "UObject/ObjectMacros.h"
 #include "Templates/SubclassOf.h"
 #include "EngineDefines.h"
diff --git a/Engine/Source/Runtime/NavigationSystem/Public/NavMesh/RecastNavMeshDataChunk.h b/Engine/Source/Runtime/NavigationSystem/Public/NavMesh/RecastNavMeshDataChunk.h
--- a/Engine/Source/Runtime/NavigationSystem/Public/NavMesh/RecastNavMeshDataChunk.h
+++ b/Engine/Source/Runtime/NavigationSystem/Public/NavMesh/RecastNavMeshDataChunk.h
@@ -2,13 +2,7 @@
 
 #pragma once
 
-#if UE_ENABLE_INCLUDE_ORDER_DEPRECATED_IN_5_4
-#include "CoreMinimal.h"
-#endif
 #include "UObject/ObjectMacros.h"
-#if UE_ENABLE_INCLUDE_ORDER_DEPRECATED_IN_5_4
-#include "AI/Navigation/NavigationTypes.h"
-#endif
 #include "AI/Navigation/NavigationDataChunk.h"
 #include "NavMesh/RecastNavMesh.h"
 #include "RecastNavMeshDataChunk.generated.h"
diff --git a/Engine/Source/Runtime/NavigationSystem/Public/NavMesh/RecastNavMeshGenerator.h b/Engine/Source/Runtime/NavigationSystem/Public/NavMesh/RecastNavMeshGenerator.h
--- a/Engine/Source/Runtime/NavigationSystem/Public/NavMesh/RecastNavMeshGenerator.h
+++ b/Engine/Source/Runtime/NavigationSystem/Public/NavMesh/RecastNavMeshGenerator.h
@@ -1,9 +1,6 @@
 // Copyright Epic Games, Inc. All Rights Reserved.
-#pragma once 
+#pragma once
 
-#if UE_ENABLE_INCLUDE_ORDER_DEPRECATED_IN_5_4
-#include "CoreMinimal.h"
-#endif // UE_ENABLE_INCLUDE_ORDER_DEPRECATED_IN_5_4
 #if UE_ENABLE_INCLUDE_ORDER_DEPRECATED_IN_5_5
 #include "Stats/Stats.h"
 #endif // UE_ENABLE_INCLUDE_ORDER_DEPRECATED_IN_5_5
diff --git a/Engine/Source/Runtime/NavigationSystem/Public/NavMesh/RecastQueryFilter.h b/Engine/Source/Runtime/NavigationSystem/Public/NavMesh/RecastQueryFilter.h
--- a/Engine/Source/Runtime/NavigationSystem/Public/NavMesh/RecastQueryFilter.h
+++ b/Engine/Source/Runtime/NavigationSystem/Public/NavMesh/RecastQueryFilter.h
@@ -1,18 +1,12 @@
 // Copyright Epic Games, Inc. All Rights Reserved.
 
-#pragma once 
+#pragma once
 
-#if UE_ENABLE_INCLUDE_ORDER_DEPRECATED_IN_5_4
-#include "CoreMinimal.h"
-#endif
 #include "AI/Navigation/NavigationTypes.h"
 #include "AI/Navigation/NavQueryFilter.h"
 
 #if WITH_RECAST
 
-#if UE_ENABLE_INCLUDE_ORDER_DEPRECATED_IN_5_4
-#include "Detour/DetourNavMesh.h"
-#endif
 #include "Detour/DetourNavMeshQuery.h"
 
 #define RECAST_VERY_SMALL_AGENT_RADIUS 0.0f
diff --git a/Engine/Source/Runtime/NavigationSystem/Public/NavModifierComponent.h b/Engine/Source/Runtime/NavigationSystem/Public/NavModifierComponent.h
--- a/Engine/Source/Runtime/NavigationSystem/Public/NavModifierComponent.h
+++ b/Engine/Source/Runtime/NavigationSystem/Public/NavModifierComponent.h
@@ -2,9 +2,6 @@
 
 #pragma once
 
-#if UE_ENABLE_INCLUDE_ORDER_DEPRECATED_IN_5_4
-#include "CoreMinimal.h"
-#endif
 #include "UObject/ObjectMacros.h"
 #include "Templates/SubclassOf.h"
 #include "NavAreas/NavArea.h"
diff --git a/Engine/Source/Runtime/NavigationSystem/Public/NavModifierVolume.h b/Engine/Source/Runtime/NavigationSystem/Public/NavModifierVolume.h
--- a/Engine/Source/Runtime/NavigationSystem/Public/NavModifierVolume.h
+++ b/Engine/Source/Runtime/NavigationSystem/Public/NavModifierVolume.h
@@ -1,10 +1,7 @@
 // Copyright Epic Games, Inc. All Rights Reserved.
 
-#pragma once 
+#pragma once
 
-#if UE_ENABLE_INCLUDE_ORDER_DEPRECATED_IN_5_4
-#include "CoreMinimal.h"
-#endif
 #include "UObject/ObjectMacros.h"
 #include "Templates/SubclassOf.h"
 #include "AI/Navigation/NavRelevantInterface.h"
diff --git a/Engine/Source/Runtime/NavigationSystem/Public/NavNodeInterface.h b/Engine/Source/Runtime/NavigationSystem/Public/NavNodeInterface.h
--- a/Engine/Source/Runtime/NavigationSystem/Public/NavNodeInterface.h
+++ b/Engine/Source/Runtime/NavigationSystem/Public/NavNodeInterface.h
@@ -2,14 +2,8 @@
 
 #pragma once
 
-#if UE_ENABLE_INCLUDE_ORDER_DEPRECATED_IN_5_4
-#include "CoreMinimal.h"
-#endif
 #include "UObject/ObjectMacros.h"
 #include "UObject/Interface.h"
-#if UE_ENABLE_INCLUDE_ORDER_DEPRECATED_IN_5_4
-#include "AI/Navigation/NavigationTypes.h"
-#endif
 #include "NavNodeInterface.generated.h"
 
 UINTERFACE(MinimalAPI, meta=(CannotImplementInterfaceInBlueprint))
diff --git a/Engine/Source/Runtime/NavigationSystem/Public/NavRelevantComponent.h b/Engine/Source/Runtime/NavigationSystem/Public/NavRelevantComponent.h
--- a/Engine/Source/Runtime/NavigationSystem/Public/NavRelevantComponent.h
+++ b/Engine/Source/Runtime/NavigationSystem/Public/NavRelevantComponent.h
@@ -2,9 +2,6 @@
 
 #pragma once
 
-#if UE_ENABLE_INCLUDE_ORDER_DEPRECATED_IN_5_4
-#include "CoreMinimal.h"
-#endif
 #include "UObject/ObjectMacros.h"
 #include "Components/ActorComponent.h"
 #include "AI/Navigation/NavRelevantInterface.h"
diff --git a/Engine/Source/Runtime/NavigationSystem/Public/NavSystemConfigOverride.h b/Engine/Source/Runtime/NavigationSystem/Public/NavSystemConfigOverride.h
--- a/Engine/Source/Runtime/NavigationSystem/Public/NavSystemConfigOverride.h
+++ b/Engine/Source/Runtime/NavigationSystem/Public/NavSystemConfigOverride.h
@@ -2,9 +2,6 @@
 
 #pragma once
 
-#if UE_ENABLE_INCLUDE_ORDER_DEPRECATED_IN_5_4
-#include "CoreMinimal.h"
-#endif
 #include "UObject/ObjectMacros.h"
 #include "GameFramework/Actor.h"
 #include "AI/NavigationSystemBase.h"
diff --git a/Engine/Source/Runtime/NavigationSystem/Public/NavigationData.h b/Engine/Source/Runtime/NavigationSystem/Public/NavigationData.h
--- a/Engine/Source/Runtime/NavigationSystem/Public/NavigationData.h
+++ b/Engine/Source/Runtime/NavigationSystem/Public/NavigationData.h
@@ -2,18 +2,12 @@
 
 #pragma once
 
-#if UE_ENABLE_INCLUDE_ORDER_DEPRECATED_IN_5_4
-#include "CoreMinimal.h"
-#endif
 #include "Stats/Stats.h"
 #include "UObject/ObjectMacros.h"
 #include "Templates/SubclassOf.h"
 #include "Engine/EngineTypes.h"
 #include "Async/TaskGraphInterfaces.h"
 #include "GameFramework/Actor.h"
-#if UE_ENABLE_INCLUDE_ORDER_DEPRECATED_IN_5_4
-#include "NavFilters/NavigationQueryFilter.h"
-#endif
 #include "AI/Navigation/NavigationTypes.h"
 #include "NavigationSystemTypes.h"
 #include "EngineDefines.h"
diff --git a/Engine/Source/Runtime/NavigationSystem/Public/NavigationInvokerComponent.h b/Engine/Source/Runtime/NavigationSystem/Public/NavigationInvokerComponent.h
--- a/Engine/Source/Runtime/NavigationSystem/Public/NavigationInvokerComponent.h
+++ b/Engine/Source/Runtime/NavigationSystem/Public/NavigationInvokerComponent.h
@@ -1,9 +1,6 @@
 // Copyright Epic Games, Inc. All Rights Reserved.
 #pragma once
 
-#if UE_ENABLE_INCLUDE_ORDER_DEPRECATED_IN_5_4
-#include "CoreMinimal.h"
-#endif
 #include "UObject/ObjectMacros.h"
 #include "UObject/UObjectGlobals.h"
 #include "AI/Navigation/NavAgentSelector.h"
diff --git a/Engine/Source/Runtime/NavigationSystem/Public/NavigationOctree.h b/Engine/Source/Runtime/NavigationSystem/Public/NavigationOctree.h
--- a/Engine/Source/Runtime/NavigationSystem/Public/NavigationOctree.h
+++ b/Engine/Source/Runtime/NavigationSystem/Public/NavigationOctree.h
@@ -2,10 +2,6 @@
 
 #pragma once
 
-#if UE_ENABLE_INCLUDE_ORDER_DEPRECATED_IN_5_4
-#include "CoreMinimal.h"
-#endif
-
 #if UE_ENABLE_INCLUDE_ORDER_DEPRECATED_IN_5_5
 #include "Stats/Stats.h"
 #include "AI/Navigation/NavigationTypes.h"
diff --git a/Engine/Source/Runtime/NavigationSystem/Public/NavigationPath.h b/Engine/Source/Runtime/NavigationSystem/Public/NavigationPath.h
--- a/Engine/Source/Runtime/NavigationSystem/Public/NavigationPath.h
+++ b/Engine/Source/Runtime/NavigationSystem/Public/NavigationPath.h
@@ -2,9 +2,6 @@
 
 #pragma once
 
-#if UE_ENABLE_INCLUDE_ORDER_DEPRECATED_IN_5_4
-#include "CoreMinimal.h"
-#endif
 #include "UObject/ObjectMacros.h"
 #include "UObject/Object.h"
 #include "NavigationData.h"
diff --git a/Engine/Source/Runtime/NavigationSystem/Public/NavigationPathGenerator.h b/Engine/Source/Runtime/NavigationSystem/Public/NavigationPathGenerator.h
--- a/Engine/Source/Runtime/NavigationSystem/Public/NavigationPathGenerator.h
+++ b/Engine/Source/Runtime/NavigationSystem/Public/NavigationPathGenerator.h
@@ -2,9 +2,6 @@
 
 #pragma once
 
-#if UE_ENABLE_INCLUDE_ORDER_DEPRECATED_IN_5_4
-#include "CoreMinimal.h"
-#endif
 #include "UObject/ObjectMacros.h"
 #include "UObject/Interface.h"
 #include "AI/Navigation/NavigationTypes.h"
diff --git a/Engine/Source/Runtime/NavigationSystem/Public/NavigationSystem.h b/Engine/Source/Runtime/NavigationSystem/Public/NavigationSystem.h
--- a/Engine/Source/Runtime/NavigationSystem/Public/NavigationSystem.h
+++ b/Engine/Source/Runtime/NavigationSystem/Public/NavigationSystem.h
@@ -2,17 +2,11 @@
 
 #pragma once
 
-#if UE_ENABLE_INCLUDE_ORDER_DEPRECATED_IN_5_4
-#include "CoreMinimal.h"
-#endif
 #include "UObject/ObjectMacros.h"
 #include "Templates/SubclassOf.h"
 #include "UObject/WeakObjectPtr.h"
 #include "Misc/CoreMisc.h"
 #include "Misc/CoreDelegates.h"
-#if UE_ENABLE_INCLUDE_ORDER_DEPRECATED_IN_5_4
-#include "NavFilters/NavigationQueryFilter.h"
-#endif
 #include "AI/Navigation/NavigationTypes.h"
 #include "AI/Navigation/NavigationDirtyElement.h"
 #include "AI/Navigation/NavigationInvokerPriority.h"
diff --git a/Engine/Source/Runtime/NavigationSystem/Public/NavigationSystemModule.h b/Engine/Source/Runtime/NavigationSystem/Public/NavigationSystemModule.h
--- a/Engine/Source/Runtime/NavigationSystem/Public/NavigationSystemModule.h
+++ b/Engine/Source/Runtime/NavigationSystem/Public/NavigationSystemModule.h
@@ -2,9 +2,6 @@
 
 #pragma once
 
-#if UE_ENABLE_INCLUDE_ORDER_DEPRECATED_IN_5_4
-#include "CoreMinimal.h"
-#endif
 #include "Modules/ModuleManager.h"
 //#include "AI/NavigationSystemBase.h"
 #if WITH_EDITOR
diff --git a/Engine/Source/Runtime/NavigationSystem/Public/NavigationSystemTypes.h b/Engine/Source/Runtime/NavigationSystem/Public/NavigationSystemTypes.h
--- a/Engine/Source/Runtime/NavigationSystem/Public/NavigationSystemTypes.h
+++ b/Engine/Source/Runtime/NavigationSystem/Public/NavigationSystemTypes.h
@@ -2,14 +2,6 @@
 
 #pragma once
 
-#if UE_ENABLE_INCLUDE_ORDER_DEPRECATED_IN_5_4
-#include "CoreMinimal.h"
-#include "NavFilters/NavigationQueryFilter.h"
-#include "AI/NavigationSystemConfig.h"
-#include "AI/Navigation/NavLinkDefinition.h"
-#include "Math/GenericOctreePublic.h"
-#include "AI/NavigationModifier.h"
-#endif
 #include "AI/Navigation/NavigationTypes.h"
 #include "AI/Navigation/NavAgentSelector.h"
 #include "UObject/WeakInterfacePtr.h"
diff --git a/Engine/Source/Runtime/NavigationSystem/Public/NavigationTestingActor.h b/Engine/Source/Runtime/NavigationSystem/Public/NavigationTestingActor.h
--- a/Engine/Source/Runtime/NavigationSystem/Public/NavigationTestingActor.h
+++ b/Engine/Source/Runtime/NavigationSystem/Public/NavigationTestingActor.h
@@ -2,10 +2,6 @@
 
 #pragma once
 
-#if UE_ENABLE_INCLUDE_ORDER_DEPRECATED_IN_5_4
-#include "CoreMinimal.h"
-#include "Stats/Stats.h"
-#endif
 #include "UObject/ObjectMacros.h"
 #include "Templates/SubclassOf.h"
 #include "GameFramework/Actor.h"
diff --git a/Engine/Source/Runtime/RHICore/Public/BoundShaderStateHistory.h b/Engine/Source/Runtime/RHICore/Public/BoundShaderStateHistory.h
--- a/Engine/Source/Runtime/RHICore/Public/BoundShaderStateHistory.h
+++ b/Engine/Source/Runtime/RHICore/Public/BoundShaderStateHistory.h
@@ -60,7 +60,3 @@ private:
 	uint32 NextBoundShaderStateIndex = 0;
 	FCriticalSection BoundShaderStateHistoryLock;
 };
-
-#if UE_ENABLE_INCLUDE_ORDER_DEPRECATED_IN_5_4
-#include "RHI.h"
-#endif
diff --git a/Engine/Source/Runtime/RHICore/Public/RHIDescriptorAllocator.h b/Engine/Source/Runtime/RHICore/Public/RHIDescriptorAllocator.h
--- a/Engine/Source/Runtime/RHICore/Public/RHIDescriptorAllocator.h
+++ b/Engine/Source/Runtime/RHICore/Public/RHIDescriptorAllocator.h
@@ -143,7 +143,3 @@ private:
 	// Offset from start of heap we belong to
 	const uint32 HeapOffset;
 };
-
-#if UE_ENABLE_INCLUDE_ORDER_DEPRECATED_IN_5_4
-#include "RHIDefinitions.h"
-#endif
diff --git a/Engine/Source/Runtime/RHICore/Public/RHIPoolAllocator.h b/Engine/Source/Runtime/RHICore/Public/RHIPoolAllocator.h
--- a/Engine/Source/Runtime/RHICore/Public/RHIPoolAllocator.h
+++ b/Engine/Source/Runtime/RHICore/Public/RHIPoolAllocator.h
@@ -313,9 +313,3 @@ protected:
 	// Max pre-allocated pool
 	int16 OverFlowPoolsCount	= 1;
 };
-
-
-#if UE_ENABLE_INCLUDE_ORDER_DEPRECATED_IN_5_4
-#include "RHI.h"
-#include "RHIResources.h"
-#endif
diff --git a/Engine/Source/Runtime/RenderCore/Public/ClearQuad.h b/Engine/Source/Runtime/RenderCore/Public/ClearQuad.h
--- a/Engine/Source/Runtime/RenderCore/Public/ClearQuad.h
+++ b/Engine/Source/Runtime/RenderCore/Public/ClearQuad.h
@@ -59,12 +59,3 @@ inline void DrawClearQuad(FRHICommandList& RHICmdList, const FLinearColor& Color
 {
 	DrawClearQuadMRT(RHICmdList, true, 1, &Color, false, 0, false, 0, ClearQuadCallbacks);
 }
-
-
-#if UE_ENABLE_INCLUDE_ORDER_DEPRECATED_IN_5_4
-#include "CoreMinimal.h"
-#include "RHI.h"
-#include "RHICommandList.h"
-#include "RHIDefinitions.h"
-#include "RendererInterface.h"
-#endif
\ No newline at end of file
diff --git a/Engine/Source/Runtime/UMG/Public/Blueprint/UserWidget.h b/Engine/Source/Runtime/UMG/Public/Blueprint/UserWidget.h
--- a/Engine/Source/Runtime/UMG/Public/Blueprint/UserWidget.h
+++ b/Engine/Source/Runtime/UMG/Public/Blueprint/UserWidget.h
@@ -2,9 +2,6 @@
 
 #pragma once
 
-#if UE_ENABLE_INCLUDE_ORDER_DEPRECATED_IN_5_4
-#include "CoreMinimal.h"
-#endif //UE_ENABLE_INCLUDE_ORDER_DEPRECATED_IN_5_4
 #include "Blueprint/UMGSequencePlayMode.h"
 #include "Blueprint/WidgetChild.h"
 #include "UObject/ObjectMacros.h"
