Mover: adding query for the next anticipated movement mode.  This assumes the current mode will continue unless a new one is queued.
Also fixes case where queueing "None" would not clear the queued mode.

#Synced-CL 51416072

diff --git a/Engine/Plugins/Experimental/ChaosMover/Source/ChaosMover/Private/Character/ChaosCharacterMoverComponent.cpp b/Engine/Plugins/Experimental/ChaosMover/Source/ChaosMover/Private/Character/ChaosCharacterMoverComponent.cpp
--- a/Engine/Plugins/Experimental/ChaosMover/Source/ChaosMover/Private/Character/ChaosCharacterMoverComponent.cpp
+++ b/Engine/Plugins/Experimental/ChaosMover/Source/ChaosMover/Private/Character/ChaosCharacterMoverComponent.cpp
@@ -69,6 +69,16 @@ void UChaosCharacterMoverComponent::DoQueueNextMode(FName DesiredModeName, bool
 	QueuedModeTransitionName = DesiredModeName;
 }
 
+FName UChaosCharacterMoverComponent::GetNextMovementModeName() const
+{
+	if (QueuedModeTransitionName != NAME_None)
+	{
+		return QueuedModeTransitionName;
+	}
+
+	return GetMovementModeName();
+}
+
 void UChaosCharacterMoverComponent::ClearQueuedMode()
 {
 	QueuedModeTransitionName = NAME_None;
diff --git a/Engine/Plugins/Experimental/ChaosMover/Source/ChaosMover/Public/ChaosMover/Character/ChaosCharacterMoverComponent.h b/Engine/Plugins/Experimental/ChaosMover/Source/ChaosMover/Public/ChaosMover/Character/ChaosCharacterMoverComponent.h
--- a/Engine/Plugins/Experimental/ChaosMover/Source/ChaosMover/Public/ChaosMover/Character/ChaosCharacterMoverComponent.h
+++ b/Engine/Plugins/Experimental/ChaosMover/Source/ChaosMover/Public/ChaosMover/Character/ChaosCharacterMoverComponent.h
@@ -32,6 +32,7 @@ public:
 	CHAOSMOVER_API virtual void ProduceInput(const int32 DeltaTimeMS, FMoverInputCmdContext* Cmd) override;
 	CHAOSMOVER_API virtual void ProduceServerInput(const int32 DeltaTimeMS, FMoverInputCmdContext* Cmd);
 	CHAOSMOVER_API virtual void DoQueueNextMode(FName DesiredModeName, bool bShouldReenter=false) override;
+	CHAOSMOVER_API virtual FName GetNextMovementModeName() const override;
 
 	CHAOSMOVER_API virtual void Crouch() override;
 	CHAOSMOVER_API virtual void UnCrouch() override;
diff --git a/Engine/Plugins/Experimental/Mover/Source/Mover/Private/MovementModeStateMachine.cpp b/Engine/Plugins/Experimental/Mover/Source/Mover/Private/MovementModeStateMachine.cpp
--- a/Engine/Plugins/Experimental/Mover/Source/Mover/Private/MovementModeStateMachine.cpp
+++ b/Engine/Plugins/Experimental/Mover/Source/Mover/Private/MovementModeStateMachine.cpp
@@ -155,6 +155,10 @@ void UMovementModeStateMachine::QueueNextModeInternal(FName DesiredNextModeName,
 			UE_LOGF(LogMover, Warning, "Attempted to queue an unregistered movement mode: %ls on owner %ls", *DesiredNextModeName.ToString(), *GetNameSafe(GetOwnerActor()));
 		}
 	}
+	else
+	{
+		ClearQueuedMode();
+	}
 }
 
 
@@ -537,6 +541,16 @@ const UBaseMovementMode* UMovementModeStateMachine::GetCurrentMode() const
 	return nullptr;
 }
 
+FName UMovementModeStateMachine::GetQueuedModeName() const
+{
+	if (QueuedModeTransition)
+	{
+		return QueuedModeTransition->GetNextModeName();
+	}
+
+	return NAME_None;
+}
+
 const UBaseMovementMode* UMovementModeStateMachine::FindMovementMode(FName ModeName) const
 {
 	if (ModeName != NAME_None && Modes.Contains(ModeName))
diff --git a/Engine/Plugins/Experimental/Mover/Source/Mover/Private/MoverComponent.cpp b/Engine/Plugins/Experimental/Mover/Source/Mover/Private/MoverComponent.cpp
--- a/Engine/Plugins/Experimental/Mover/Source/Mover/Private/MoverComponent.cpp
+++ b/Engine/Plugins/Experimental/Mover/Source/Mover/Private/MoverComponent.cpp
@@ -2458,6 +2458,22 @@ FName UMoverComponent::GetMovementModeName() const
 	return MoverSyncStateDoubleBuffer.GetReadable().MovementMode;
 }
 
+FName UMoverComponent::GetNextMovementModeName() const
+{
+	if (ModeFSM)
+	{
+		const FName QueuedModeName = ModeFSM->GetQueuedModeName();
+
+		if (QueuedModeName != NAME_None)
+		{
+			return QueuedModeName;
+		}
+	}
+
+	return GetMovementModeName();
+}
+
+
 const UBaseMovementMode* UMoverComponent::GetMovementMode() const
 {
 	return GetActiveModeInternal(UBaseMovementMode::StaticClass());
diff --git a/Engine/Plugins/Experimental/Mover/Source/Mover/Public/MovementModeStateMachine.h b/Engine/Plugins/Experimental/Mover/Source/Mover/Public/MovementModeStateMachine.h
--- a/Engine/Plugins/Experimental/Mover/Source/Mover/Public/MovementModeStateMachine.h
+++ b/Engine/Plugins/Experimental/Mover/Source/Mover/Public/MovementModeStateMachine.h
@@ -52,6 +52,7 @@ public:
 	UE_API void OnSimulationRollback(const FMoverSyncState* SyncState, const FMoverAuxStateContext* AuxState, const FMoverTimeStep& NewBaseTimeStep);
 
 	FName GetCurrentModeName() const { return CurrentModeName; }
+	FName GetQueuedModeName() const;
 
 	UE_API const UBaseMovementMode* GetCurrentMode() const;
 
diff --git a/Engine/Plugins/Experimental/Mover/Source/Mover/Public/MoverComponent.h b/Engine/Plugins/Experimental/Mover/Source/Mover/Public/MoverComponent.h
--- a/Engine/Plugins/Experimental/Mover/Source/Mover/Public/MoverComponent.h
+++ b/Engine/Plugins/Experimental/Mover/Source/Mover/Public/MoverComponent.h
@@ -490,7 +490,11 @@ public:	// Queries
 	// Get the current movement mode name
 	UFUNCTION(BlueprintPure, Category = Mover)
 	MOVER_API FName GetMovementModeName() const;
-	
+
+	// Gets the name of the next movement mode, if a change is queued. Returns the current mode name otherwise.
+	UFUNCTION(BlueprintPure, Category = Mover)
+	MOVER_API virtual FName GetNextMovementModeName() const;
+
 	// Get the current movement mode 
 	UFUNCTION(BlueprintPure, Category = Mover)
 	MOVER_API const UBaseMovementMode* GetMovementMode() const;
