ChaosMover: Making FLayeredMove_LinearVelocity async friendly

#Synced-CL 51470979

diff --git a/Engine/Plugins/Experimental/Mover/Source/Mover/Private/DefaultMovementSet/LayeredMoves/BasicLayeredMoves.cpp b/Engine/Plugins/Experimental/Mover/Source/Mover/Private/DefaultMovementSet/LayeredMoves/BasicLayeredMoves.cpp
--- a/Engine/Plugins/Experimental/Mover/Source/Mover/Private/DefaultMovementSet/LayeredMoves/BasicLayeredMoves.cpp
+++ b/Engine/Plugins/Experimental/Mover/Source/Mover/Private/DefaultMovementSet/LayeredMoves/BasicLayeredMoves.cpp
@@ -23,9 +23,30 @@ FLayeredMove_LinearVelocity::FLayeredMove_LinearVelocity()
 {
 }
 
-bool FLayeredMove_LinearVelocity::GenerateMove(const FMoverTickStartData& SimState, const FMoverTimeStep& TimeStep, const UMoverComponent* MoverComp, UMoverBlackboard* SimBlackboard, FProposedMove& OutProposedMove)
-{ 
-	const FMoverDefaultSyncState* SyncState = SimState.SyncState.SyncStateCollection.FindDataByType<FMoverDefaultSyncState>();
+bool FLayeredMove_LinearVelocity::GenerateMove(const FMoverTickStartData& StartState, const FMoverTimeStep& TimeStep, const UMoverComponent* MoverComp, UMoverBlackboard* SimBlackboard, FProposedMove& OutProposedMove)
+{
+	bool GeneratedAProposedMove = GenerateMove_Async(StartState, TimeStep, SimBlackboard, OutProposedMove);
+	
+	if (GeneratedAProposedMove && MagnitudeOverTime && DurationMs > 0)
+	{
+		const double TimeValue = FMath::Clamp((TimeStep.BaseSimTimeMs - StartSimTimeMs) / DurationMs, 0.0, 1.0);
+		const float TimeFactor = MagnitudeOverTime->GetFloatValue(TimeValue);
+		OutProposedMove.LinearVelocity *= TimeFactor;
+	}
+
+	return GeneratedAProposedMove;
+}
+
+bool FLayeredMove_LinearVelocity::GenerateMove_Async(const FMoverTickStartData& StartState, const FMoverTimeStep& TimeStep, UMoverBlackboard* SimBlackboard, FProposedMove& OutProposedMove)
+{
+	if (!ensureMsgf(IsInGameThread() || MagnitudeOverTime == nullptr, TEXT("MagnitudeOverTime is not thread safe and cannot be set in async mode")))
+	{
+		// Setting the duration to 0.0 to also end this move
+		DurationMs = 0.0f;
+		return false;
+	}
+
+	const FMoverDefaultSyncState* SyncState = StartState.SyncState.SyncStateCollection.FindDataByType<FMoverDefaultSyncState>();
 	check(SyncState);
 
 	const float DeltaSeconds = TimeStep.StepMs * 0.001f;
@@ -40,19 +61,12 @@ bool FLayeredMove_LinearVelocity::GenerateMove(const FMoverTickStartData& SimSta
 
 	FVector VelocityThisFrame = Velocity;
 
-	// Put velocity into worldspace
+	// Put velocity into world space
 	if (SettingsFlags & (uint8)ELayeredMove_ConstantVelocitySettingsFlags::VelocityAlwaysRelative)
 	{
 		VelocityThisFrame = SyncState->GetOrientation_WorldSpace().RotateVector(Velocity);
 	}
 
-	if (MagnitudeOverTime && DurationMs > 0)
-	{
-		const double TimeValue = FMath::Clamp((TimeStep.BaseSimTimeMs - StartSimTimeMs) / DurationMs, 0.0, 1.0);
-		const float TimeFactor = MagnitudeOverTime->GetFloatValue(TimeValue);
-		VelocityThisFrame *= TimeFactor;
-	}
-	
 	OutProposedMove.LinearVelocity = VelocityThisFrame;
 
 	return true;
diff --git a/Engine/Plugins/Experimental/Mover/Source/Mover/Public/DefaultMovementSet/LayeredMoves/BasicLayeredMoves.h b/Engine/Plugins/Experimental/Mover/Source/Mover/Public/DefaultMovementSet/LayeredMoves/BasicLayeredMoves.h
--- a/Engine/Plugins/Experimental/Mover/Source/Mover/Public/DefaultMovementSet/LayeredMoves/BasicLayeredMoves.h
+++ b/Engine/Plugins/Experimental/Mover/Source/Mover/Public/DefaultMovementSet/LayeredMoves/BasicLayeredMoves.h
@@ -38,6 +38,7 @@ struct FLayeredMove_LinearVelocity : public FLayeredMoveBase
 	FVector Velocity;
 
 	// Optional curve float for controlling the magnitude of the velocity applied to the actor over the duration of the move
+	// Warning: This should not be set in async mode as it is not thread safe
 	UPROPERTY(EditAnywhere, BlueprintReadWrite, Category = Mover)
 	TObjectPtr<UCurveFloat> MagnitudeOverTime;
 	
@@ -45,8 +46,10 @@ struct FLayeredMove_LinearVelocity : public FLayeredMoveBase
 	UPROPERTY(EditAnywhere, BlueprintReadWrite, Category = Mover)
 	uint8 SettingsFlags;
 
-	// Generate a movement 
+	// Generate a movement
 	UE_API virtual bool GenerateMove(const FMoverTickStartData& StartState, const FMoverTimeStep& TimeStep, const UMoverComponent* MoverComp, UMoverBlackboard* SimBlackboard, FProposedMove& OutProposedMove) override;
+	// Async safe version of movement generation (does not use MagnitudeOverTime)
+	UE_API virtual bool GenerateMove_Async(const FMoverTickStartData& StartState, const FMoverTimeStep& TimeStep, UMoverBlackboard* SimBlackboard, FProposedMove& OutProposedMove) override;
 
 	UE_API virtual FLayeredMoveBase* Clone() const override;
 
