Mover: FLayeredMove_LinearVelocity::SettingsFlags can now be edited as a proper bit flag as intended in editor, both in Blueprint (Make Layered Move Linear Velocity node) and in the details panel of a layered move (in the context of a move based mode I'm experimenting with)

#Synced-CL 51470979

diff --git a/Engine/Plugins/Experimental/Mover/Source/Mover/Public/DefaultMovementSet/LayeredMoves/BasicLayeredMoves.h b/Engine/Plugins/Experimental/Mover/Source/Mover/Public/DefaultMovementSet/LayeredMoves/BasicLayeredMoves.h
--- a/Engine/Plugins/Experimental/Mover/Source/Mover/Public/DefaultMovementSet/LayeredMoves/BasicLayeredMoves.h
+++ b/Engine/Plugins/Experimental/Mover/Source/Mover/Public/DefaultMovementSet/LayeredMoves/BasicLayeredMoves.h
@@ -11,7 +11,7 @@
 class UCurveVector;
 class UCurveFloat;
 
-UENUM(BlueprintType)
+UENUM(BlueprintType, meta=(Bitflags, UseEnumValuesAsMaskValuesInEditor="true"))
 enum class ELayeredMove_ConstantVelocitySettingsFlags : uint8
 {
 	NoFlags						= 0,
@@ -43,7 +43,7 @@ struct FLayeredMove_LinearVelocity : public FLayeredMoveBase
 	TObjectPtr<UCurveFloat> MagnitudeOverTime;
 	
 	// @see ELayeredMove_ConstantVelocitySettingsFlags
-	UPROPERTY(EditAnywhere, BlueprintReadWrite, Category = Mover)
+	UPROPERTY(EditAnywhere, BlueprintReadWrite, Category = Mover, meta=(Bitmask, BitmaskEnum="/Script/Mover.ELayeredMove_ConstantVelocitySettingsFlags"))
 	uint8 SettingsFlags;
 
 	// Generate a movement
