Mover: Remove trajectory prediction warnings in editor preview worlds like thumbnail/blueprint actor preview
#Synced-CL 51770067

diff --git a/Engine/Plugins/Experimental/Mover/Source/Mover/Private/MoverComponent.cpp b/Engine/Plugins/Experimental/Mover/Source/Mover/Private/MoverComponent.cpp
--- a/Engine/Plugins/Experimental/Mover/Source/Mover/Private/MoverComponent.cpp
+++ b/Engine/Plugins/Experimental/Mover/Source/Mover/Private/MoverComponent.cpp
@@ -2496,7 +2496,14 @@ TArray<FTrajectorySampleInfo> UMoverComponent::GetPredictedTrajectory(FMoverPred
 		}
 	}
 
-	UE_LOGF(LogMover, Warning, "Cannot predict trajectory on a MoverComponent that is not initialized. Returning defaulted samples. Actor: %ls  Initialized? %i BegunPlay? %i BeingDestroyed? %i", *GetNameSafe(GetOwner()), HasBeenInitialized(), HasBegunPlay(), IsBeingDestroyed());
+	
+	// Editor preview worlds i.e. thumbnails/blueprint actor previews have undefined init orders. We can skip warning in those cases.
+	// TODO Revisit if we expect/need trajectory to work in an editor preview world.
+	const UWorld* World = GetWorld();
+	if (World && World->WorldType != EWorldType::EditorPreview)
+	{
+		UE_LOGF(LogMover, Warning, "Cannot predict trajectory on a MoverComponent that is not initialized. Returning defaulted samples. Actor: %ls  Initialized? %i BegunPlay? %i BeingDestroyed? %i", *GetNameSafe(GetOwner()), HasBeenInitialized(), HasBegunPlay(), IsBeingDestroyed());
+	}
 
 	TArray<FTrajectorySampleInfo> BlankDefaultSamples;
 	BlankDefaultSamples.AddDefaulted(PredictionParams.NumPredictionSamples);
