Chaos Visual Debugger / Mover Data Tracing: Better support for structuring what appears in CVD's Mover Data plugin, previously all local struts were thrown together into one data collection. Also using Sergio's name serialization to provide better backwards compatibility of FNames (the rest is still not backwards compatible, it's just baby steps)

#Synced-CL 51470979

diff --git a/Engine/Plugins/Experimental/Mover/Source/Mover/Private/ChaosVisualDebugger/MoverCVDRuntimeTrace.cpp b/Engine/Plugins/Experimental/Mover/Source/Mover/Private/ChaosVisualDebugger/MoverCVDRuntimeTrace.cpp
--- a/Engine/Plugins/Experimental/Mover/Source/Mover/Private/ChaosVisualDebugger/MoverCVDRuntimeTrace.cpp
+++ b/Engine/Plugins/Experimental/Mover/Source/Mover/Private/ChaosVisualDebugger/MoverCVDRuntimeTrace.cpp
@@ -25,7 +25,7 @@ namespace UE::MoverUtils
 CVD_DEFINE_OPTIONAL_DATA_CHANNEL(MoverNetworkedData, EChaosVDDataChannelInitializationFlags::CanChangeEnabledState)
 CVD_DEFINE_OPTIONAL_DATA_CHANNEL(MoverLocalSimData, EChaosVDDataChannelInitializationFlags::CanChangeEnabledState)
 
-// FSkipObjectRefsMemoryWriter and FSkipObjectRefsMemoryReader are a 
+// FSkipObjectRefsMemoryWriter and FSkipObjectRefsMemoryReader are a
 // workaround for serializing mover info structs with object references in them, such as the mover base.
 // It currently skips object references altogether, except if those are UScriptStruct objects, which it serializes
 // as the script struct name, hoping the type exists on the receiving end. This allows us to pass FInstancedStruct
@@ -103,10 +103,8 @@ public:
 	}
 };
 
-void FMoverCVDRuntimeTrace::UnwrapSimData(const FMoverCVDSimDataWrapper& InSimDataWrapper, TSharedPtr<FMoverInputCmdContext>& OutInputCmd, TSharedPtr<FMoverSyncState>& OutSyncState, TSharedPtr<FMoverDataCollection>& OutLocalSimData)
+void FMoverCVDRuntimeTrace::UnwrapSimData(const FMoverCVDSimDataWrapper& InSimDataWrapper, TSharedPtr<FMoverInputCmdContext>& OutInputCmd, TSharedPtr<FMoverSyncState>& OutSyncState, TArray<TPair<FName, TSharedPtr<FMoverDataCollection>>>& OutLocalSimDataSections)
 {
-	// Consider using Chaos::VisualDebugger::FChaosVDMemoryReader
-
 	// Input Cmd
 	{
 		// Deserialize into a new struct allocated dynamically
@@ -129,20 +127,21 @@ void FMoverCVDRuntimeTrace::UnwrapSimData(const FMoverCVDSimDataWrapper& InSimDa
 		OutSyncState->SyncStateCollection.SerializeDebugData(ArSyncStateCollectionReader);
 	}
 
+	// Local sim data — one deserialized FMoverDataCollection per named section
+	for (const TPair<FName, TArray<uint8>>& Section : InSimDataWrapper.LocalSimDataSections)
 	{
-		FSkipObjectRefsMemoryReader ArReader(InSimDataWrapper.LocalSimDataBytes, true);
-		OutLocalSimData = MakeShared<FMoverDataCollection>();
-		OutLocalSimData->SerializeDebugData(ArReader);
+		TSharedPtr<FMoverDataCollection> DataCollection = MakeShared<FMoverDataCollection>();
+		FSkipObjectRefsMemoryReader ArReader(Section.Value, true);
+		DataCollection->SerializeDebugData(ArReader);
+		OutLocalSimDataSections.Add(TPair<FName, TSharedPtr<FMoverDataCollection>>(Section.Key, MoveTemp(DataCollection)));
 	}
 }
 
-void FMoverCVDRuntimeTrace::WrapSimData(uint32 SolverID, uint32 ParticleID, const FMoverInputCmdContext& InInputCmd, const FMoverSyncState& InSyncState, const FMoverDataCollection* LocalSimData, FMoverCVDSimDataWrapper& OutSimDataWrapper)
+void FMoverCVDRuntimeTrace::WrapSimData(uint32 SolverID, uint32 ParticleID, const FMoverInputCmdContext& InInputCmd, const FMoverSyncState& InSyncState, const NamedDataCollections* LocalSimDataCollections, FMoverCVDSimDataWrapper& OutSimDataWrapper)
 {
 	OutSimDataWrapper.SolverID = SolverID;
 	OutSimDataWrapper.ParticleID = ParticleID;
 
-	// Consider using Chaos::VisualDebugger::FChaosVDMemoryWriter
-
 	// Input Cmd
 	{
 		FMemoryWriter ArWriter(OutSimDataWrapper.InputCmdBytes, true);
@@ -166,47 +165,20 @@ void FMoverCVDRuntimeTrace::WrapSimData(uint32 SolverID, uint32 ParticleID, cons
 		const_cast<FMoverSyncState&>(InSyncState).SyncStateCollection.SerializeDebugData(ArWriter);
 	}
 
-	// Local sim data (catch all other structs we want to record)
-	{
-		FSkipObjectRefsMemoryWriter ArWriter(OutSimDataWrapper.LocalSimDataBytes, true);
-		if (LocalSimData)
-		{
-			(const_cast<FMoverDataCollection*>(LocalSimData))->SerializeDebugData(ArWriter);
-		}
-		else
-		{
-			FMoverDataCollection EmptyDataCollection;
-			EmptyDataCollection.SerializeDebugData(ArWriter);
-		}
-	}
-}
-
-USTRUCT()
-struct FMoverMergeableDataCollection : FMoverDataCollection
-{
-	void AppendShallow(const FMoverDataCollection& Other)
+	// Local sim data — serialize each named collection into its own section to preserve origin in the CVD viewer
+	if (LocalSimDataCollections)
 	{
-		const FMoverMergeableDataCollection& MergeableOther = reinterpret_cast<const FMoverMergeableDataCollection&>(Other);
-		for (int i = 0; i < MergeableOther.DataArray.Num(); ++i)
+		for (const TPair<FName, const FMoverDataCollection*>& NamedCollection : *LocalSimDataCollections)
 		{
-			if (MergeableOther.DataArray[i].IsValid())
+			TPair<FName, TArray<uint8>>& Section = OutSimDataWrapper.LocalSimDataSections.AddDefaulted_GetRef();
+			Section.Key = NamedCollection.Key;
+			if (NamedCollection.Value)
 			{
-				DataArray.Add(MergeableOther.DataArray[i]);
+				FSkipObjectRefsMemoryWriter ArWriter(Section.Value, true);
+				const_cast<FMoverDataCollection*>(NamedCollection.Value)->SerializeDebugData(ArWriter);
 			}
 		}
 	}
-};
-
-void CombineDataCollections(const NamedDataCollections& DataCollections, FMoverDataCollection& OutDataCollection)
-{
-	FMoverMergeableDataCollection& OutMergeableDataCollection = reinterpret_cast<FMoverMergeableDataCollection&>(OutDataCollection);
-	for (const TPair<FName, const FMoverDataCollection*>& DataCollectionEntry : DataCollections)
-	{
-		if (const FMoverDataCollection* DataCollectionToAppend = DataCollectionEntry.Value)
-		{
-			OutMergeableDataCollection.AppendShallow(*DataCollectionToAppend);
-		}
-	}
 }
 
 void FMoverCVDRuntimeTrace::TraceMoverData(UMoverComponent* MoverComponent, const FMoverInputCmdContext* InputCmd, const FMoverSyncState* SyncState, const NamedDataCollections* LocalSimDataCollections /*= nullptr*/)
@@ -238,15 +210,8 @@ void FMoverCVDRuntimeTrace::TraceMoverData(UMoverComponent* MoverComponent, cons
 		}
 
 		int32 SolverID = CVD_TRACE_GET_SOLVER_ID_FROM_WORLD(World);
-		const FMoverDataCollection* RecordedLocalSimData = nullptr;
-		FMoverDataCollection MergedDataCollection;
-		if (CVDDC_MoverLocalSimData->IsChannelEnabled() && LocalSimDataCollections)
-		{
-			// LocalSimState could add a lot of extra bytes, especially without some sort of delta serialization, so it is only optionally recorded
-			CombineDataCollections(*LocalSimDataCollections, MergedDataCollection);
-			RecordedLocalSimData = &MergedDataCollection;
-		}
-		TraceMoverDataPrivate(SolverID, ParticleID, InputCmd, SyncState, RecordedLocalSimData);
+		const NamedDataCollections* RecordedLocalSimDataCollections = CVDDC_MoverLocalSimData->IsChannelEnabled() ? LocalSimDataCollections : nullptr;
+		TraceMoverDataPrivate(SolverID, ParticleID, InputCmd, SyncState, RecordedLocalSimDataCollections);
 	}
 }
 
@@ -266,22 +231,15 @@ void FMoverCVDRuntimeTrace::TraceMoverData(uint32 SolverID, uint32 ParticleID, c
 	{
 		return;
 	}
-		
-	const FMoverDataCollection* RecordedLocalSimData = nullptr;
-	FMoverDataCollection MergedDataCollection;
-	if (CVDDC_MoverLocalSimData->IsChannelEnabled() && LocalSimDataCollections)
-	{
-		// LocalSimState could add a lot of extra bytes, especially without some sort of delta serialization, so it is only optionally recorded
-		CombineDataCollections(*LocalSimDataCollections, MergedDataCollection);
-		RecordedLocalSimData = &MergedDataCollection;
-	}
-	TraceMoverDataPrivate(SolverID, ParticleID, InputCmd, SyncState, RecordedLocalSimData);
+
+	const NamedDataCollections* RecordedLocalSimDataCollections = CVDDC_MoverLocalSimData->IsChannelEnabled() ? LocalSimDataCollections : nullptr;
+	TraceMoverDataPrivate(SolverID, ParticleID, InputCmd, SyncState, RecordedLocalSimDataCollections);
 }
 
-void FMoverCVDRuntimeTrace::TraceMoverDataPrivate(uint32 SolverID, uint32 ParticleID, const FMoverInputCmdContext* InputCmd, const FMoverSyncState* SyncState, const FMoverDataCollection* LocalSimData /*= nullptr*/)
+void FMoverCVDRuntimeTrace::TraceMoverDataPrivate(uint32 SolverID, uint32 ParticleID, const FMoverInputCmdContext* InputCmd, const FMoverSyncState* SyncState, const NamedDataCollections* LocalSimDataCollections /*= nullptr*/)
 {
 	FMoverCVDSimDataWrapper SimDataWrapper;
-	WrapSimData(SolverID, ParticleID, *InputCmd, *SyncState, LocalSimData, SimDataWrapper);
+	WrapSimData(SolverID, ParticleID, *InputCmd, *SyncState, LocalSimDataCollections, SimDataWrapper);
 	SimDataWrapper.MarkAsValid();
 
 	FChaosVDScopedTLSBufferAccessor TLSDataBuffer;
@@ -293,4 +251,3 @@ void FMoverCVDRuntimeTrace::TraceMoverDataPrivate(uint32 SolverID, uint32 Partic
 } // namespace UE::MoverUtils
 
 #endif // WITH_CHAOS_VISUAL_DEBUGGER
-
diff --git a/Engine/Plugins/Experimental/Mover/Source/Mover/Public/ChaosVisualDebugger/MoverCVDRuntimeTrace.h b/Engine/Plugins/Experimental/Mover/Source/Mover/Public/ChaosVisualDebugger/MoverCVDRuntimeTrace.h
--- a/Engine/Plugins/Experimental/Mover/Source/Mover/Public/ChaosVisualDebugger/MoverCVDRuntimeTrace.h
+++ b/Engine/Plugins/Experimental/Mover/Source/Mover/Public/ChaosVisualDebugger/MoverCVDRuntimeTrace.h
@@ -29,11 +29,11 @@ public:
 	MOVER_API static void TraceMoverData(UMoverComponent* MoverComponent, const FMoverInputCmdContext* InputCmd, const FMoverSyncState* SyncState, const NamedDataCollections* LocalSimDataCollections = nullptr);
 	MOVER_API static void TraceMoverData(uint32 SolverID, uint32 ParticleID, const FMoverInputCmdContext* InputCmd, const FMoverSyncState* SyncState, const NamedDataCollections* LocalSimDataCollections = nullptr);
 
-	MOVER_API static void UnwrapSimData(const FMoverCVDSimDataWrapper& InSimDataWrapper, TSharedPtr<FMoverInputCmdContext>& OutInputCmd, TSharedPtr<FMoverSyncState>& OutSyncState, TSharedPtr<FMoverDataCollection>& OutLocalSimState);
-	MOVER_API static void WrapSimData(uint32 SolverID, uint32 ParticleID, const FMoverInputCmdContext& InInputCmd, const FMoverSyncState& InSyncState, const FMoverDataCollection* LocalSimState, FMoverCVDSimDataWrapper& OutSimDataWrapper);
+	MOVER_API static void UnwrapSimData(const FMoverCVDSimDataWrapper& InSimDataWrapper, TSharedPtr<FMoverInputCmdContext>& OutInputCmd, TSharedPtr<FMoverSyncState>& OutSyncState, TArray<TPair<FName, TSharedPtr<FMoverDataCollection>>>& OutLocalSimDataSections);
+	MOVER_API static void WrapSimData(uint32 SolverID, uint32 ParticleID, const FMoverInputCmdContext& InInputCmd, const FMoverSyncState& InSyncState, const NamedDataCollections* LocalSimDataCollections, FMoverCVDSimDataWrapper& OutSimDataWrapper);
 
 private:
-	static void TraceMoverDataPrivate(uint32 SolverID, uint32 ParticleID, const FMoverInputCmdContext* InputCmd, const FMoverSyncState* SyncState, const FMoverDataCollection* LocalSimState = nullptr);
+	static void TraceMoverDataPrivate(uint32 SolverID, uint32 ParticleID, const FMoverInputCmdContext* InputCmd, const FMoverSyncState* SyncState, const NamedDataCollections* LocalSimDataCollections = nullptr);
 };
 
 // This is all mover data that is networked, either input command (client to server) or sync state (server to client)
@@ -49,8 +49,8 @@ class FMoverCVDRuntimeTrace
 public:
 	static void TraceMoverData(UMoverComponent* MoverComponent, const FMoverInputCmdContext* InputCmd, const FMoverSyncState* SyncState, const NamedDataCollections* FMoverLocalSimState = nullptr) {}
 	static void TraceMoverData(uint32 SolverID, uint32 ParticleID, const FMoverInputCmdContext* InputCmd, const FMoverSyncState* SyncState, const NamedDataCollections* FMoverLocalSimState = nullptr) {}
-	static void UnwrapSimData(const FMoverCVDSimDataWrapper& InSimDataWrapper, TSharedPtr<FMoverInputCmdContext>& OutInputCmd, TSharedPtr<FMoverSyncState>& OutSyncState, TSharedPtr<FMoverDataCollection>& FMoverLocalSimState) {}
-	static void WrapSimData(uint32 SolverID, uint32 ParticleID, const FMoverInputCmdContext& InInputCmd, const FMoverSyncState& InSyncState, const FMoverDataCollection* FMoverLocalSimState, FMoverCVDSimDataWrapper& OutSimDataWrapper) {}
+	static void UnwrapSimData(const FMoverCVDSimDataWrapper& InSimDataWrapper, TSharedPtr<FMoverInputCmdContext>& OutInputCmd, TSharedPtr<FMoverSyncState>& OutSyncState, TArray<TPair<FName, TSharedPtr<FMoverDataCollection>>>& OutLocalSimDataSections) {}
+	static void WrapSimData(uint32 SolverID, uint32 ParticleID, const FMoverInputCmdContext& InInputCmd, const FMoverSyncState& InSyncState, const NamedDataCollections* LocalSimDataCollections, FMoverCVDSimDataWrapper& OutSimDataWrapper) {}
 };
 
 #endif // WITH_CHAOS_VISUAL_DEBUGGER
diff --git a/Engine/Plugins/Experimental/Mover/Source/MoverCVDData/Private/MoverCVDDataWrappers.cpp b/Engine/Plugins/Experimental/Mover/Source/MoverCVDData/Private/MoverCVDDataWrappers.cpp
--- a/Engine/Plugins/Experimental/Mover/Source/MoverCVDData/Private/MoverCVDDataWrappers.cpp
+++ b/Engine/Plugins/Experimental/Mover/Source/MoverCVDData/Private/MoverCVDDataWrappers.cpp
@@ -2,8 +2,11 @@
 
 #include "MoverCVDDataWrappers.h"
 
+#include "Logging/LogMacros.h"
 #include UE_INLINE_GENERATED_CPP_BY_NAME(MoverCVDDataWrappers)
 
+DEFINE_LOG_CATEGORY_STATIC(LogMoverCVDData, Log, All);
+
 FStringView FMoverCVDSimDataWrapper::WrapperTypeName = TEXT("FMoverCVDSimDataWrapper");
 
 bool FMoverCVDSimDataWrapper::Serialize(FArchive& Ar)
@@ -21,7 +24,30 @@ bool FMoverCVDSimDataWrapper::Serialize(FArchive& Ar)
 	Ar << SyncStateDataCollectionBytes;
 	Ar << InputCmdBytes;
 	Ar << InputMoverDataCollectionBytes;
-	Ar << LocalSimDataBytes;
+
+	// Serialize named sections individually so section keys (FNames) are tracked by the CVD name table
+	// when the outer archive is an FChaosVDMemoryWriter/Reader.
+	int32 NumSections = LocalSimDataSections.Num();
+	Ar << NumSections;
+	if (Ar.IsLoading())
+	{
+		// A single particle's local sim data will only ever have a small number of named sections
+		// (e.g. LocalSimInput, InternalSimData, DebugSimData). Guard against corrupt/malicious data
+		// that could cause a huge or negative allocation.
+		constexpr int32 MaxExpectedSections = 256;
+		if (NumSections < 0 || NumSections > MaxExpectedSections)
+		{
+			UE_LOG(LogMoverCVDData, Warning, TEXT("FMoverCVDSimDataWrapper::Serialize - Invalid NumSections (%d), expected 0..%d. Data may be corrupt."), NumSections, MaxExpectedSections);
+			Ar.SetError();
+			return false;
+		}
+		LocalSimDataSections.SetNum(NumSections);
+	}
+	for (TPair<FName, TArray<uint8>>& Section : LocalSimDataSections)
+	{
+		Ar << Section.Key;
+		Ar << Section.Value;
+	}
 
 	return !Ar.IsError();
 }
diff --git a/Engine/Plugins/Experimental/Mover/Source/MoverCVDData/Public/MoverCVDDataWrappers.h b/Engine/Plugins/Experimental/Mover/Source/MoverCVDData/Public/MoverCVDDataWrappers.h
--- a/Engine/Plugins/Experimental/Mover/Source/MoverCVDData/Public/MoverCVDDataWrappers.h
+++ b/Engine/Plugins/Experimental/Mover/Source/MoverCVDData/Public/MoverCVDDataWrappers.h
@@ -25,7 +25,10 @@ struct FMoverCVDSimDataWrapper : public FChaosVDWrapperDataBase
 	TArray<uint8> SyncStateDataCollectionBytes;
 	TArray<uint8> InputCmdBytes;
 	TArray<uint8> InputMoverDataCollectionBytes;
-	TArray<uint8> LocalSimDataBytes;
+
+	/** Named sections of local simulation data (e.g. LocalSimInput, InternalSimData, DebugSimData).
+	 *  Each section is serialized independently so the origin of each struct is preserved in the CVD viewer. */
+	TArray<TPair<FName, TArray<uint8>>> LocalSimDataSections;
 
 	MOVERCVDDATA_API bool Serialize(FArchive& Ar);
 };
diff --git a/Engine/Plugins/Experimental/Mover/Source/MoverCVDEditor/Private/MoverCVDSimDataComponent.cpp b/Engine/Plugins/Experimental/Mover/Source/MoverCVDEditor/Private/MoverCVDSimDataComponent.cpp
--- a/Engine/Plugins/Experimental/Mover/Source/MoverCVDEditor/Private/MoverCVDSimDataComponent.cpp
+++ b/Engine/Plugins/Experimental/Mover/Source/MoverCVDEditor/Private/MoverCVDSimDataComponent.cpp
@@ -35,7 +35,7 @@ void UMoverCVDSimDataComponent::ClearData()
 	FrameSimDataArray.Reset();
 }
 
-bool UMoverCVDSimDataComponent::FindAndUnwrapSimDataForParticle(uint32 ParticleID, TSharedPtr<FMoverCVDSimDataWrapper>& OutSimDataWrapper, TSharedPtr<FMoverSyncState>& OutSyncState, TSharedPtr<FMoverInputCmdContext>& OutInputCmd, TSharedPtr<FMoverDataCollection>& OutLocalSimData)
+bool UMoverCVDSimDataComponent::FindAndUnwrapSimDataForParticle(uint32 ParticleID, TSharedPtr<FMoverCVDSimDataWrapper>& OutSimDataWrapper, TSharedPtr<FMoverSyncState>& OutSyncState, TSharedPtr<FMoverInputCmdContext>& OutInputCmd, TArray<TPair<FName, TSharedPtr<FMoverDataCollection>>>& OutLocalSimDataSections)
 {
 	// Look for a sim data corresponding to ParticleID
 	TSharedPtr<FMoverCVDSimDataWrapper>* FoundSimData = FrameSimDataArray.FindByPredicate
@@ -62,13 +62,13 @@ bool UMoverCVDSimDataComponent::FindAndUnwrapSimDataForParticle(uint32 ParticleI
 		DeserializedMoverStates = MakeShared<FDeserializedMoverStates>();
 		// This means that the SimData wasn't deserialized yet, so we do it now
 		// Otherwise we use the cached version
-		UE::MoverUtils::FMoverCVDRuntimeTrace::UnwrapSimData(*SimDataPtr, DeserializedMoverStates->InputCommand, DeserializedMoverStates->SyncState, DeserializedMoverStates->LocalSimData);
+		UE::MoverUtils::FMoverCVDRuntimeTrace::UnwrapSimData(*SimDataPtr, DeserializedMoverStates->InputCommand, DeserializedMoverStates->SyncState, DeserializedMoverStates->LocalSimDataSections);
 	}
 	if (DeserializedMoverStates)
 	{
 		OutSyncState = DeserializedMoverStates->SyncState;
 		OutInputCmd = DeserializedMoverStates->InputCommand;
-		OutLocalSimData = DeserializedMoverStates->LocalSimData;
+		OutLocalSimDataSections = DeserializedMoverStates->LocalSimDataSections;
 	}
 
 	return true;
diff --git a/Engine/Plugins/Experimental/Mover/Source/MoverCVDEditor/Private/MoverCVDSimDataComponent.h b/Engine/Plugins/Experimental/Mover/Source/MoverCVDEditor/Private/MoverCVDSimDataComponent.h
--- a/Engine/Plugins/Experimental/Mover/Source/MoverCVDEditor/Private/MoverCVDSimDataComponent.h
+++ b/Engine/Plugins/Experimental/Mover/Source/MoverCVDEditor/Private/MoverCVDSimDataComponent.h
@@ -14,7 +14,7 @@ struct FDeserializedMoverStates
 {
 	TSharedPtr<FMoverSyncState> SyncState;
 	TSharedPtr<FMoverInputCmdContext> InputCommand;
-	TSharedPtr<FMoverDataCollection> LocalSimData;
+	TArray<TPair<FName, TSharedPtr<FMoverDataCollection>>> LocalSimDataSections;
 };
 
 /** Component holding Mover data for the current visualized frame */
@@ -31,7 +31,7 @@ public:
 	virtual void UpdateFromSolverFrameData(const FChaosVDSolverFrameData& InSolverFrameData) override;
 
 	virtual void ClearData() override;
-	
+
 	TConstArrayView<TSharedPtr<FMoverCVDSimDataWrapper>> GetFrameSimDataArray() const
 	{
 		return FrameSimDataArray;
@@ -42,7 +42,7 @@ public:
 		TSharedPtr<FMoverCVDSimDataWrapper>& OutSimDataWrapper,
 		TSharedPtr<FMoverSyncState>& OutSyncState,
 		TSharedPtr<FMoverInputCmdContext>& OutInputCmd,
-		TSharedPtr<FMoverDataCollection>& OutLocalSimData);
+		TArray<TPair<FName, TSharedPtr<FMoverDataCollection>>>& OutLocalSimDataSections);
 
 private:
 	// This is the array of FMoverCVDSimDataWrapper for the current frame (as it is updated in UpdateFromNewGameFrameData)
diff --git a/Engine/Plugins/Experimental/Mover/Source/MoverCVDEditor/Private/MoverCVDTab.cpp b/Engine/Plugins/Experimental/Mover/Source/MoverCVDEditor/Private/MoverCVDTab.cpp
--- a/Engine/Plugins/Experimental/Mover/Source/MoverCVDEditor/Private/MoverCVDTab.cpp
+++ b/Engine/Plugins/Experimental/Mover/Source/MoverCVDEditor/Private/MoverCVDTab.cpp
@@ -120,7 +120,7 @@ void FMoverCVDTab::DisplaySingleParticleInfo(int32 SelectedSolverID, int32 Selec
 			UMoverCVDSimDataComponent* SimDataComponent = SimDataComponentSharedPtr.Get();
 			if (IsValid(SimDataComponent))
 			{
-				if (SimDataComponent->FindAndUnwrapSimDataForParticle(SelectedParticleID, MoverSimDataWrapper, MoverSyncState, MoverInputCmd, MoverLocalSimData))
+				if (SimDataComponent->FindAndUnwrapSimDataForParticle(SelectedParticleID, MoverSimDataWrapper, MoverSyncState, MoverInputCmd, MoverLocalSimDataSections))
 				{
 					if (MoverSimDataWrapper && MoverInputCmd && MoverSyncState)
 					{
@@ -128,6 +128,9 @@ void FMoverCVDTab::DisplaySingleParticleInfo(int32 SelectedSolverID, int32 Selec
 						{
 							MultiViewWrapper.Clear();
 
+							static const FName InputCommandCategory("InputCommand");
+							static const FName SyncStateCategory("SyncState");
+
 							// This displays the particle ID and solver ID
 							MultiViewWrapper.AddData(MoverSimDataWrapper.Get());
 
@@ -136,28 +139,31 @@ void FMoverCVDTab::DisplaySingleParticleInfo(int32 SelectedSolverID, int32 Selec
 							{
 								if (FMoverDataStructBase* MoverDataStructBase = InputDataStructSharedPtr.Get())
 								{
-									MultiViewWrapper.AddData(MakeShared<FStructOnScope>(MoverDataStructBase->GetScriptStruct(), reinterpret_cast<uint8*>(MoverDataStructBase)));
+									MultiViewWrapper.AddData(MakeShared<FStructOnScope>(MoverDataStructBase->GetScriptStruct(), reinterpret_cast<uint8*>(MoverDataStructBase)), InputCommandCategory);
 								}
 							}
 
-							MultiViewWrapper.AddData(MoverSyncState.Get());
+							MultiViewWrapper.AddData(MoverSyncState.Get(), SyncStateCategory);
 							// Sync State Data Collection structs
-							for (const TSharedPtr<FMoverDataStructBase>& InputDataStructSharedPtr : MoverSyncState->SyncStateCollection.GetDataArray())
+							for (const TSharedPtr<FMoverDataStructBase>& SyncStateDataStructSharedPtr : MoverSyncState->SyncStateCollection.GetDataArray())
 							{
-								if (FMoverDataStructBase* MoverDataStructBase = InputDataStructSharedPtr.Get())
+								if (FMoverDataStructBase* MoverDataStructBase = SyncStateDataStructSharedPtr.Get())
 								{
-									MultiViewWrapper.AddData(MakeShared<FStructOnScope>(MoverDataStructBase->GetScriptStruct(), reinterpret_cast<uint8*>(MoverDataStructBase)));
+									MultiViewWrapper.AddData(MakeShared<FStructOnScope>(MoverDataStructBase->GetScriptStruct(), reinterpret_cast<uint8*>(MoverDataStructBase)), SyncStateCategory);
 								}
 							}
 
-							// Local Simulation State structs
-							if (MoverLocalSimData)
+							// Local Simulation State — each named section becomes its own category in the details panel
+							for (const TPair<FName, TSharedPtr<FMoverDataCollection>>& Section : MoverLocalSimDataSections)
 							{
-								for (const TSharedPtr<FMoverDataStructBase>& LocalSimDataStructSharedPtr : MoverLocalSimData->GetDataArray())
+								if (Section.Value)
 								{
-									if (FMoverDataStructBase* MoverDataStructBase = LocalSimDataStructSharedPtr.Get())
+									for (const TSharedPtr<FMoverDataStructBase>& LocalSimDataStructSharedPtr : Section.Value->GetDataArray())
 									{
-										MultiViewWrapper.AddData(MakeShared<FStructOnScope>(MoverDataStructBase->GetScriptStruct(), reinterpret_cast<uint8*>(MoverDataStructBase)));
+										if (FMoverDataStructBase* MoverDataStructBase = LocalSimDataStructSharedPtr.Get())
+										{
+											MultiViewWrapper.AddData(MakeShared<FStructOnScope>(MoverDataStructBase->GetScriptStruct(), reinterpret_cast<uint8*>(MoverDataStructBase)), Section.Key);
+										}
 									}
 								}
 							}
@@ -175,7 +181,7 @@ void FMoverCVDTab::DisplaySingleParticleInfo(int32 SelectedSolverID, int32 Selec
 			}
 		}
 	}
-	
+
 	if (!WasDataFoundForParticle)
 	{
 		CurrentlyDisplayedSolverID = INDEX_NONE;
@@ -239,9 +245,10 @@ void FMoverCVDTab::HandlePostSelectionChange(const UTypedElementSelectionSet* Ch
 {
 	// This is called when the selection changes
 	TArray<FTypedElementHandle> SelectedParticlesHandles;
-	ChangedSelectionSet->GetSelectedElementHandles(SelectedParticlesHandles, UChaosVDSelectionInterface::StaticClass());	
+	ChangedSelectionSet->GetSelectedElementHandles(SelectedParticlesHandles, UChaosVDSelectionInterface::StaticClass());
 	DisplayMoverInfoForSelectedElements(SelectedParticlesHandles);
 }
 
 #undef LOCTEXT_NAMESPACE
 
+
diff --git a/Engine/Plugins/Experimental/Mover/Source/MoverCVDEditor/Private/MoverCVDTab.h b/Engine/Plugins/Experimental/Mover/Source/MoverCVDEditor/Private/MoverCVDTab.h
--- a/Engine/Plugins/Experimental/Mover/Source/MoverCVDEditor/Private/MoverCVDTab.h
+++ b/Engine/Plugins/Experimental/Mover/Source/MoverCVDEditor/Private/MoverCVDTab.h
@@ -56,7 +56,7 @@ private:
 	TSharedPtr<FMoverCVDSimDataWrapper> MoverSimDataWrapper;
 	TSharedPtr<FMoverSyncState> MoverSyncState;
 	TSharedPtr<FMoverInputCmdContext> MoverInputCmd;
-	TSharedPtr<FMoverDataCollection> MoverLocalSimData;
+	TArray<TPair<FName, TSharedPtr<FMoverDataCollection>>> MoverLocalSimDataSections;
 };
 
 
