Mover: adding simulation context as a parameter to core Movement Mode functions.
This is a step towards decoupling movement from Actor/ActorComponent requirements, and opens the door for modes to NOT need to rely on MoverComponent.  They will eventually solely rely on the Simulation object.
Also adds more parameters to in-simulation Activate/Deactivate functions.

#Synced-CL 51816243

#Synced-CL 51816243

diff --git a/Engine/Plugins/Experimental/ChaosMover/Source/ChaosMover/Private/Backends/ChaosMoverAsyncCallback.cpp b/Engine/Plugins/Experimental/ChaosMover/Source/ChaosMover/Private/Backends/ChaosMoverAsyncCallback.cpp
--- a/Engine/Plugins/Experimental/ChaosMover/Source/ChaosMover/Private/Backends/ChaosMoverAsyncCallback.cpp
+++ b/Engine/Plugins/Experimental/ChaosMover/Source/ChaosMover/Private/Backends/ChaosMoverAsyncCallback.cpp
@@ -104,7 +104,12 @@ namespace UE::ChaosMover
 			AsyncOutput.TimeStep[BackendIndex] = TimeStep;
 			if (TStrongObjectPtr<UChaosMoverBackendComponent> Backend = AsyncInput->Backends[BackendIndex].Pin())
 			{
-				Backend->GetSimulation()->SimulationTick(TimeStep, AsyncInput->InputData[BackendIndex], AsyncOutput.OutputData[BackendIndex]);
+				FMoverSimContext SimContext;
+				SimContext.Simulation = Backend->GetSimulation();
+				SimContext.Blackboard = SimContext.Simulation->GetRollbackBlackboardSimWrapper();
+				SimContext.SimulationOwner = nullptr;	// Not required for ChaosMover modes
+
+				Backend->GetSimulation()->SimulationTick(SimContext, TimeStep, AsyncInput->InputData[BackendIndex], AsyncOutput.OutputData[BackendIndex]);
 			}
 		};
 
diff --git a/Engine/Plugins/Experimental/ChaosMover/Source/ChaosMover/Private/ChaosMoverSimulation.cpp b/Engine/Plugins/Experimental/ChaosMover/Source/ChaosMover/Private/ChaosMoverSimulation.cpp
--- a/Engine/Plugins/Experimental/ChaosMover/Source/ChaosMover/Private/ChaosMoverSimulation.cpp
+++ b/Engine/Plugins/Experimental/ChaosMover/Source/ChaosMover/Private/ChaosMoverSimulation.cpp
@@ -301,6 +301,14 @@ void UChaosMoverSimulation::PredictTrajectory(const FMoverTimeStep& TimeStep, co
 	StepState.SyncState = StartSyncState;
 	StepState.AuxState = InputData.AuxInputState;
 
+	FMoverSimContext SimContext;
+	// TODO: Need to figure out how to put the simulation in predictive mode. 
+	// Currently, any GenerateMove call could find the Simulation object and perform a non-const operation. 
+	// So this const-cast is not exposing anything new.
+	SimContext.Simulation = const_cast<UChaosMoverSimulation*>(this);	
+	SimContext.Blackboard = SimContext.Simulation->GetRollbackBlackboardPredictionWrapper();
+	SimContext.SimulationOwner = nullptr;
+
 	FMoverDefaultSyncState* StepDefaultSyncState = StepState.SyncState.SyncStateCollection.FindMutableDataByType<FMoverDefaultSyncState>();
 	if (!StepDefaultSyncState)
 	{
@@ -354,7 +362,7 @@ void UChaosMoverSimulation::PredictTrajectory(const FMoverTimeStep& TimeStep, co
 		}
 
 		FProposedMove Proposed;
-		Mode->GenerateMove(StepState, FutureTimeStep, Proposed);
+		Mode->GenerateMove(SimContext, StepState, FutureTimeStep, Proposed);
 
 		// Advance step (simple integration)
 		const FVector NextPos = StepDefaultSyncState->GetLocation_WorldSpace() + (Proposed.LinearVelocity * Dt);
@@ -442,7 +450,7 @@ void UChaosMoverSimulation::ProcessInputs(int32 PhysicsStep, const FMoverTimeSte
 	}
 }
 
-void UChaosMoverSimulation::SimulationTick(const FMoverTimeStep& TimeStep, const UE::ChaosMover::FSimulationInputData& InputData, UE::ChaosMover::FSimulationOutputData& OutputData)
+void UChaosMoverSimulation::SimulationTick(const FMoverSimContext& SimContext, const FMoverTimeStep& TimeStep, const UE::ChaosMover::FSimulationInputData& InputData, UE::ChaosMover::FSimulationOutputData& OutputData)
 {
 	QUICK_SCOPE_CYCLE_COUNTER(ChaosMoverSimulation_SimulationTick);
 
@@ -456,7 +464,7 @@ void UChaosMoverSimulation::SimulationTick(const FMoverTimeStep& TimeStep, const
 		RollbackBlackboard->BeginSimulationFrame(TimeStep);
 
 		OnPreSimulationTick(TimeStep, InputData);
-		OnSimulationTick(TimeStep, InputData, OutputData);
+		OnSimulationTick(SimContext, TimeStep, InputData, OutputData);
 		OnPostSimulationTick(TimeStep, OutputData);
 
 		RollbackBlackboard->EndSimulationFrame();
@@ -643,7 +651,7 @@ void UChaosMoverSimulation::OnPreSimulationTick(const FMoverTimeStep& TimeStep,
 	}
 }
 
-void UChaosMoverSimulation::OnSimulationTick(const FMoverTimeStep& TimeStep, const UE::ChaosMover::FSimulationInputData& InputData, UE::ChaosMover::FSimulationOutputData& OutputData)
+void UChaosMoverSimulation::OnSimulationTick(const FMoverSimContext& SimContext, const FMoverTimeStep& TimeStep, const UE::ChaosMover::FSimulationInputData& InputData, UE::ChaosMover::FSimulationOutputData& OutputData)
 {
 	QUICK_SCOPE_CYCLE_COUNTER(ChaosMoverSimulation_OnSimulationTick);
 
@@ -652,7 +660,7 @@ void UChaosMoverSimulation::OnSimulationTick(const FMoverTimeStep& TimeStep, con
 	FMoverTickStartData TickStartData(InputData.InputCmd, CurrentSyncState, InputData.AuxInputState);
 	FMoverTickEndData TickEndData(&CurrentSyncState, &InputData.AuxInputState);
 
-	StateMachine.OnSimulationTick(TimeStep, TickStartData, Blackboard.Get(), MovementMixerWeakPtr.Get(), TickEndData);
+	StateMachine.OnSimulationTick(SimContext, TimeStep, TickStartData, Blackboard.Get(), MovementMixerWeakPtr.Get(), TickEndData);
 
 	if (const FChaosMoverTrajectoryPredictionInputs* PredInputs = LocalSimInput.FindDataByType<FChaosMoverTrajectoryPredictionInputs>())
 	{
diff --git a/Engine/Plugins/Experimental/ChaosMover/Source/ChaosMover/Private/ChaosMoverStateMachine.cpp b/Engine/Plugins/Experimental/ChaosMover/Source/ChaosMover/Private/ChaosMoverStateMachine.cpp
--- a/Engine/Plugins/Experimental/ChaosMover/Source/ChaosMover/Private/ChaosMoverStateMachine.cpp
+++ b/Engine/Plugins/Experimental/ChaosMover/Source/ChaosMover/Private/ChaosMoverStateMachine.cpp
@@ -66,7 +66,7 @@ namespace UE::ChaosMover
 			DefaultModeName = ModeName;
 		}
 
-		Mode->OnRegistered(ModeName);
+		Mode->OnRegistered(ModeName, GetSimContext());
 	}
 
 	void FMoverStateMachine::UnregisterMovementMode(FName ModeName)
@@ -76,7 +76,7 @@ namespace UE::ChaosMover
 
 		if (ModeToUnregister)
 		{
-			ModeToUnregister->OnUnregistered();
+			ModeToUnregister->OnUnregistered(GetSimContext());
 		}
 	}
 
@@ -88,7 +88,7 @@ namespace UE::ChaosMover
 
 			if (Mode)
 			{
-				Mode->OnUnregistered();
+				Mode->OnUnregistered(GetSimContext());
 			}
 		}
 		Modes.Empty();
@@ -386,7 +386,7 @@ namespace UE::ChaosMover
 		}
 	}
 
-	void FMoverStateMachine::OnSimulationTick(const FMoverTimeStep& TimeStep, const FMoverTickStartData& StartState, UMoverBlackboard* SimBlackboard, UMovementMixer* MovementMixer, FMoverTickEndData& OutputState)
+	void FMoverStateMachine::OnSimulationTick(const FMoverSimContext& SimContext, const FMoverTimeStep& TimeStep, const FMoverTickStartData& StartState, UMoverBlackboard* SimBlackboard, UMovementMixer* MovementMixer, FMoverTickEndData& OutputState)
 	{
 		QUICK_SCOPE_CYCLE_COUNTER(MoverStateMachine_OnSimulationTick);
 
@@ -394,6 +394,9 @@ namespace UE::ChaosMover
 		WorkingSubstepStartData = StartState;
 		bool bIsWorkingStartStateReady = true;	// this flag is used to avoid unneeded data copying after substeps
 
+		WorkingSimTickEndData = &OutputState;
+		ON_SCOPE_EXIT{ WorkingSimTickEndData = nullptr; };
+
 		if (!ensure(MovementMixer))
 		{
 			return;
@@ -553,7 +556,7 @@ namespace UE::ChaosMover
 						!(bHasLayeredMoveContributions && CombinedLayeredMove.MixMode == EMoveMixMode::OverrideAll))
 					{
 						QUICK_SCOPE_CYCLE_COUNTER(MoverStateMachine_OnSimulationTick_GenerateMoveFromMode);
-						CurrentMode->GenerateMove(WorkingSubstepStartData, SubTimeStep, OUT CombinedMove);
+						CurrentMode->GenerateMove(SimContext, WorkingSubstepStartData, SubTimeStep, OUT CombinedMove);
 
 						bHasModeMoveContribution = true;
 					}
@@ -590,6 +593,7 @@ namespace UE::ChaosMover
 					// Execute the combined proposed move
 					{
 						// WorkingSimTickParams.MovingComps is left empty in the async case, so we don't access resources used by the concurrent gameplay thread
+						WorkingSimTickParams.SimContext = SimContext;
 						WorkingSimTickParams.StartState = WorkingSubstepStartData;
 						WorkingSimTickParams.SimBlackboard = SimBlackboard;
 						WorkingSimTickParams.TimeStep = SubTimeStep;
@@ -859,6 +863,21 @@ namespace UE::ChaosMover
 		}
 	}
 
+	FMoverSimContext FMoverStateMachine::GetSimContext() const
+	{
+		FMoverSimContext SimContext;
+
+		SimContext.Simulation = Simulation.Get();
+		if (SimContext.Simulation)
+		{
+			SimContext.Blackboard = Simulation->GetRollbackBlackboardSimWrapper();
+		}
+
+		SimContext.SimulationOwner = nullptr;	// This crutch isn't needed for ChaosMover cases
+
+		return SimContext;
+	}
+
 	const TWeakObjectPtr<UBaseMovementMode> FMoverStateMachine::GetCurrentMode() const
 	{
 		if (CurrentModeName != NAME_None && Modes.Contains(CurrentModeName))
@@ -957,11 +976,13 @@ namespace UE::ChaosMover
 						MoverEventContext.ServerFrame = TimeStep.ServerFrame;
 						MoverEventContext.bIsCausedByRollback = QueuedModeTransition->IsCausedByRollback();
 
+						FMoverSimContext SimContext = GetSimContext();
+
 						if (PreviousModeName != NAME_None && Modes.Contains(PreviousModeName))
 						{
 							if (TStrongObjectPtr<UBaseMovementMode> PreviousMode = Modes[PreviousModeName].Pin())
 							{
-								PreviousMode->Deactivate(MoverEventContext);
+								PreviousMode->Deactivate(MoverEventContext, NextModeName, SimContext);
 							}
 						}
 
@@ -977,7 +998,16 @@ namespace UE::ChaosMover
 
 						if (TStrongObjectPtr<UBaseMovementMode> CurrentMode = Modes[CurrentModeName].Pin())
 						{	
-							CurrentMode->Activate(MoverEventContext);
+							FMoverSyncState* InProgressSyncState = nullptr;
+							FMoverAuxStateContext* InProgressAuxState = nullptr;
+
+							if (WorkingSimTickEndData)
+							{
+								InProgressSyncState = &WorkingSimTickEndData->SyncState;
+								InProgressAuxState = &WorkingSimTickEndData->AuxState;
+							}
+
+							CurrentMode->Activate(MoverEventContext, PreviousModeName, SimContext, WorkingSubstepStartData, InProgressSyncState, InProgressAuxState);
 						}
 
 						// Notify the simulation of a mode change so it can react accordingly
diff --git a/Engine/Plugins/Experimental/ChaosMover/Source/ChaosMover/Private/Character/Modes/ChaosCharacterMovementMode.cpp b/Engine/Plugins/Experimental/ChaosMover/Source/ChaosMover/Private/Character/Modes/ChaosCharacterMovementMode.cpp
--- a/Engine/Plugins/Experimental/ChaosMover/Source/ChaosMover/Private/Character/Modes/ChaosCharacterMovementMode.cpp
+++ b/Engine/Plugins/Experimental/ChaosMover/Source/ChaosMover/Private/Character/Modes/ChaosCharacterMovementMode.cpp
@@ -26,9 +26,9 @@ UChaosCharacterMovementMode::UChaosCharacterMovementMode(const FObjectInitialize
 	SharedSettingsClasses.Add(USharedChaosCharacterMovementSettings::StaticClass());
 }
 
-void UChaosCharacterMovementMode::OnRegistered(const FName ModeName)
+void UChaosCharacterMovementMode::OnRegistered(const FName ModeName, const FMoverSimContext& SimContext)
 {
-	Super::OnRegistered(ModeName);
+	Super::OnRegistered(ModeName, SimContext);
 
 	if (TargetHeightOverride.IsSet())
 	{
diff --git a/Engine/Plugins/Experimental/ChaosMover/Source/ChaosMover/Private/Character/Modes/ChaosFallingMode.cpp b/Engine/Plugins/Experimental/ChaosMover/Source/ChaosMover/Private/Character/Modes/ChaosFallingMode.cpp
--- a/Engine/Plugins/Experimental/ChaosMover/Source/ChaosMover/Private/Character/Modes/ChaosFallingMode.cpp
+++ b/Engine/Plugins/Experimental/ChaosMover/Source/ChaosMover/Private/Character/Modes/ChaosFallingMode.cpp
@@ -43,7 +43,7 @@ UChaosFallingMode::UChaosFallingMode(const FObjectInitializer& ObjectInitializer
 	Transitions.Add(CreateDefaultSubobject<UChaosCharacterLandingCheck>(TEXT("DefaultLandingCheck")));
 }
 
-void UChaosFallingMode::GenerateMove_Implementation(const FMoverTickStartData& StartState, const FMoverTimeStep& TimeStep, FProposedMove& OutProposedMove) const
+void UChaosFallingMode::GenerateMove_Implementation(const FMoverSimContext& SimContext, const FMoverTickStartData& StartState, const FMoverTimeStep& TimeStep, FProposedMove& OutProposedMove) const
 {
 	QUICK_SCOPE_CYCLE_COUNTER(ChaosFallingMode_GenerateMove);
 
diff --git a/Engine/Plugins/Experimental/ChaosMover/Source/ChaosMover/Private/Character/Modes/ChaosFlyingMode.cpp b/Engine/Plugins/Experimental/ChaosMover/Source/ChaosMover/Private/Character/Modes/ChaosFlyingMode.cpp
--- a/Engine/Plugins/Experimental/ChaosMover/Source/ChaosMover/Private/Character/Modes/ChaosFlyingMode.cpp
+++ b/Engine/Plugins/Experimental/ChaosMover/Source/ChaosMover/Private/Character/Modes/ChaosFlyingMode.cpp
@@ -24,7 +24,7 @@ UChaosFlyingMode::UChaosFlyingMode(const FObjectInitializer& ObjectInitializer)
 	TwistTorqueLimit = 0.0f;
 }
 
-void UChaosFlyingMode::GenerateMove_Implementation(const FMoverTickStartData& StartState, const FMoverTimeStep& TimeStep, FProposedMove& OutProposedMove) const
+void UChaosFlyingMode::GenerateMove_Implementation(const FMoverSimContext& SimContext, const FMoverTickStartData& StartState, const FMoverTimeStep& TimeStep, FProposedMove& OutProposedMove) const
 {
 	QUICK_SCOPE_CYCLE_COUNTER(ChaosFlyingMode_GenerateMove);
 
diff --git a/Engine/Plugins/Experimental/ChaosMover/Source/ChaosMover/Private/Character/Modes/ChaosSwimmingMode.cpp b/Engine/Plugins/Experimental/ChaosMover/Source/ChaosMover/Private/Character/Modes/ChaosSwimmingMode.cpp
--- a/Engine/Plugins/Experimental/ChaosMover/Source/ChaosMover/Private/Character/Modes/ChaosSwimmingMode.cpp
+++ b/Engine/Plugins/Experimental/ChaosMover/Source/ChaosMover/Private/Character/Modes/ChaosSwimmingMode.cpp
@@ -28,7 +28,7 @@ UChaosSwimmingMode::UChaosSwimmingMode(const FObjectInitializer& ObjectInitializ
 	Transitions.Add(CreateDefaultSubobject<UChaosCharacterWaterCheck>(TEXT("DefaultWaterCheck")));
 }
 
-void UChaosSwimmingMode::GenerateMove_Implementation(const FMoverTickStartData& StartState, const FMoverTimeStep& TimeStep, FProposedMove& OutProposedMove) const
+void UChaosSwimmingMode::GenerateMove_Implementation(const FMoverSimContext& SimContext, const FMoverTickStartData& StartState, const FMoverTimeStep& TimeStep, FProposedMove& OutProposedMove) const
 {
 	QUICK_SCOPE_CYCLE_COUNTER(ChaosSwimmingMode_GenerateMove);
 
diff --git a/Engine/Plugins/Experimental/ChaosMover/Source/ChaosMover/Private/Character/Modes/ChaosWalkingMode.cpp b/Engine/Plugins/Experimental/ChaosMover/Source/ChaosMover/Private/Character/Modes/ChaosWalkingMode.cpp
--- a/Engine/Plugins/Experimental/ChaosMover/Source/ChaosMover/Private/Character/Modes/ChaosWalkingMode.cpp
+++ b/Engine/Plugins/Experimental/ChaosMover/Source/ChaosMover/Private/Character/Modes/ChaosWalkingMode.cpp
@@ -51,7 +51,7 @@ void UChaosWalkingMode::UpdateConstraintSettings(Chaos::FCharacterGroundConstrai
 	ConstraintSettings.RadialForceMotionTargetScaling = FractionalRadialForceLimitScaling;
 }
 
-void UChaosWalkingMode::GenerateMove_Implementation(const FMoverTickStartData& StartState, const FMoverTimeStep& TimeStep, FProposedMove& OutProposedMove) const
+void UChaosWalkingMode::GenerateMove_Implementation(const FMoverSimContext& SimContext, const FMoverTickStartData& StartState, const FMoverTimeStep& TimeStep, FProposedMove& OutProposedMove) const
 {
 	QUICK_SCOPE_CYCLE_COUNTER(ChaosWalkingMode_GenerateMove);
 
diff --git a/Engine/Plugins/Experimental/ChaosMover/Source/ChaosMover/Private/PathedMovement/ChaosPathedMovementMode.cpp b/Engine/Plugins/Experimental/ChaosMover/Source/ChaosMover/Private/PathedMovement/ChaosPathedMovementMode.cpp
--- a/Engine/Plugins/Experimental/ChaosMover/Source/ChaosMover/Private/PathedMovement/ChaosPathedMovementMode.cpp
+++ b/Engine/Plugins/Experimental/ChaosMover/Source/ChaosMover/Private/PathedMovement/ChaosPathedMovementMode.cpp
@@ -30,11 +30,11 @@ UChaosPathedMovementMode::UChaosPathedMovementMode(const FObjectInitializer& Obj
 	SetDefaultJointConstraintProperties();
 }
 
-void UChaosPathedMovementMode::OnRegistered(const FName ModeName)
+void UChaosPathedMovementMode::OnRegistered(const FName ModeName, const FMoverSimContext& SimContext)
 {
 	Chaos::EnsureIsInGameThreadContext();
 
-	Super::OnRegistered(ModeName);
+	Super::OnRegistered(ModeName, SimContext);
 
 	InitializePath();
 
@@ -51,7 +51,7 @@ const Chaos::FPBDJointSettings& UChaosPathedMovementMode::GetConstraintSettings(
 	return JointSettings;
 }
 
-void UChaosPathedMovementMode::GenerateMove_Implementation(const FMoverTickStartData& StartState, const FMoverTimeStep& TimeStep, FProposedMove& OutProposedMove) const
+void UChaosPathedMovementMode::GenerateMove_Implementation(const FMoverSimContext& SimContext, const FMoverTickStartData& StartState, const FMoverTimeStep& TimeStep, FProposedMove& OutProposedMove) const
 {
 }
 
diff --git a/Engine/Plugins/Experimental/ChaosMover/Source/ChaosMover/Public/ChaosMover/ChaosMoverSimulation.h b/Engine/Plugins/Experimental/ChaosMover/Source/ChaosMover/Public/ChaosMover/ChaosMoverSimulation.h
--- a/Engine/Plugins/Experimental/ChaosMover/Source/ChaosMover/Public/ChaosMover/ChaosMoverSimulation.h
+++ b/Engine/Plugins/Experimental/ChaosMover/Source/ChaosMover/Public/ChaosMover/ChaosMoverSimulation.h
@@ -117,7 +117,7 @@ public:
 	CHAOSMOVER_API void Deinit();
 
 	CHAOSMOVER_API void ProcessInputs(int32 PhysicsStep, const FMoverTimeStep& TimeStep, const UE::ChaosMover::FSimulationInputData& InputData);
-	CHAOSMOVER_API void SimulationTick(const FMoverTimeStep& TimeStep, const UE::ChaosMover::FSimulationInputData& InputData, UE::ChaosMover::FSimulationOutputData& OutputData);
+	CHAOSMOVER_API void SimulationTick(const FMoverSimContext& SimContext, const FMoverTimeStep& TimeStep, const UE::ChaosMover::FSimulationInputData& InputData, UE::ChaosMover::FSimulationOutputData& OutputData);
 	CHAOSMOVER_API void ModifyContacts(const FMoverTimeStep& TimeStep, const UE::ChaosMover::FSimulationInputData& InputData, const UE::ChaosMover::FSimulationOutputData& OutputData, Chaos::FCollisionContactModifier& Modifier);
 
 	CHAOSMOVER_API void AddEvent(TSharedPtr<FMoverSimulationEventData> Event);
@@ -144,7 +144,7 @@ protected:
 	CHAOSMOVER_API virtual void OnInit();
 	CHAOSMOVER_API virtual void OnDeinit();
 	CHAOSMOVER_API virtual void OnPreSimulationTick(const FMoverTimeStep& TimeStep, const UE::ChaosMover::FSimulationInputData& InputData);
-	CHAOSMOVER_API virtual void OnSimulationTick(const FMoverTimeStep& TimeStep, const UE::ChaosMover::FSimulationInputData& InputData, UE::ChaosMover::FSimulationOutputData& OutputData);
+	CHAOSMOVER_API virtual void OnSimulationTick(const FMoverSimContext& SimContext, const FMoverTimeStep& TimeStep, const UE::ChaosMover::FSimulationInputData& InputData, UE::ChaosMover::FSimulationOutputData& OutputData);
 	CHAOSMOVER_API virtual void OnPostSimulationTick(const FMoverTimeStep& TimeStep, UE::ChaosMover::FSimulationOutputData& OutputData);
 	CHAOSMOVER_API virtual void OnModifyContacts(const FMoverTimeStep& TimeStep, const UE::ChaosMover::FSimulationInputData& InputData, const UE::ChaosMover::FSimulationOutputData& OutputData, Chaos::FCollisionContactModifier& Modifier);
 	CHAOSMOVER_API virtual void OnSimulationRollback(const FMoverTimeStep& NewTimeStep, const FMoverSyncState& PrevSyncState);
diff --git a/Engine/Plugins/Experimental/ChaosMover/Source/ChaosMover/Public/ChaosMover/ChaosMoverStateMachine.h b/Engine/Plugins/Experimental/ChaosMover/Source/ChaosMover/Public/ChaosMover/ChaosMoverStateMachine.h
--- a/Engine/Plugins/Experimental/ChaosMover/Source/ChaosMover/Public/ChaosMover/ChaosMoverStateMachine.h
+++ b/Engine/Plugins/Experimental/ChaosMover/Source/ChaosMover/Public/ChaosMover/ChaosMoverStateMachine.h
@@ -43,7 +43,7 @@ namespace UE::ChaosMover
 		CHAOSMOVER_API void SetModeImmediately(const FMoverTimeStep& TimeStep, FName DesiredModeName, bool bShouldReenter = false);
 		CHAOSMOVER_API void ClearQueuedMode();
 
-		CHAOSMOVER_API void OnSimulationTick(const FMoverTimeStep& TimeStep, const FMoverTickStartData& StartState, UMoverBlackboard* SimBlackboard, UMovementMixer* MovementMixer, FMoverTickEndData& OutputState);
+		CHAOSMOVER_API void OnSimulationTick(const FMoverSimContext& SimContext, const FMoverTimeStep& TimeStep, const FMoverTickStartData& StartState, UMoverBlackboard* SimBlackboard, UMovementMixer* MovementMixer, FMoverTickEndData& OutputState);
 		CHAOSMOVER_API void OnSimulationRollback(const FMoverTimeStep& NewTimeStep, const FMoverSyncState& InvalidSyncState, const FMoverSyncState& NewSyncState);
 
 		FName GetCurrentModeName() const
@@ -112,6 +112,7 @@ namespace UE::ChaosMover
 		// Internal-use-only tick data structs, for efficiency since they typically have the same contents from frame to frame
 		FMoverTickStartData WorkingSubstepStartData;
 		FSimulationTickParams WorkingSimTickParams;
+		FMoverTickEndData* WorkingSimTickEndData = nullptr;	// Only valid during OnSimulationTick and null otherwise
 
 		TWeakObjectPtr<UChaosMoverSimulation> Simulation;
 
@@ -127,6 +128,7 @@ namespace UE::ChaosMover
 		void ReceiveLayeredMoves(const FMoverTimeStep& TimeStep, const FChaosNetLayeredMovesQueue* LayeredMovesQueue);
 		void RollBackQueuedInstantMovementEffects(int32 NewServerFrame);
 		void RollBackQueuedLayeredMoves(int32 NewServerFrame);
+		FMoverSimContext GetSimContext() const;
 	
 		FMoverTimeStep InternalSimTime;
 
diff --git a/Engine/Plugins/Experimental/ChaosMover/Source/ChaosMover/Public/ChaosMover/Character/Modes/ChaosCharacterMovementMode.h b/Engine/Plugins/Experimental/ChaosMover/Source/ChaosMover/Public/ChaosMover/Character/Modes/ChaosCharacterMovementMode.h
--- a/Engine/Plugins/Experimental/ChaosMover/Source/ChaosMover/Public/ChaosMover/Character/Modes/ChaosCharacterMovementMode.h
+++ b/Engine/Plugins/Experimental/ChaosMover/Source/ChaosMover/Public/ChaosMover/Character/Modes/ChaosCharacterMovementMode.h
@@ -33,7 +33,7 @@ class UChaosCharacterMovementMode :	public UChaosMovementMode,
 public:
 	CHAOSMOVER_API UChaosCharacterMovementMode(const FObjectInitializer& ObjectInitializer);
 
-	CHAOSMOVER_API virtual void OnRegistered(const FName ModeName) override;
+	CHAOSMOVER_API virtual void OnRegistered(const FName ModeName, const FMoverSimContext& SimContext) override;
 
 	// IChaosCharacterConstraintMovementModeInterface
 	virtual bool ShouldEnableConstraint() const override
diff --git a/Engine/Plugins/Experimental/ChaosMover/Source/ChaosMover/Public/ChaosMover/Character/Modes/ChaosFallingMode.h b/Engine/Plugins/Experimental/ChaosMover/Source/ChaosMover/Public/ChaosMover/Character/Modes/ChaosFallingMode.h
--- a/Engine/Plugins/Experimental/ChaosMover/Source/ChaosMover/Public/ChaosMover/Character/Modes/ChaosFallingMode.h
+++ b/Engine/Plugins/Experimental/ChaosMover/Source/ChaosMover/Public/ChaosMover/Character/Modes/ChaosFallingMode.h
@@ -19,7 +19,7 @@ class UChaosFallingMode : public UChaosCharacterMovementMode
 public:
 	UE_API UChaosFallingMode(const FObjectInitializer& ObjectInitializer);
 
-	UE_API virtual void GenerateMove_Implementation(const FMoverTickStartData& StartState, const FMoverTimeStep& TimeStep, FProposedMove& OutProposedMove) const override;
+	UE_API virtual void GenerateMove_Implementation(const FMoverSimContext& SimContext, const FMoverTickStartData& StartState, const FMoverTimeStep& TimeStep, FProposedMove& OutProposedMove) const override;
 	UE_API virtual void SimulationTick_Implementation(const FSimulationTickParams& Params, FMoverTickEndData& OutputState) override;
 
 	/**
diff --git a/Engine/Plugins/Experimental/ChaosMover/Source/ChaosMover/Public/ChaosMover/Character/Modes/ChaosFlyingMode.h b/Engine/Plugins/Experimental/ChaosMover/Source/ChaosMover/Public/ChaosMover/Character/Modes/ChaosFlyingMode.h
--- a/Engine/Plugins/Experimental/ChaosMover/Source/ChaosMover/Public/ChaosMover/Character/Modes/ChaosFlyingMode.h
+++ b/Engine/Plugins/Experimental/ChaosMover/Source/ChaosMover/Public/ChaosMover/Character/Modes/ChaosFlyingMode.h
@@ -17,7 +17,7 @@ class UChaosFlyingMode : public UChaosCharacterMovementMode
 public:
 	CHAOSMOVER_API UChaosFlyingMode(const FObjectInitializer& ObjectInitializer);
 
-	CHAOSMOVER_API virtual void GenerateMove_Implementation(const FMoverTickStartData& StartState, const FMoverTimeStep& TimeStep, FProposedMove& OutProposedMove) const override;
+	CHAOSMOVER_API virtual void GenerateMove_Implementation(const FMoverSimContext& SimContext, const FMoverTickStartData& StartState, const FMoverTimeStep& TimeStep, FProposedMove& OutProposedMove) const override;
 	CHAOSMOVER_API virtual void SimulationTick_Implementation(const FSimulationTickParams& Params, FMoverTickEndData& OutputState) override;
 
 	CHAOSMOVER_API virtual void UpdateCurrentFloor(const FMoverTimeStep& TimeStep) const override;
diff --git a/Engine/Plugins/Experimental/ChaosMover/Source/ChaosMover/Public/ChaosMover/Character/Modes/ChaosSwimmingMode.h b/Engine/Plugins/Experimental/ChaosMover/Source/ChaosMover/Public/ChaosMover/Character/Modes/ChaosSwimmingMode.h
--- a/Engine/Plugins/Experimental/ChaosMover/Source/ChaosMover/Public/ChaosMover/Character/Modes/ChaosSwimmingMode.h
+++ b/Engine/Plugins/Experimental/ChaosMover/Source/ChaosMover/Public/ChaosMover/Character/Modes/ChaosSwimmingMode.h
@@ -101,7 +101,7 @@ public:
 	UPROPERTY(EditAnywhere, BlueprintReadOnly, Category = "Water Control")
 	FSwimmingSettings SurfaceSwimmingWaterControlSettings;
 	
-	UE_API virtual void GenerateMove_Implementation(const FMoverTickStartData& StartState, const FMoverTimeStep& TimeStep, FProposedMove& OutProposedMove) const override;
+	UE_API virtual void GenerateMove_Implementation(const FMoverSimContext& SimContext, const FMoverTickStartData& StartState, const FMoverTimeStep& TimeStep, FProposedMove& OutProposedMove) const override;
 	UE_API virtual void SimulationTick_Implementation(const FSimulationTickParams& Params, FMoverTickEndData& OutputState) override;
 
 protected:
diff --git a/Engine/Plugins/Experimental/ChaosMover/Source/ChaosMover/Public/ChaosMover/Character/Modes/ChaosWalkingMode.h b/Engine/Plugins/Experimental/ChaosMover/Source/ChaosMover/Public/ChaosMover/Character/Modes/ChaosWalkingMode.h
--- a/Engine/Plugins/Experimental/ChaosMover/Source/ChaosMover/Public/ChaosMover/Character/Modes/ChaosWalkingMode.h
+++ b/Engine/Plugins/Experimental/ChaosMover/Source/ChaosMover/Public/ChaosMover/Character/Modes/ChaosWalkingMode.h
@@ -65,7 +65,7 @@ public:
 
 	UE_API virtual void UpdateConstraintSettings(Chaos::FCharacterGroundConstraintSettings& ConstraintSettings) const override;
 
-	UE_API virtual void GenerateMove_Implementation(const FMoverTickStartData& StartState, const FMoverTimeStep& TimeStep, FProposedMove& OutProposedMove) const override;
+	UE_API virtual void GenerateMove_Implementation(const FMoverSimContext& SimContext, const FMoverTickStartData& StartState, const FMoverTimeStep& TimeStep, FProposedMove& OutProposedMove) const override;
 	UE_API virtual void SimulationTick_Implementation(const FSimulationTickParams& Params, FMoverTickEndData& OutputState) override;
 
 	UE_API virtual void ModifyContacts(const FMoverTimeStep& TimeStep, const UE::ChaosMover::FSimulationInputData& InputData, const UE::ChaosMover::FSimulationOutputData& OutputData, Chaos::FCollisionContactModifier& Modifier) const override;
diff --git a/Engine/Plugins/Experimental/ChaosMover/Source/ChaosMover/Public/ChaosMover/PathedMovement/ChaosPathedMovementMode.h b/Engine/Plugins/Experimental/ChaosMover/Source/ChaosMover/Public/ChaosMover/PathedMovement/ChaosPathedMovementMode.h
--- a/Engine/Plugins/Experimental/ChaosMover/Source/ChaosMover/Public/ChaosMover/PathedMovement/ChaosPathedMovementMode.h
+++ b/Engine/Plugins/Experimental/ChaosMover/Source/ChaosMover/Public/ChaosMover/PathedMovement/ChaosPathedMovementMode.h
@@ -31,8 +31,8 @@ public:
 	UChaosPathedMovementMode(const FObjectInitializer& ObjectInitializer);
 
 	// UBaseMovementMode
-	virtual void OnRegistered(const FName ModeName) override;
-	virtual void GenerateMove_Implementation(const FMoverTickStartData& StartState, const FMoverTimeStep& TimeStep, FProposedMove& OutProposedMove) const override;
+	virtual void OnRegistered(const FName ModeName, const FMoverSimContext& SimContext) override;
+	virtual void GenerateMove_Implementation(const FMoverSimContext& SimContext, const FMoverTickStartData& StartState, const FMoverTimeStep& TimeStep, FProposedMove& OutProposedMove) const override;
 	virtual void SimulationTick_Implementation(const FSimulationTickParams& Params, FMoverTickEndData& OutputState) override;
 	virtual bool UsesMovementBasisTransform() const override;
 
diff --git a/Engine/Plugins/Experimental/Mover/Source/Mover/Private/Backends/MoverNetworkPhysicsLiaisonBase.cpp b/Engine/Plugins/Experimental/Mover/Source/Mover/Private/Backends/MoverNetworkPhysicsLiaisonBase.cpp
--- a/Engine/Plugins/Experimental/Mover/Source/Mover/Private/Backends/MoverNetworkPhysicsLiaisonBase.cpp
+++ b/Engine/Plugins/Experimental/Mover/Source/Mover/Private/Backends/MoverNetworkPhysicsLiaisonBase.cpp
@@ -829,6 +829,11 @@ void UMoverNetworkPhysicsLiaisonComponentBase::PerformPreSimulate_Internal(const
 	UMoverComponent& MoverComp = GetMoverComponent();
 	const FMoverTimeStep TimeStep = bUsingAsyncPhysics ? GetCurrentAsyncMoverTimeStep_Internal() : GetCurrentMoverTimeStep(TickParams.DeltaTimeSeconds);
 
+	FMoverSimContext SimContext;
+	SimContext.Simulation = Simulation;
+	SimContext.Blackboard = Simulation->GetRollbackBlackboardSimWrapper();
+	SimContext.SimulationOwner = &MoverComp;
+
 	// Update movement state machine
 	if (MoverComp.bHasJustRolledBack)
 	{
@@ -845,7 +850,7 @@ void UMoverNetworkPhysicsLiaisonComponentBase::PerformPreSimulate_Internal(const
 	RollbackBlackboard->BeginSimulationFrame(TimeStep);
 
 
-	MoverComp.ModeFSM->OnSimulationTick(MoverComp.UpdatedComponent, MoverComp.UpdatedCompAsPrimitive, MoverComp.GetSimBlackboard_Mutable(), TickStartData, TimeStep, TickEndData);
+	MoverComp.ModeFSM->OnSimulationTick(MoverComp.UpdatedComponent, MoverComp.UpdatedCompAsPrimitive, MoverComp.GetSimBlackboard_Mutable(), SimContext, TickStartData, TimeStep, TickEndData);
 
 	// Set the output sync state and fill in the movement mode
 	Output.SyncState = TickEndData.SyncState;
diff --git a/Engine/Plugins/Experimental/Mover/Source/Mover/Private/Backends/MoverNetworkPredictionLiaison.cpp b/Engine/Plugins/Experimental/Mover/Source/Mover/Private/Backends/MoverNetworkPredictionLiaison.cpp
--- a/Engine/Plugins/Experimental/Mover/Source/Mover/Private/Backends/MoverNetworkPredictionLiaison.cpp
+++ b/Engine/Plugins/Experimental/Mover/Source/Mover/Private/Backends/MoverNetworkPredictionLiaison.cpp
@@ -145,6 +145,11 @@ void UMoverNetworkPredictionLiaisonComponent::SimulationTick(const FNetSimTimeSt
 		return;
 	}
 
+	FMoverSimContext SimContext;
+	SimContext.Simulation = Simulation;
+	SimContext.Blackboard = Simulation->GetRollbackBlackboardSimWrapper();
+	SimContext.SimulationOwner = MoverComp;
+
 	FMoverTickStartData StartData;
 	FMoverTickEndData EndData;
 
@@ -164,7 +169,7 @@ void UMoverNetworkPredictionLiaisonComponent::SimulationTick(const FNetSimTimeSt
 	MoverTimeStep.BaseSimTimeMs = TimeStep.TotalSimulationTime;
 	MoverTimeStep.StepMs		= TimeStep.StepMS;
 
-	MoverComp->SimulationTick(MoverTimeStep, StartData, OUT EndData);
+	MoverComp->SimulationTick(SimContext, MoverTimeStep, StartData, OUT EndData);
 
 	*SimOutput.Sync = EndData.SyncState;
     *SimOutput.Aux.Get() = EndData.AuxState;
diff --git a/Engine/Plugins/Experimental/Mover/Source/Mover/Private/Backends/MoverStandaloneLiaison.cpp b/Engine/Plugins/Experimental/Mover/Source/Mover/Private/Backends/MoverStandaloneLiaison.cpp
--- a/Engine/Plugins/Experimental/Mover/Source/Mover/Private/Backends/MoverStandaloneLiaison.cpp
+++ b/Engine/Plugins/Experimental/Mover/Source/Mover/Private/Backends/MoverStandaloneLiaison.cpp
@@ -382,6 +382,11 @@ void UMoverStandaloneLiaisonComponent::TickMovementSimulation(float DeltaSeconds
 	TimeStep.BaseSimTimeMs = CurrentSimTimeMs;
 	TimeStep.StepMs = DeltaSeconds * 1000.f;
 
+	FMoverSimContext SimContext;
+	SimContext.Simulation = Simulation;
+	SimContext.Blackboard = Simulation->GetRollbackBlackboardSimWrapper();
+	SimContext.SimulationOwner = MoverComp;
+
 	WorkingStartData.InputCmd = LastProducedInputCmd;
 
 	WorkingEndData.InitForNewFrame();
@@ -394,7 +399,7 @@ void UMoverStandaloneLiaisonComponent::TickMovementSimulation(float DeltaSeconds
 
 	{
 		SCOPED_NAMED_EVENT(StandaloneMoverBackendLiaison_Tick_SimulationOnly, FColor::Blue);
-		MoverComp->SimulationTick(TimeStep, WorkingStartData, OUT WorkingEndData);
+		MoverComp->SimulationTick(SimContext, TimeStep, WorkingStartData, OUT WorkingEndData);
 	}
 
 	{
diff --git a/Engine/Plugins/Experimental/Mover/Source/Mover/Private/DefaultMovementSet/Modes/AsyncFallingMode.cpp b/Engine/Plugins/Experimental/Mover/Source/Mover/Private/DefaultMovementSet/Modes/AsyncFallingMode.cpp
--- a/Engine/Plugins/Experimental/Mover/Source/Mover/Private/DefaultMovementSet/Modes/AsyncFallingMode.cpp
+++ b/Engine/Plugins/Experimental/Mover/Source/Mover/Private/DefaultMovementSet/Modes/AsyncFallingMode.cpp
@@ -35,7 +35,7 @@ UAsyncFallingMode::UAsyncFallingMode(const FObjectInitializer& ObjectInitializer
 }
 
 
-void UAsyncFallingMode::GenerateMove_Implementation(const FMoverTickStartData& StartState, const FMoverTimeStep& TimeStep, FProposedMove& OutProposedMove) const
+void UAsyncFallingMode::GenerateMove_Implementation(const FMoverSimContext& SimContext, const FMoverTickStartData& StartState, const FMoverTimeStep& TimeStep, FProposedMove& OutProposedMove) const
 {
 	const UMoverComponent* MoverComp = GetMoverComponent();
 	const FCharacterDefaultInputs* CharacterInputs = StartState.InputCmd.InputCollection.FindDataByType<FCharacterDefaultInputs>();
@@ -268,20 +268,20 @@ void UAsyncFallingMode::SimulationTick_Implementation(const FSimulationTickParam
 }
 
 
-void UAsyncFallingMode::OnRegistered(const FName ModeName)
+void UAsyncFallingMode::OnRegistered(const FName ModeName, const FMoverSimContext& SimContext)
 {
-	Super::OnRegistered(ModeName);
+	Super::OnRegistered(ModeName, SimContext);
 
 	CommonLegacySettings = GetMoverComponent()->FindSharedSettings<UCommonLegacyMovementSettings>();
 	ensureMsgf(CommonLegacySettings.IsValid(), TEXT("Failed to find instance of CommonLegacyMovementSettings on %s. Movement may not function properly."), *GetPathNameSafe(this));
 }
 
 
-void UAsyncFallingMode::OnUnregistered()
+void UAsyncFallingMode::OnUnregistered(const FMoverSimContext& SimContext)
 {
 	CommonLegacySettings = nullptr;
 
-	Super::OnUnregistered();
+	Super::OnUnregistered(SimContext);
 }
 
 void UAsyncFallingMode::ProcessLanded(const FFloorCheckResult& FloorResult, FVector& Velocity, FRelativeBaseInfo& BaseInfo, FMoverTickEndData& TickEndData) const
diff --git a/Engine/Plugins/Experimental/Mover/Source/Mover/Private/DefaultMovementSet/Modes/AsyncFlyingMode.cpp b/Engine/Plugins/Experimental/Mover/Source/Mover/Private/DefaultMovementSet/Modes/AsyncFlyingMode.cpp
--- a/Engine/Plugins/Experimental/Mover/Source/Mover/Private/DefaultMovementSet/Modes/AsyncFlyingMode.cpp
+++ b/Engine/Plugins/Experimental/Mover/Source/Mover/Private/DefaultMovementSet/Modes/AsyncFlyingMode.cpp
@@ -22,7 +22,7 @@ UAsyncFlyingMode::UAsyncFlyingMode(const FObjectInitializer& ObjectInitializer)
 	GameplayTags.AddTag(Mover_IsFlying);
 }
 
-void UAsyncFlyingMode::GenerateMove_Implementation(const FMoverTickStartData& StartState, const FMoverTimeStep& TimeStep, FProposedMove& OutProposedMove) const
+void UAsyncFlyingMode::GenerateMove_Implementation(const FMoverSimContext& SimContext, const FMoverTickStartData& StartState, const FMoverTimeStep& TimeStep, FProposedMove& OutProposedMove) const
 {
 	const UMoverComponent* MoverComp = GetMoverComponent();
 	const FCharacterDefaultInputs* CharacterInputs = StartState.InputCmd.InputCollection.FindDataByType<FCharacterDefaultInputs>();
@@ -153,18 +153,18 @@ void UAsyncFlyingMode::SimulationTick_Implementation(const FSimulationTickParams
 }
 
 
-void UAsyncFlyingMode::OnRegistered(const FName ModeName)
+void UAsyncFlyingMode::OnRegistered(const FName ModeName, const FMoverSimContext& SimContext)
 {
-	Super::OnRegistered(ModeName);
+	Super::OnRegistered(ModeName, SimContext);
 
 	CommonLegacySettings = GetMoverComponent()->FindSharedSettings<UCommonLegacyMovementSettings>();
 	ensureMsgf(CommonLegacySettings, TEXT("Failed to find instance of CommonLegacyMovementSettings on %s. Movement may not function properly."), *GetPathNameSafe(this));
 }
 
 
-void UAsyncFlyingMode::OnUnregistered()
+void UAsyncFlyingMode::OnUnregistered(const FMoverSimContext& SimContext)
 {
 	CommonLegacySettings = nullptr;
 
-	Super::OnUnregistered();
+	Super::OnUnregistered(SimContext);
 }
\ No newline at end of file
diff --git a/Engine/Plugins/Experimental/Mover/Source/Mover/Private/DefaultMovementSet/Modes/AsyncNavWalkingMode.cpp b/Engine/Plugins/Experimental/Mover/Source/Mover/Private/DefaultMovementSet/Modes/AsyncNavWalkingMode.cpp
--- a/Engine/Plugins/Experimental/Mover/Source/Mover/Private/DefaultMovementSet/Modes/AsyncNavWalkingMode.cpp
+++ b/Engine/Plugins/Experimental/Mover/Source/Mover/Private/DefaultMovementSet/Modes/AsyncNavWalkingMode.cpp
@@ -56,7 +56,7 @@ UAsyncNavWalkingMode::UAsyncNavWalkingMode()
 	GameplayTags.AddTag(Mover_IsNavWalking);
 }
 
-void UAsyncNavWalkingMode::GenerateMove_Implementation(const FMoverTickStartData& StartState, const FMoverTimeStep& TimeStep, FProposedMove& OutProposedMove) const
+void UAsyncNavWalkingMode::GenerateMove_Implementation(const FMoverSimContext& SimContext, const FMoverTickStartData& StartState, const FMoverTimeStep& TimeStep, FProposedMove& OutProposedMove) const
 {
 	if (!CommonLegacySettings.IsValid())
 	{
@@ -411,9 +411,9 @@ void UAsyncNavWalkingMode::SetTurnGeneratorClass(TSubclassOf<UObject> TurnGenera
 	}
 }
 
-void UAsyncNavWalkingMode::Activate(const FMoverEventContext& Context)
+void UAsyncNavWalkingMode::Activate(const FMoverEventContext& Context, FName PrevModeName, const FMoverSimContext& SimContext, const FMoverTickStartData& StartState, FMoverSyncState* OutSyncState, FMoverAuxStateContext* OutAuxState)
 {
-	Super::Activate(Context);
+	Super::Activate(Context, PrevModeName, SimContext, StartState, OutSyncState, OutAuxState);
 	
 	if (UMoverComponent* MoverComp = GetMoverComponent())
 	{
@@ -623,9 +623,9 @@ FVector UAsyncNavWalkingMode::ProjectLocationFromNavMesh(float DeltaSeconds, con
 	return NewLocation;
 }
 
-void UAsyncNavWalkingMode::OnRegistered(const FName ModeName)
+void UAsyncNavWalkingMode::OnRegistered(const FName ModeName, const FMoverSimContext& SimContext)
 {
-	Super::OnRegistered(ModeName);
+	Super::OnRegistered(ModeName, SimContext);
 
 	UMoverComponent* MoverComponent = GetMoverComponent();
 	CommonLegacySettings = MoverComponent->FindSharedSettings<UCommonLegacyMovementSettings>();
@@ -642,13 +642,13 @@ void UAsyncNavWalkingMode::OnRegistered(const FName ModeName)
 	}
 }
 
-void UAsyncNavWalkingMode::OnUnregistered()
+void UAsyncNavWalkingMode::OnUnregistered(const FMoverSimContext& SimContext)
 {
 	CommonLegacySettings = nullptr;
 	NavMoverComponent = nullptr;
 	NavDataInterface = nullptr;
 
-	Super::OnUnregistered();
+	Super::OnUnregistered(SimContext);
 }
 
 void UAsyncNavWalkingMode::CaptureOutputState(const FMoverDefaultSyncState& StartSyncState, const FVector& FinalLocation, const FRotator& FinalRotation, const FMovementRecord& Record, const FVector& AngularVelocityDegrees, FMoverDefaultSyncState& OutputSyncState, FMoverTickEndData& TickEndData) const
diff --git a/Engine/Plugins/Experimental/Mover/Source/Mover/Private/DefaultMovementSet/Modes/AsyncWalkingMode.cpp b/Engine/Plugins/Experimental/Mover/Source/Mover/Private/DefaultMovementSet/Modes/AsyncWalkingMode.cpp
--- a/Engine/Plugins/Experimental/Mover/Source/Mover/Private/DefaultMovementSet/Modes/AsyncWalkingMode.cpp
+++ b/Engine/Plugins/Experimental/Mover/Source/Mover/Private/DefaultMovementSet/Modes/AsyncWalkingMode.cpp
@@ -26,7 +26,7 @@ UAsyncWalkingMode::UAsyncWalkingMode(const FObjectInitializer& ObjectInitializer
 	GameplayTags.AddTag(Mover_IsOnGround);
 }
 
-void UAsyncWalkingMode::GenerateMove_Implementation(const FMoverTickStartData& StartState, const FMoverTimeStep& TimeStep, FProposedMove& OutProposedMove) const
+void UAsyncWalkingMode::GenerateMove_Implementation(const FMoverSimContext& SimContext, const FMoverTickStartData& StartState, const FMoverTimeStep& TimeStep, FProposedMove& OutProposedMove) const
 {
 	const UMoverComponent* MoverComp = GetMoverComponent();
 	const FCharacterDefaultInputs* CharacterInputs = StartState.InputCmd.InputCollection.FindDataByType<FCharacterDefaultInputs>();
@@ -352,19 +352,19 @@ void UAsyncWalkingMode::SetTurnGeneratorClass(TSubclassOf<UObject> TurnGenerator
 }
 
 
-void UAsyncWalkingMode::OnRegistered(const FName ModeName)
+void UAsyncWalkingMode::OnRegistered(const FName ModeName, const FMoverSimContext& SimContext)
 {
-	Super::OnRegistered(ModeName);
+	Super::OnRegistered(ModeName, SimContext);
 
 	CommonLegacySettings = GetMoverComponent()->FindSharedSettings<UCommonLegacyMovementSettings>();
 	ensureMsgf(CommonLegacySettings.IsValid(), TEXT("Failed to find instance of CommonLegacyMovementSettings on %s. Movement may not function properly."), *GetPathNameSafe(this));
 }
 
-void UAsyncWalkingMode::OnUnregistered()
+void UAsyncWalkingMode::OnUnregistered(const FMoverSimContext& SimContext)
 {
 	CommonLegacySettings = nullptr;
 
-	Super::OnUnregistered();
+	Super::OnUnregistered(SimContext);
 }
 
 void UAsyncWalkingMode::CaptureFinalState(const FVector FinalLocation, const FRotator FinalRotation, bool bDidAttemptMovement, const FFloorCheckResult& FloorResult, const FMovementRecord& Record, const FVector& AngularVelocityDegrees, FMoverDefaultSyncState& OutputSyncState) const
diff --git a/Engine/Plugins/Experimental/Mover/Source/Mover/Private/DefaultMovementSet/Modes/FallingMode.cpp b/Engine/Plugins/Experimental/Mover/Source/Mover/Private/DefaultMovementSet/Modes/FallingMode.cpp
--- a/Engine/Plugins/Experimental/Mover/Source/Mover/Private/DefaultMovementSet/Modes/FallingMode.cpp
+++ b/Engine/Plugins/Experimental/Mover/Source/Mover/Private/DefaultMovementSet/Modes/FallingMode.cpp
@@ -33,7 +33,7 @@ UFallingMode::UFallingMode(const FObjectInitializer& ObjectInitializer)
 }
 
 
-void UFallingMode::GenerateMove_Implementation(const FMoverTickStartData& StartState, const FMoverTimeStep& TimeStep, FProposedMove& OutProposedMove) const
+void UFallingMode::GenerateMove_Implementation(const FMoverSimContext& SimContext, const FMoverTickStartData& StartState, const FMoverTimeStep& TimeStep, FProposedMove& OutProposedMove) const
 {
 	const UMoverComponent* MoverComp = GetMoverComponent();
 	const FCharacterDefaultInputs* CharacterInputs = StartState.InputCmd.InputCollection.FindDataByType<FCharacterDefaultInputs>();
@@ -259,20 +259,20 @@ void UFallingMode::SimulationTick_Implementation(const FSimulationTickParams& Pa
 }
 
 
-void UFallingMode::OnRegistered(const FName ModeName)
+void UFallingMode::OnRegistered(const FName ModeName, const FMoverSimContext& SimContext)
 {
-	Super::OnRegistered(ModeName);
+	Super::OnRegistered(ModeName, SimContext);
 
 	CommonLegacySettings = GetMoverComponent()->FindSharedSettings<UCommonLegacyMovementSettings>();
 	ensureMsgf(CommonLegacySettings, TEXT("Failed to find instance of CommonLegacyMovementSettings on %s. Movement may not function properly."), *GetPathNameSafe(this));
 }
 
 
-void UFallingMode::OnUnregistered()
+void UFallingMode::OnUnregistered(const FMoverSimContext& SimContext)
 {
 	CommonLegacySettings = nullptr;
 
-	Super::OnUnregistered();
+	Super::OnUnregistered(SimContext);
 }
 
 void UFallingMode::ProcessLanded(const FFloorCheckResult& FloorResult, FVector& Velocity, FRelativeBaseInfo& BaseInfo, FMoverTickEndData& TickEndData) const
diff --git a/Engine/Plugins/Experimental/Mover/Source/Mover/Private/DefaultMovementSet/Modes/FlyingMode.cpp b/Engine/Plugins/Experimental/Mover/Source/Mover/Private/DefaultMovementSet/Modes/FlyingMode.cpp
--- a/Engine/Plugins/Experimental/Mover/Source/Mover/Private/DefaultMovementSet/Modes/FlyingMode.cpp
+++ b/Engine/Plugins/Experimental/Mover/Source/Mover/Private/DefaultMovementSet/Modes/FlyingMode.cpp
@@ -20,7 +20,7 @@ UFlyingMode::UFlyingMode(const FObjectInitializer& ObjectInitializer)
 	GameplayTags.AddTag(Mover_IsFlying);
 }
 
-void UFlyingMode::GenerateMove_Implementation(const FMoverTickStartData& StartState, const FMoverTimeStep& TimeStep, FProposedMove& OutProposedMove) const
+void UFlyingMode::GenerateMove_Implementation(const FMoverSimContext& SimContext, const FMoverTickStartData& StartState, const FMoverTimeStep& TimeStep, FProposedMove& OutProposedMove) const
 {
 	const UMoverComponent* MoverComp = GetMoverComponent();
 	const FCharacterDefaultInputs* CharacterInputs = StartState.InputCmd.InputCollection.FindDataByType<FCharacterDefaultInputs>();
@@ -162,18 +162,18 @@ void UFlyingMode::CaptureFinalState(USceneComponent* UpdatedComponent, FMovement
 	UpdatedComponent->ComponentVelocity = FinalVelocity;
 }
 
-void UFlyingMode::OnRegistered(const FName ModeName)
+void UFlyingMode::OnRegistered(const FName ModeName, const FMoverSimContext& SimContext)
 {
-	Super::OnRegistered(ModeName);
+	Super::OnRegistered(ModeName, SimContext);
 
 	CommonLegacySettings = GetMoverComponent()->FindSharedSettings<UCommonLegacyMovementSettings>();
 	ensureMsgf(CommonLegacySettings, TEXT("Failed to find instance of CommonLegacyMovementSettings on %s. Movement may not function properly."), *GetPathNameSafe(this));
 }
 
 
-void UFlyingMode::OnUnregistered()
+void UFlyingMode::OnUnregistered(const FMoverSimContext& SimContext)
 {
 	CommonLegacySettings = nullptr;
 
-	Super::OnUnregistered();
+	Super::OnUnregistered(SimContext);
 }
diff --git a/Engine/Plugins/Experimental/Mover/Source/Mover/Private/DefaultMovementSet/Modes/NavWalkingMode.cpp b/Engine/Plugins/Experimental/Mover/Source/Mover/Private/DefaultMovementSet/Modes/NavWalkingMode.cpp
--- a/Engine/Plugins/Experimental/Mover/Source/Mover/Private/DefaultMovementSet/Modes/NavWalkingMode.cpp
+++ b/Engine/Plugins/Experimental/Mover/Source/Mover/Private/DefaultMovementSet/Modes/NavWalkingMode.cpp
@@ -36,7 +36,7 @@ UNavWalkingMode::UNavWalkingMode()
 	GameplayTags.AddTag(Mover_IsNavWalking);
 }
 
-void UNavWalkingMode::GenerateMove_Implementation(const FMoverTickStartData& StartState, const FMoverTimeStep& TimeStep, FProposedMove& OutProposedMove) const
+void UNavWalkingMode::GenerateMove_Implementation(const FMoverSimContext& SimContext, const FMoverTickStartData& StartState, const FMoverTimeStep& TimeStep, FProposedMove& OutProposedMove) const
 {
 	const UMoverComponent* MoverComp = GetMoverComponent();
 	const FCharacterDefaultInputs* CharacterInputs = StartState.InputCmd.InputCollection.FindDataByType<FCharacterDefaultInputs>();
@@ -416,18 +416,18 @@ void UNavWalkingMode::SetCollisionForNavWalking(bool bEnable)
 	}
 }
 
-void UNavWalkingMode::Activate(const FMoverEventContext& Context)
+void UNavWalkingMode::Activate(const FMoverEventContext& Context, FName PrevModeName, const FMoverSimContext& SimContext, const FMoverTickStartData& StartState, FMoverSyncState* OutSyncState, FMoverAuxStateContext* OutAuxState)
 {
-	Super::Activate(Context);
+	Super::Activate(Context, PrevModeName, SimContext, StartState, OutSyncState, OutAuxState);
 	SetCollisionForNavWalking(true);
 
 	NavDataInterface = GetNavData();
 }
 
-void UNavWalkingMode::Deactivate(const FMoverEventContext& Context)
+void UNavWalkingMode::Deactivate(const FMoverEventContext& Context, FName NextModeName, const FMoverSimContext& SimContext)
 {
 	SetCollisionForNavWalking(false);
-	Super::Deactivate(Context);
+	Super::Deactivate(Context, NextModeName, SimContext);
 }
 
 const INavigationDataInterface* UNavWalkingMode::GetNavData() const
@@ -614,9 +614,9 @@ FVector UNavWalkingMode::ProjectLocationFromNavMesh(float DeltaSeconds, const FV
 	return NewLocation;
 }
 
-void UNavWalkingMode::OnRegistered(const FName ModeName)
+void UNavWalkingMode::OnRegistered(const FName ModeName, const FMoverSimContext& SimContext)
 {
-	Super::OnRegistered(ModeName);
+	Super::OnRegistered(ModeName, SimContext);
 
 	UMoverComponent* MoverComponent = GetMoverComponent();
 	CommonLegacySettings = MoverComponent->FindSharedSettings<UCommonLegacyMovementSettings>();
@@ -633,12 +633,12 @@ void UNavWalkingMode::OnRegistered(const FName ModeName)
 	}
 }
 
-void UNavWalkingMode::OnUnregistered()
+void UNavWalkingMode::OnUnregistered(const FMoverSimContext& SimContext)
 {
 	CommonLegacySettings = nullptr;
 	NavDataInterface = nullptr;
 
-	Super::OnUnregistered();
+	Super::OnUnregistered(SimContext);
 }
 
 void UNavWalkingMode::CaptureFinalState(USceneComponent* UpdatedComponent, const FMovementRecord& Record, const FVector& AngularVelocityDegrees, FMoverDefaultSyncState& OutputSyncState) const
diff --git a/Engine/Plugins/Experimental/Mover/Source/Mover/Private/DefaultMovementSet/Modes/SimpleFlyingMode.cpp b/Engine/Plugins/Experimental/Mover/Source/Mover/Private/DefaultMovementSet/Modes/SimpleFlyingMode.cpp
--- a/Engine/Plugins/Experimental/Mover/Source/Mover/Private/DefaultMovementSet/Modes/SimpleFlyingMode.cpp
+++ b/Engine/Plugins/Experimental/Mover/Source/Mover/Private/DefaultMovementSet/Modes/SimpleFlyingMode.cpp
@@ -9,7 +9,7 @@
 
 #include UE_INLINE_GENERATED_CPP_BY_NAME(SimpleFlyingMode)
 
-void USimpleFlyingMode::GenerateMove_Implementation(const FMoverTickStartData& StartState, const FMoverTimeStep& TimeStep,
+void USimpleFlyingMode::GenerateMove_Implementation(const FMoverSimContext& SimContext, const FMoverTickStartData& StartState, const FMoverTimeStep& TimeStep,
                                                    FProposedMove& OutProposedMove) const
 {
 	const FMoverDefaultSyncState* StartingSyncState = StartState.SyncState.SyncStateCollection.FindDataByType<FMoverDefaultSyncState>();
diff --git a/Engine/Plugins/Experimental/Mover/Source/Mover/Private/DefaultMovementSet/Modes/SimpleWalkingMode.cpp b/Engine/Plugins/Experimental/Mover/Source/Mover/Private/DefaultMovementSet/Modes/SimpleWalkingMode.cpp
--- a/Engine/Plugins/Experimental/Mover/Source/Mover/Private/DefaultMovementSet/Modes/SimpleWalkingMode.cpp
+++ b/Engine/Plugins/Experimental/Mover/Source/Mover/Private/DefaultMovementSet/Modes/SimpleWalkingMode.cpp
@@ -7,7 +7,7 @@
 
 #include UE_INLINE_GENERATED_CPP_BY_NAME(SimpleWalkingMode)
 
-void USimpleWalkingMode::GenerateMove_Implementation(const FMoverTickStartData& StartState, const FMoverTimeStep& TimeStep,
+void USimpleWalkingMode::GenerateMove_Implementation(const FMoverSimContext& SimContext, const FMoverTickStartData& StartState, const FMoverTimeStep& TimeStep,
                                                    FProposedMove& OutProposedMove) const
 {
 	const FMoverDefaultSyncState* StartingSyncState = StartState.SyncState.SyncStateCollection.FindDataByType<FMoverDefaultSyncState>();
diff --git a/Engine/Plugins/Experimental/Mover/Source/Mover/Private/DefaultMovementSet/Modes/SwimmingMode.cpp b/Engine/Plugins/Experimental/Mover/Source/Mover/Private/DefaultMovementSet/Modes/SwimmingMode.cpp
--- a/Engine/Plugins/Experimental/Mover/Source/Mover/Private/DefaultMovementSet/Modes/SwimmingMode.cpp
+++ b/Engine/Plugins/Experimental/Mover/Source/Mover/Private/DefaultMovementSet/Modes/SwimmingMode.cpp
@@ -21,7 +21,7 @@ USwimmingMode::USwimmingMode(const FObjectInitializer& ObjectInitializer)
 	GameplayTags.AddTag(Mover_IsSwimming);
 }
 
-void USwimmingMode::GenerateMove_Implementation(const FMoverTickStartData& StartState, const FMoverTimeStep& TimeStep, FProposedMove& OutProposedMove) const
+void USwimmingMode::GenerateMove_Implementation(const FMoverSimContext& SimContext, const FMoverTickStartData& StartState, const FMoverTimeStep& TimeStep, FProposedMove& OutProposedMove) const
 {
 	const UMoverComponent* MoverComp = GetMoverComponent();
 	const FCharacterDefaultInputs* CharacterInputs = StartState.InputCmd.InputCollection.FindDataByType<FCharacterDefaultInputs>();
@@ -166,9 +166,9 @@ void USwimmingMode::SimulationTick_Implementation(const FSimulationTickParams& P
 {
 }
 
-void USwimmingMode::OnRegistered(const FName ModeName)
+void USwimmingMode::OnRegistered(const FName ModeName, const FMoverSimContext& SimContext)
 {
-	Super::OnRegistered(ModeName);
+	Super::OnRegistered(ModeName, SimContext);
 
 	if (UMoverComponent* MoverComp = GetMoverComponent())
 	{
@@ -178,11 +178,11 @@ void USwimmingMode::OnRegistered(const FName ModeName)
 	ensureMsgf(CommonLegacySettings, TEXT("Failed to find instance of CommonLegacyMovementSettings on %s. Movement may not function properly."), *GetPathNameSafe(this));
 }
 
-void USwimmingMode::OnUnregistered()
+void USwimmingMode::OnUnregistered(const FMoverSimContext& SimContext)
 {
 	CommonLegacySettings = nullptr;
 
-	Super::OnUnregistered();
+	Super::OnUnregistered(SimContext);
 }
 
 bool USwimmingMode::AttemptJump(const FSimulationTickParams& Params, float UpwardsSpeed, FMoverTickEndData& OutputState)
diff --git a/Engine/Plugins/Experimental/Mover/Source/Mover/Private/DefaultMovementSet/Modes/WalkingMode.cpp b/Engine/Plugins/Experimental/Mover/Source/Mover/Private/DefaultMovementSet/Modes/WalkingMode.cpp
--- a/Engine/Plugins/Experimental/Mover/Source/Mover/Private/DefaultMovementSet/Modes/WalkingMode.cpp
+++ b/Engine/Plugins/Experimental/Mover/Source/Mover/Private/DefaultMovementSet/Modes/WalkingMode.cpp
@@ -24,7 +24,7 @@ UWalkingMode::UWalkingMode(const FObjectInitializer& ObjectInitializer)
 	GameplayTags.AddTag(Mover_IsOnGround);
 }
 
-void UWalkingMode::GenerateMove_Implementation(const FMoverTickStartData& StartState, const FMoverTimeStep& TimeStep, FProposedMove& OutProposedMove) const
+void UWalkingMode::GenerateMove_Implementation(const FMoverSimContext& SimContext, const FMoverTickStartData& StartState, const FMoverTimeStep& TimeStep, FProposedMove& OutProposedMove) const
 {
 	const UMoverComponent* MoverComp = GetMoverComponent();
 	const FCharacterDefaultInputs* CharacterInputs = StartState.InputCmd.InputCollection.FindDataByType<FCharacterDefaultInputs>();
@@ -309,19 +309,19 @@ void UWalkingMode::SetTurnGeneratorClass(TSubclassOf<UObject> TurnGeneratorClass
 }
 
 
-void UWalkingMode::OnRegistered(const FName ModeName)
+void UWalkingMode::OnRegistered(const FName ModeName, const FMoverSimContext& SimContext)
 {
-	Super::OnRegistered(ModeName);
+	Super::OnRegistered(ModeName, SimContext);
 
 	CommonLegacySettings = GetMoverComponent()->FindSharedSettings<UCommonLegacyMovementSettings>();
 	ensureMsgf(CommonLegacySettings, TEXT("Failed to find instance of CommonLegacyMovementSettings on %s. Movement may not function properly."), *GetPathNameSafe(this));
 }
 
-void UWalkingMode::OnUnregistered()
+void UWalkingMode::OnUnregistered(const FMoverSimContext& SimContext)
 {
 	CommonLegacySettings = nullptr;
 
-	Super::OnUnregistered();
+	Super::OnUnregistered(SimContext);
 }
 
 void UWalkingMode::CaptureFinalState(USceneComponent* UpdatedComponent, bool bDidAttemptMovement, const FFloorCheckResult& FloorResult, const FMovementRecord& Record, const FVector& AngularVelocityDegrees, FMoverDefaultSyncState& OutputSyncState) const
diff --git a/Engine/Plugins/Experimental/Mover/Source/Mover/Private/MovementMode.cpp b/Engine/Plugins/Experimental/Mover/Source/Mover/Private/MovementMode.cpp
--- a/Engine/Plugins/Experimental/Mover/Source/Mover/Private/MovementMode.cpp
+++ b/Engine/Plugins/Experimental/Mover/Source/Mover/Private/MovementMode.cpp
@@ -27,7 +27,7 @@ UWorld* UBaseMovementMode::GetWorld() const
 }
 
 
-void UBaseMovementMode::OnRegistered(const FName ModeName)
+void UBaseMovementMode::OnRegistered(const FName ModeName, const FMoverSimContext& SimContext)
 {
 	for (TObjectPtr<UBaseMovementModeTransition>& Transition : Transitions)
 	{
@@ -44,7 +44,7 @@ void UBaseMovementMode::OnRegistered(const FName ModeName)
 	K2_OnRegistered(ModeName);
 }
 
-void UBaseMovementMode::OnUnregistered()
+void UBaseMovementMode::OnUnregistered(const FMoverSimContext& SimContext)
 {
 	for (TObjectPtr<UBaseMovementModeTransition>& Transition : Transitions)
 	{
@@ -61,7 +61,7 @@ void UBaseMovementMode::OnUnregistered()
 	K2_OnUnregistered();
 }
 
-void UBaseMovementMode::Activate(const FMoverEventContext& Context)
+void UBaseMovementMode::Activate(const FMoverEventContext& Context, FName PrevModeName, const FMoverSimContext& SimContext, const FMoverTickStartData& StartState, FMoverSyncState* OutSyncState, FMoverAuxStateContext* OutAuxState)
 {
 	if (!bSupportsAsync)
 	{
@@ -69,7 +69,7 @@ void UBaseMovementMode::Activate(const FMoverEventContext& Context)
 	}
 }
 
-void UBaseMovementMode::Deactivate(const FMoverEventContext& Context)
+void UBaseMovementMode::Deactivate(const FMoverEventContext& Context, FName NextModeName, const FMoverSimContext& SimContext)
 {
 	if (!bSupportsAsync)
 	{
@@ -93,7 +93,7 @@ void UBaseMovementMode::Deactivate_External(const FMoverEventContext& Context)
 	}
 }
 
-void UBaseMovementMode::GenerateMove_Implementation(const FMoverTickStartData& StartState, const FMoverTimeStep& TimeStep, FProposedMove& OutProposedMove) const
+void UBaseMovementMode::GenerateMove_Implementation(const FMoverSimContext& SimContext, const FMoverTickStartData& StartState, const FMoverTimeStep& TimeStep, FProposedMove& OutProposedMove) const
 {
 }
 
diff --git a/Engine/Plugins/Experimental/Mover/Source/Mover/Private/MovementModeStateMachine.cpp b/Engine/Plugins/Experimental/Mover/Source/Mover/Private/MovementModeStateMachine.cpp
--- a/Engine/Plugins/Experimental/Mover/Source/Mover/Private/MovementModeStateMachine.cpp
+++ b/Engine/Plugins/Experimental/Mover/Source/Mover/Private/MovementModeStateMachine.cpp
@@ -46,7 +46,7 @@ void UMovementModeStateMachine::RegisterMovementMode(FName ModeName, TObjectPtr<
 		DefaultModeName = ModeName;
 	}
 
-	Mode->OnRegistered(ModeName);
+	Mode->OnRegistered(ModeName, GetSimContext());
 }
 
 
@@ -65,7 +65,7 @@ void UMovementModeStateMachine::UnregisterMovementMode(FName ModeName)
 
 	if (ModeToUnregister)
 	{
-		ModeToUnregister->OnUnregistered();
+		ModeToUnregister->OnUnregistered(GetSimContext());
 	}
 }
 
@@ -79,12 +79,12 @@ void UMovementModeStateMachine::ClearAllMovementModes()
 		MoverEventContext.bIsDuringResimulation = false;
 		MoverEventContext.bIsCausedByRollback = false;
 
-		Modes[CurrentModeName]->Deactivate(MoverEventContext);
+		Modes[CurrentModeName]->Deactivate(MoverEventContext, NAME_None, GetSimContext());
 	}
 
 	for (TPair<FName, TObjectPtr<UBaseMovementMode>>& Element : Modes)
 	{
-		Element.Value->OnUnregistered();
+		Element.Value->OnUnregistered(GetSimContext());
 	}
 
 	Modes.Empty();
@@ -181,12 +181,14 @@ void UMovementModeStateMachine::ClearQueuedMode()
 }
 
 
-void UMovementModeStateMachine::OnSimulationTick(USceneComponent* UpdatedComponent, UPrimitiveComponent* UpdatedPrimitive, UMoverBlackboard* SimBlackboard, const FMoverTickStartData& StartState, const FMoverTimeStep& TimeStep, FMoverTickEndData& OutputState)
+void UMovementModeStateMachine::OnSimulationTick(USceneComponent* UpdatedComponent, UPrimitiveComponent* UpdatedPrimitive, UMoverBlackboard* SimBlackboard, const FMoverSimContext& SimContext, const FMoverTickStartData& StartState, const FMoverTimeStep& TimeStep, FMoverTickEndData& OutputState)
 {
 	//GEngine->AddOnScreenDebugMessage(-1, -0.1f, FColor::White, FString::Printf(TEXT("Mode Tick: %s  Queued: %s"), *CurrentModeName.ToString(), *NextModeName.ToString()));
 	FMoverTimeStep SubTimeStep = TimeStep;
 	CurrentBaseTimeStep = TimeStep;
 
+	WorkingSimTickEndData = &OutputState;
+	ON_SCOPE_EXIT { WorkingSimTickEndData = nullptr; };
 
 	WorkingSubstepStartData = StartState;
 	bool bIsWorkingStartStateReady = true;	// this flag is used to avoid unneeded data copying after substeps
@@ -338,7 +340,7 @@ void UMovementModeStateMachine::OnSimulationTick(USceneComponent* UpdatedCompone
 				!(bHasLayeredMoveContributions && CombinedLayeredMove.MixMode == EMoveMixMode::OverrideAll))
 			{
 				QUICK_SCOPE_CYCLE_COUNTER(STAT_GenerateMoveFromMode);
-				CurrentMode->GenerateMove(WorkingSubstepStartData, SubTimeStep, OUT CombinedMove);
+				CurrentMode->GenerateMove(SimContext, WorkingSubstepStartData, SubTimeStep, OUT CombinedMove);
 
 				bHasModeMoveContribution = true;
 			}
@@ -824,25 +826,39 @@ void UMovementModeStateMachine::AdvanceToNextMode()
 		MoverEventContext.bIsDuringResimulation = CurrentBaseTimeStep.bIsResimulating;
 		MoverEventContext.bIsCausedByRollback = bIsCausedByRollback;
 
+		FMoverSimContext SimContext = GetSimContext();
+
 		if (PreviousModeName != NAME_None && Modes.Contains(PreviousModeName))
 		{
-			Modes[PreviousModeName]->Deactivate(MoverEventContext);
+			Modes[PreviousModeName]->Deactivate(MoverEventContext, NextModeName, SimContext);
 		}
 
 		// signal movement mode change event
 		const UMoverComponent* MoverComp = CastChecked<UMoverComponent>(GetOuter());
 
-		FRollbackBlackboardSimWrapper RollbackBlackboard = Simulation->GetRollbackBlackboardSimWrapper();
+		if (Simulation)
+		{
+			FRollbackBlackboardSimWrapper RollbackBlackboard = Simulation->GetRollbackBlackboardSimWrapper();
+
+			FMovementModeChangeRecord ModeChangeRecord;
+			ModeChangeRecord.ModeName = CurrentModeName;
+			ModeChangeRecord.PrevModeName = PreviousModeName;
+			ModeChangeRecord.Frame = CurrentBaseTimeStep.ServerFrame;
+			ModeChangeRecord.SimTimeMs = CurrentBaseTimeStep.BaseSimTimeMs;
+
+			RollbackBlackboard.TrySet(CommonBlackboard::LastModeChangeRecord, ModeChangeRecord);
+		}
 
-		FMovementModeChangeRecord ModeChangeRecord;
-		ModeChangeRecord.ModeName = CurrentModeName;
-		ModeChangeRecord.PrevModeName = PreviousModeName;
-		ModeChangeRecord.Frame = CurrentBaseTimeStep.ServerFrame;
-		ModeChangeRecord.SimTimeMs = CurrentBaseTimeStep.BaseSimTimeMs;
+		FMoverSyncState* InProgressSyncState = nullptr;
+		FMoverAuxStateContext* InProgressAuxState = nullptr;
 
-		RollbackBlackboard.TrySet(CommonBlackboard::LastModeChangeRecord, ModeChangeRecord);
+		if (WorkingSimTickEndData)
+		{
+			InProgressSyncState = &WorkingSimTickEndData->SyncState;
+			InProgressAuxState = &WorkingSimTickEndData->AuxState;
+		}
 
-		Modes[CurrentModeName]->Activate(MoverEventContext);
+		Modes[CurrentModeName]->Activate(MoverEventContext, PreviousModeName, SimContext, WorkingSimTickParams.StartState, InProgressSyncState, InProgressAuxState);
 
 		MoverComp->OnMovementModeChanged.Broadcast(PreviousModeName, NextModeName);
 	}
@@ -1064,6 +1080,21 @@ AActor* UMovementModeStateMachine::GetOwnerActor() const
 	return nullptr;
 }
 
+FMoverSimContext UMovementModeStateMachine::GetSimContext() const
+{
+	FMoverSimContext SimContext;
+	
+	if (Simulation)
+	{
+		SimContext.Simulation = Simulation;
+		SimContext.Blackboard = Simulation->GetRollbackBlackboardSimWrapper();
+	}
+
+	SimContext.SimulationOwner = Cast<UMoverComponent>(GetOuter());
+
+	return SimContext;
+}
+
 void UMovementModeStateMachine::PostInitProperties()
 {
 	Super::PostInitProperties();
diff --git a/Engine/Plugins/Experimental/Mover/Source/Mover/Private/MoverComponent.cpp b/Engine/Plugins/Experimental/Mover/Source/Mover/Private/MoverComponent.cpp
--- a/Engine/Plugins/Experimental/Mover/Source/Mover/Private/MoverComponent.cpp
+++ b/Engine/Plugins/Experimental/Mover/Source/Mover/Private/MoverComponent.cpp
@@ -101,11 +101,45 @@ void UMoverComponent::InitializeComponent()
 		FindDefaultUpdatedComponent();
 
 
+		// Instantiate our sister backend component that will actually talk to the system driving the simulation
+		if (BackendClass)
+		{
+			UActorComponent* NewLiaisonComp = Cast<UActorComponent>(GetOwner()->FindComponentByInterface<IMoverBackendLiaisonInterface>());
+			if (!NewLiaisonComp)
+			{
+				NewLiaisonComp = NewObject<UActorComponent>(this, BackendClass, TEXT("BackendLiaisonComponent"));
+			}
+
+			BackendLiaisonComp.SetObject(NewLiaisonComp);
+			BackendLiaisonComp.SetInterface(CastChecked<IMoverBackendLiaisonInterface>(NewLiaisonComp));
+			if (BackendLiaisonComp)
+			{
+				NewLiaisonComp->RegisterComponent();
+				if (!NewLiaisonComp->HasBeenInitialized())
+				{
+					NewLiaisonComp->InitializeComponent();
+				}
+				NewLiaisonComp->SetNetAddressable();
+			}
+		}
+		else
+		{
+			UE_LOGF(LogMover, Error, "No backend class set on %ls. Mover actor will not function.", *GetNameSafe(GetOwner()));
+		}
+
+
 		// Set up FSM and initial movement states
 		ModeFSM = NewObject<UMovementModeStateMachine>(this, TEXT("MoverStateMachine"), RF_Transient);
+
 		ModeFSM->ClearAllMovementModes();
 		ModeFSM->ClearAllGlobalTransitions();
 
+		if (Simulation)
+		{
+			ModeFSM->SetSimulation(*Simulation);
+		}
+
+
 		bool bHasMatchingStartingState = false;
 
 		for (const TPair<FName, TObjectPtr<UBaseMovementMode>>& Element : MovementModes)
@@ -117,9 +151,11 @@ void UMoverComponent::InitializeComponent()
 				continue;
 			}
 
-			ModeFSM->RegisterMovementMode(Element.Key, Element.Value);
+			const bool bIsStartingMode = (StartingMovementMode == Element.Key);
+
+			ModeFSM->RegisterMovementMode(Element.Key, Element.Value, bIsStartingMode);
 
-			bHasMatchingStartingState |= (StartingMovementMode == Element.Key);
+			bHasMatchingStartingState |= bIsStartingMode;
 		}
 
 		for (TObjectPtr<UBaseMovementModeTransition>& Transition : Transitions)
@@ -136,31 +172,6 @@ void UMoverComponent::InitializeComponent()
 			ModeFSM->QueueNextMode(StartingMovementMode);
 		}
 
-		// Instantiate our sister backend component that will actually talk to the system driving the simulation
-		if (BackendClass)
-		{
-			UActorComponent* NewLiaisonComp = Cast<UActorComponent>(GetOwner()->FindComponentByInterface<IMoverBackendLiaisonInterface>());
-			if (!NewLiaisonComp)
-			{
-				NewLiaisonComp = NewObject<UActorComponent>(this, BackendClass, TEXT("BackendLiaisonComponent"));
-			}
-
-			BackendLiaisonComp.SetObject(NewLiaisonComp);
-			BackendLiaisonComp.SetInterface(CastChecked<IMoverBackendLiaisonInterface>(NewLiaisonComp));
-			if (BackendLiaisonComp)
-			{
-				NewLiaisonComp->RegisterComponent();
-				if (!NewLiaisonComp->HasBeenInitialized())
-				{
-					NewLiaisonComp->InitializeComponent();
-				}
-				NewLiaisonComp->SetNetAddressable();
-			}
-		}
-		else
-		{
-			UE_LOGF(LogMover, Error, "No backend class set on %ls. Mover actor will not function.", *GetNameSafe(GetOwner()));
-		}
 
 	}
 
@@ -527,7 +538,7 @@ void UMoverComponent::InitializeSimulationState(FMoverSyncState* OutSync, FMover
 
 }
 
-void UMoverComponent::SimulationTick(const FMoverTimeStep& InTimeStep, const FMoverTickStartData& SimInput, OUT FMoverTickEndData& SimOutput)
+void UMoverComponent::SimulationTick(const FMoverSimContext& SimContext, const FMoverTimeStep& InTimeStep, const FMoverTickStartData& SimInput, OUT FMoverTickEndData& SimOutput)
 {
 	// Send mover info to the Chaos Visual Debugger (this will do nothing if CVD is not recording, or the mover info data channel not enabled)
 	UE::MoverUtils::FMoverCVDRuntimeTrace::TraceMoverData(this, &SimInput.InputCmd, &SimInput.SyncState);
@@ -620,11 +631,11 @@ void UMoverComponent::SimulationTick(const FMoverTimeStep& InTimeStep, const FMo
 	{
 		// If we're on the game thread, we can make use of a scoped movement update for better perf of multi-step movements.  If not, then we're definitely not moving the component in immediate mode so the scope would have no effect.
 		FScopedMovementUpdate ScopedMovementUpdate(UpdatedComponent, EScopedUpdate::DeferredUpdates);
-		ModeFSM->OnSimulationTick(UpdatedComponent, UpdatedCompAsPrimitive, SimBlackboard.Get(), SimInput, MoverTimeStep, SimOutput);
+		ModeFSM->OnSimulationTick(UpdatedComponent, UpdatedCompAsPrimitive, SimBlackboard.Get(), SimContext, SimInput, MoverTimeStep, SimOutput);
 	}
 	else
 	{
-		ModeFSM->OnSimulationTick(UpdatedComponent, UpdatedCompAsPrimitive, SimBlackboard.Get(), SimInput, MoverTimeStep, SimOutput);
+		ModeFSM->OnSimulationTick(UpdatedComponent, UpdatedCompAsPrimitive, SimBlackboard.Get(), SimContext, SimInput, MoverTimeStep, SimOutput);
 	}
 
 	if (FMoverDefaultSyncState* OutputSyncState = SimOutput.SyncState.SyncStateCollection.FindMutableDataByType<FMoverDefaultSyncState>())
@@ -1036,7 +1047,10 @@ void UMoverComponent::SetSimulation(UMoverSimulation& InSimulation)
 	RollbackBlackboard = Simulation->RollbackBlackboard;	// TODO: fetching the true blackboard is temporary: this will be removed once simulations are used consistently
 	RollbackBlackboardExternalWrapper = FRollbackBlackboardExternalWrapper(*Simulation->RollbackBlackboard);
 
-	ModeFSM->SetSimulation(InSimulation);
+	if (ModeFSM)
+	{
+		ModeFSM->SetSimulation(InSimulation);
+	}
 }
 
 void UMoverComponent::HandleImpact(FMoverOnImpactParams& ImpactParams)
@@ -2061,7 +2075,7 @@ bool UMoverComponent::AddMovementModeFromObject(FName ModeName, UBaseMovementMod
 		}
 		
 		MovementModes.Add(ModeName, MovementMode);
-		ModeFSM->RegisterMovementMode(ModeName, MovementMode);
+		ModeFSM->RegisterMovementMode(ModeName, MovementMode, /*bIsDefaultMode*/ false);
 	}
 	else
 	{
@@ -2379,6 +2393,11 @@ TArray<FTrajectorySampleInfo> UMoverComponent::GetPredictedTrajectory(FMoverPred
 {
 	if (ModeFSM)
 	{
+		FMoverSimContext SimContext;
+		SimContext.Simulation = Simulation;
+		SimContext.Blackboard = Simulation->GetRollbackBlackboardPredictionWrapper();
+		SimContext.SimulationOwner = this;
+
 		FMoverTickStartData StepState;
 
 		// Use the last-known input if none are specified.
@@ -2457,7 +2476,7 @@ TArray<FTrajectorySampleInfo> UMoverComponent::GetPredictedTrajectory(FMoverPred
 
 					// Generate next move from current step state
 					FProposedMove StepMove;
-					CurrentMovementMode->GenerateMove(StepState, FutureTimeStep, StepMove);
+					CurrentMovementMode->GenerateMove(SimContext, StepState, FutureTimeStep, StepMove);
 
 					// Advance state based on move
 					StepSyncState->SetTransforms_WorldSpace(StepSyncState->GetLocation_WorldSpace() + (StepMove.LinearVelocity * PredictionParams.SecondsPerSample),
diff --git a/Engine/Plugins/Experimental/Mover/Source/Mover/Private/MoverSimulation.cpp b/Engine/Plugins/Experimental/Mover/Source/Mover/Private/MoverSimulation.cpp
--- a/Engine/Plugins/Experimental/Mover/Source/Mover/Private/MoverSimulation.cpp
+++ b/Engine/Plugins/Experimental/Mover/Source/Mover/Private/MoverSimulation.cpp
@@ -31,3 +31,7 @@ FRollbackBlackboardSimWrapper UMoverSimulation::GetRollbackBlackboardSimWrapper(
 	return RollbackBlackboardSimWrapper;
 }
 
+FRollbackBlackboardSimWrapper UMoverSimulation::GetRollbackBlackboardPredictionWrapper() const
+{
+	return FRollbackBlackboardSimWrapper(*RollbackBlackboard, /*bIsPredictive*/ true);
+}
\ No newline at end of file
diff --git a/Engine/Plugins/Experimental/Mover/Source/Mover/Private/PhysicsMover/Modes/PhysicsDrivenFallingMode.cpp b/Engine/Plugins/Experimental/Mover/Source/Mover/Private/PhysicsMover/Modes/PhysicsDrivenFallingMode.cpp
--- a/Engine/Plugins/Experimental/Mover/Source/Mover/Private/PhysicsMover/Modes/PhysicsDrivenFallingMode.cpp
+++ b/Engine/Plugins/Experimental/Mover/Source/Mover/Private/PhysicsMover/Modes/PhysicsDrivenFallingMode.cpp
@@ -66,9 +66,9 @@ void UPhysicsDrivenFallingMode::ClearTargetHeightOverride()
 	}
 }
 
-void UPhysicsDrivenFallingMode::OnRegistered(const FName ModeName)
+void UPhysicsDrivenFallingMode::OnRegistered(const FName ModeName, const FMoverSimContext& SimContext)
 {
-	Super::OnRegistered(ModeName);
+	Super::OnRegistered(ModeName, SimContext);
 
 	if (TargetHeightOverride.IsSet())
 	{
diff --git a/Engine/Plugins/Experimental/Mover/Source/Mover/Private/PhysicsMover/Modes/PhysicsDrivenSwimmingMode.cpp b/Engine/Plugins/Experimental/Mover/Source/Mover/Private/PhysicsMover/Modes/PhysicsDrivenSwimmingMode.cpp
--- a/Engine/Plugins/Experimental/Mover/Source/Mover/Private/PhysicsMover/Modes/PhysicsDrivenSwimmingMode.cpp
+++ b/Engine/Plugins/Experimental/Mover/Source/Mover/Private/PhysicsMover/Modes/PhysicsDrivenSwimmingMode.cpp
@@ -67,9 +67,9 @@ void UPhysicsDrivenSwimmingMode::ClearTargetHeightOverride()
 	}
 }
 
-void UPhysicsDrivenSwimmingMode::OnRegistered(const FName ModeName)
+void UPhysicsDrivenSwimmingMode::OnRegistered(const FName ModeName, const FMoverSimContext& SimContext)
 {
-	Super::OnRegistered(ModeName);
+	Super::OnRegistered(ModeName, SimContext);
 
 	if (TargetHeightOverride.IsSet())
 	{
diff --git a/Engine/Plugins/Experimental/Mover/Source/Mover/Private/PhysicsMover/Modes/PhysicsDrivenWalkingMode.cpp b/Engine/Plugins/Experimental/Mover/Source/Mover/Private/PhysicsMover/Modes/PhysicsDrivenWalkingMode.cpp
--- a/Engine/Plugins/Experimental/Mover/Source/Mover/Private/PhysicsMover/Modes/PhysicsDrivenWalkingMode.cpp
+++ b/Engine/Plugins/Experimental/Mover/Source/Mover/Private/PhysicsMover/Modes/PhysicsDrivenWalkingMode.cpp
@@ -283,9 +283,9 @@ void UPhysicsDrivenWalkingMode::ClearTargetHeightOverride()
 	}
 }
 
-void UPhysicsDrivenWalkingMode::OnRegistered(const FName ModeName)
+void UPhysicsDrivenWalkingMode::OnRegistered(const FName ModeName, const FMoverSimContext& SimContext)
 {
-	Super::OnRegistered(ModeName);
+	Super::OnRegistered(ModeName, SimContext);
 
 	if (TargetHeightOverride.IsSet())
 	{
diff --git a/Engine/Plugins/Experimental/Mover/Source/Mover/Public/DefaultMovementSet/Modes/AsyncFallingMode.h b/Engine/Plugins/Experimental/Mover/Source/Mover/Public/DefaultMovementSet/Modes/AsyncFallingMode.h
--- a/Engine/Plugins/Experimental/Mover/Source/Mover/Public/DefaultMovementSet/Modes/AsyncFallingMode.h
+++ b/Engine/Plugins/Experimental/Mover/Source/Mover/Public/DefaultMovementSet/Modes/AsyncFallingMode.h
@@ -24,9 +24,9 @@ class UAsyncFallingMode : public UBaseMovementMode
 public:
 	UE_API UAsyncFallingMode(const FObjectInitializer& ObjectInitializer);
 
-	UE_API virtual void OnRegistered(const FName ModeName) override;
-    UE_API virtual void OnUnregistered() override;
-	UE_API virtual void GenerateMove_Implementation(const FMoverTickStartData& StartState, const FMoverTimeStep& TimeStep, FProposedMove& OutProposedMove) const override;
+	UE_API virtual void OnRegistered(const FName ModeName, const FMoverSimContext& SimContext) override;
+    UE_API virtual void OnUnregistered(const FMoverSimContext& SimContext) override;
+	UE_API virtual void GenerateMove_Implementation(const FMoverSimContext& SimContext, const FMoverTickStartData& StartState, const FMoverTimeStep& TimeStep, FProposedMove& OutProposedMove) const override;
 	UE_API virtual void SimulationTick_Implementation(const FSimulationTickParams& Params, FMoverTickEndData& OutputState) override;
 
 	// Broadcast when this actor lands on a valid surface.
diff --git a/Engine/Plugins/Experimental/Mover/Source/Mover/Public/DefaultMovementSet/Modes/AsyncFlyingMode.h b/Engine/Plugins/Experimental/Mover/Source/Mover/Public/DefaultMovementSet/Modes/AsyncFlyingMode.h
--- a/Engine/Plugins/Experimental/Mover/Source/Mover/Public/DefaultMovementSet/Modes/AsyncFlyingMode.h
+++ b/Engine/Plugins/Experimental/Mover/Source/Mover/Public/DefaultMovementSet/Modes/AsyncFlyingMode.h
@@ -30,9 +30,9 @@ public:
 	UPROPERTY(Category = Mover, EditAnywhere, BlueprintReadWrite)
 	bool bRespectDistanceOverWalkableSurfaces = false;
 
-	UE_API virtual void OnRegistered(const FName ModeName) override;
-	UE_API virtual void OnUnregistered() override;
-	UE_API virtual void GenerateMove_Implementation(const FMoverTickStartData& StartState, const FMoverTimeStep& TimeStep, FProposedMove& OutProposedMove) const override;
+	UE_API virtual void OnRegistered(const FName ModeName, const FMoverSimContext& SimContext) override;
+	UE_API virtual void OnUnregistered(const FMoverSimContext& SimContext) override;
+	UE_API virtual void GenerateMove_Implementation(const FMoverSimContext& SimContext, const FMoverTickStartData& StartState, const FMoverTimeStep& TimeStep, FProposedMove& OutProposedMove) const override;
 	UE_API virtual void SimulationTick_Implementation(const FSimulationTickParams& Params, FMoverTickEndData& OutputState) override;
 
 protected:
diff --git a/Engine/Plugins/Experimental/Mover/Source/Mover/Public/DefaultMovementSet/Modes/AsyncNavWalkingMode.h b/Engine/Plugins/Experimental/Mover/Source/Mover/Public/DefaultMovementSet/Modes/AsyncNavWalkingMode.h
--- a/Engine/Plugins/Experimental/Mover/Source/Mover/Public/DefaultMovementSet/Modes/AsyncNavWalkingMode.h
+++ b/Engine/Plugins/Experimental/Mover/Source/Mover/Public/DefaultMovementSet/Modes/AsyncNavWalkingMode.h
@@ -29,9 +29,9 @@ class UAsyncNavWalkingMode : public UBaseMovementMode
 public:
 	UE_API UAsyncNavWalkingMode();
 
-	UE_API virtual void Activate(const FMoverEventContext& Context) override;
+	UE_API virtual void Activate(const FMoverEventContext& Context, FName PrevModeName, const FMoverSimContext& SimContext, const FMoverTickStartData& StartState, FMoverSyncState* OutSyncState, FMoverAuxStateContext* OutAuxState) override;
 	UE_API virtual void SimulationTick_Implementation(const FSimulationTickParams& Params, FMoverTickEndData& OutputState) override;
-	UE_API virtual void GenerateMove_Implementation(const FMoverTickStartData& StartState, const FMoverTimeStep& TimeStep, FProposedMove& OutProposedMove) const override;
+	UE_API virtual void GenerateMove_Implementation(const FMoverSimContext& SimContext, const FMoverTickStartData& StartState, const FMoverTimeStep& TimeStep, FProposedMove& OutProposedMove) const override;
 
 	/**
 	 * Whether or not the actor should sweep for collision geometry while walking.
@@ -126,8 +126,8 @@ protected:
 	UE_API virtual FVector ProjectLocationFromNavMesh(float DeltaSeconds, const FVector& CurrentFeetLocation, const FVector& TargetNavLocation, float UpOffset, float DownOffset);
 
 
-	UE_API virtual void OnRegistered(const FName ModeName) override;
-	UE_API virtual void OnUnregistered() override;
+	UE_API virtual void OnRegistered(const FName ModeName, const FMoverSimContext& SimContext) override;
+	UE_API virtual void OnUnregistered(const FMoverSimContext& SimContext) override;
 
 	UE_API void CaptureOutputState(const FMoverDefaultSyncState& StartSyncState, const FVector& FinalLocation, const FRotator& FinalRotation, const FMovementRecord& Record, const FVector& AngularVelocityDegrees, FMoverDefaultSyncState& OutputSyncState, FMoverTickEndData& OutputState) const;
 
diff --git a/Engine/Plugins/Experimental/Mover/Source/Mover/Public/DefaultMovementSet/Modes/AsyncWalkingMode.h b/Engine/Plugins/Experimental/Mover/Source/Mover/Public/DefaultMovementSet/Modes/AsyncWalkingMode.h
--- a/Engine/Plugins/Experimental/Mover/Source/Mover/Public/DefaultMovementSet/Modes/AsyncWalkingMode.h
+++ b/Engine/Plugins/Experimental/Mover/Source/Mover/Public/DefaultMovementSet/Modes/AsyncWalkingMode.h
@@ -27,7 +27,7 @@ class UAsyncWalkingMode : public UBaseMovementMode
 public:
 	UE_API UAsyncWalkingMode(const FObjectInitializer& ObjectInitializer);
 	
-	UE_API virtual void GenerateMove_Implementation(const FMoverTickStartData& StartState, const FMoverTimeStep& TimeStep, FProposedMove& OutProposedMove) const override;
+	UE_API virtual void GenerateMove_Implementation(const FMoverSimContext& SimContext, const FMoverTickStartData& StartState, const FMoverTimeStep& TimeStep, FProposedMove& OutProposedMove) const override;
 
 	UE_API virtual void SimulationTick_Implementation(const FSimulationTickParams& Params, FMoverTickEndData& OutputState) override;
 
@@ -49,8 +49,8 @@ protected:
 	UPROPERTY(EditAnywhere, Instanced, Category=Mover, meta=(ObjectMustImplement="/Script/Mover.TurnGeneratorInterface"))
 	TObjectPtr<UObject> TurnGenerator;
 
-	UE_API virtual void OnRegistered(const FName ModeName) override; 
-	UE_API virtual void OnUnregistered() override;
+	UE_API virtual void OnRegistered(const FName ModeName, const FMoverSimContext& SimContext) override;
+	UE_API virtual void OnUnregistered(const FMoverSimContext& SimContext) override;
 
 	UE_API void CaptureFinalState(const FVector FinalLocation, const FRotator FinalRotation, bool bDidAttemptMovement, const FFloorCheckResult& FloorResult, const FMovementRecord& Record, const FVector& AngularVelocityDegrees, FMoverDefaultSyncState& OutputSyncState) const;
 
diff --git a/Engine/Plugins/Experimental/Mover/Source/Mover/Public/DefaultMovementSet/Modes/FallingMode.h b/Engine/Plugins/Experimental/Mover/Source/Mover/Public/DefaultMovementSet/Modes/FallingMode.h
--- a/Engine/Plugins/Experimental/Mover/Source/Mover/Public/DefaultMovementSet/Modes/FallingMode.h
+++ b/Engine/Plugins/Experimental/Mover/Source/Mover/Public/DefaultMovementSet/Modes/FallingMode.h
@@ -24,7 +24,7 @@ class UFallingMode : public UBaseMovementMode
 
 
 public:
-	UE_API virtual void GenerateMove_Implementation(const FMoverTickStartData& StartState, const FMoverTimeStep& TimeStep, FProposedMove& OutProposedMove) const override;
+	UE_API virtual void GenerateMove_Implementation(const FMoverSimContext& SimContext, const FMoverTickStartData& StartState, const FMoverTimeStep& TimeStep, FProposedMove& OutProposedMove) const override;
 
 	UE_API virtual void SimulationTick_Implementation(const FSimulationTickParams& Params, FMoverTickEndData& OutputState) override;
 
@@ -99,8 +99,8 @@ protected:
 	UFUNCTION(BlueprintCallable, Category=Mover)
 	UE_API virtual void ProcessLanded(const FFloorCheckResult& FloorResult, FVector& Velocity, FRelativeBaseInfo& BaseInfo, FMoverTickEndData& TickEndData) const;
 	
-	UE_API virtual void OnRegistered(const FName ModeName) override;
-	UE_API virtual void OnUnregistered() override;
+	UE_API virtual void OnRegistered(const FName ModeName, const FMoverSimContext& SimContext) override;
+	UE_API virtual void OnUnregistered(const FMoverSimContext& SimContext) override;
 
 	UE_API void CaptureFinalState(USceneComponent* UpdatedComponent, const FMoverDefaultSyncState& StartSyncState, const FFloorCheckResult& FloorResult, float DeltaSeconds, float DeltaSecondsUsed, const FVector& AngularVelocityDegrees, FMoverDefaultSyncState& OutputSyncState, FMoverTickEndData& TickEndData, FMovementRecord& Record) const;
 
diff --git a/Engine/Plugins/Experimental/Mover/Source/Mover/Public/DefaultMovementSet/Modes/FlyingMode.h b/Engine/Plugins/Experimental/Mover/Source/Mover/Public/DefaultMovementSet/Modes/FlyingMode.h
--- a/Engine/Plugins/Experimental/Mover/Source/Mover/Public/DefaultMovementSet/Modes/FlyingMode.h
+++ b/Engine/Plugins/Experimental/Mover/Source/Mover/Public/DefaultMovementSet/Modes/FlyingMode.h
@@ -29,13 +29,13 @@ public:
 	UPROPERTY(Category = Mover, EditAnywhere, BlueprintReadWrite)
 	bool bRespectDistanceOverWalkableSurfaces = false;
 
-	UE_API virtual void GenerateMove_Implementation(const FMoverTickStartData& StartState, const FMoverTimeStep& TimeStep, FProposedMove& OutProposedMove) const override;
+	UE_API virtual void GenerateMove_Implementation(const FMoverSimContext& SimContext, const FMoverTickStartData& StartState, const FMoverTimeStep& TimeStep, FProposedMove& OutProposedMove) const override;
 
 	UE_API virtual void SimulationTick_Implementation(const FSimulationTickParams& Params, FMoverTickEndData& OutputState) override;
 
 protected:
-	UE_API virtual void OnRegistered(const FName ModeName) override;
-	UE_API virtual void OnUnregistered() override;
+	UE_API virtual void OnRegistered(const FName ModeName, const FMoverSimContext& SimContext) override;
+	UE_API virtual void OnUnregistered(const FMoverSimContext& SimContext) override;
 
 	UE_API void CaptureFinalState(USceneComponent* UpdatedComponent, FMovementRecord& Record, const FMoverDefaultSyncState& StartSyncState, const FVector& AngularVelocityDegrees, FMoverDefaultSyncState& OutputSyncState, const float DeltaSeconds) const;
 
diff --git a/Engine/Plugins/Experimental/Mover/Source/Mover/Public/DefaultMovementSet/Modes/NavWalkingMode.h b/Engine/Plugins/Experimental/Mover/Source/Mover/Public/DefaultMovementSet/Modes/NavWalkingMode.h
--- a/Engine/Plugins/Experimental/Mover/Source/Mover/Public/DefaultMovementSet/Modes/NavWalkingMode.h
+++ b/Engine/Plugins/Experimental/Mover/Source/Mover/Public/DefaultMovementSet/Modes/NavWalkingMode.h
@@ -49,10 +49,10 @@ class UNavWalkingMode : public UBaseMovementMode
 public:
 	UE_API UNavWalkingMode();
 	
-	UE_API virtual void Activate(const FMoverEventContext& Context) override;
-	UE_API virtual void Deactivate(const FMoverEventContext& Context) override;
+	UE_API virtual void Activate(const FMoverEventContext& Context, FName PrevModeName, const FMoverSimContext& SimContext, const FMoverTickStartData& StartState, FMoverSyncState* OutSyncState, FMoverAuxStateContext* OutAuxState) override;
+	UE_API virtual void Deactivate(const FMoverEventContext& Context, FName NextModeName, const FMoverSimContext& SimContext) override;
 	UE_API virtual void SimulationTick_Implementation(const FSimulationTickParams& Params, FMoverTickEndData& OutputState) override;
-	UE_API virtual void GenerateMove_Implementation(const FMoverTickStartData& StartState, const FMoverTimeStep& TimeStep, FProposedMove& OutProposedMove) const override;
+	UE_API virtual void GenerateMove_Implementation(const FMoverSimContext& SimContext, const FMoverTickStartData& StartState, const FMoverTimeStep& TimeStep, FProposedMove& OutProposedMove) const override;
 
 	/**
 	 * Whether or not the actor should sweep for collision geometry while walking.
@@ -150,8 +150,8 @@ protected:
 	UE_API virtual FVector ProjectLocationFromNavMesh(float DeltaSeconds, const FVector& CurrentFeetLocation, const FVector& TargetNavLocation, float UpOffset, float DownOffset);
 
 	
-	UE_API virtual void OnRegistered(const FName ModeName) override;
-	UE_API virtual void OnUnregistered() override;
+	UE_API virtual void OnRegistered(const FName ModeName, const FMoverSimContext& SimContext) override;
+	UE_API virtual void OnUnregistered(const FMoverSimContext& SimContext) override;
 
 	UE_API void CaptureFinalState(USceneComponent* UpdatedComponent, const FMovementRecord& Record, const FVector& AngularVelocityDegrees, FMoverDefaultSyncState& OutputSyncState) const;
 
diff --git a/Engine/Plugins/Experimental/Mover/Source/Mover/Public/DefaultMovementSet/Modes/SimpleFlyingMode.h b/Engine/Plugins/Experimental/Mover/Source/Mover/Public/DefaultMovementSet/Modes/SimpleFlyingMode.h
--- a/Engine/Plugins/Experimental/Mover/Source/Mover/Public/DefaultMovementSet/Modes/SimpleFlyingMode.h
+++ b/Engine/Plugins/Experimental/Mover/Source/Mover/Public/DefaultMovementSet/Modes/SimpleFlyingMode.h
@@ -18,7 +18,7 @@ class USimpleFlyingMode : public UFlyingMode
 
 public:
 
-	UE_API virtual void GenerateMove_Implementation(const FMoverTickStartData& StartState, const FMoverTimeStep& TimeStep, FProposedMove& OutProposedMove) const override;
+	UE_API virtual void GenerateMove_Implementation(const FMoverSimContext& SimContext, const FMoverTickStartData& StartState, const FMoverTimeStep& TimeStep, FProposedMove& OutProposedMove) const override;
 
 	// Override this to make a simple flying mode
 	UFUNCTION(BlueprintNativeEvent, meta = (DisplayName = "Generate Simple Flying Move"))
diff --git a/Engine/Plugins/Experimental/Mover/Source/Mover/Public/DefaultMovementSet/Modes/SimpleWalkingMode.h b/Engine/Plugins/Experimental/Mover/Source/Mover/Public/DefaultMovementSet/Modes/SimpleWalkingMode.h
--- a/Engine/Plugins/Experimental/Mover/Source/Mover/Public/DefaultMovementSet/Modes/SimpleWalkingMode.h
+++ b/Engine/Plugins/Experimental/Mover/Source/Mover/Public/DefaultMovementSet/Modes/SimpleWalkingMode.h
@@ -18,7 +18,7 @@ class USimpleWalkingMode : public UWalkingMode
 
 public:
 
-	UE_API virtual void GenerateMove_Implementation(const FMoverTickStartData& StartState, const FMoverTimeStep& TimeStep, FProposedMove& OutProposedMove) const override;
+	UE_API virtual void GenerateMove_Implementation(const FMoverSimContext& SimContext, const FMoverTickStartData& StartState, const FMoverTimeStep& TimeStep, FProposedMove& OutProposedMove) const override;
 
 	// Override this to make a simple walking mode
 	UFUNCTION(BlueprintNativeEvent, meta = (DisplayName = "Generate Simple Walk Move"))
diff --git a/Engine/Plugins/Experimental/Mover/Source/Mover/Public/DefaultMovementSet/Modes/SwimmingMode.h b/Engine/Plugins/Experimental/Mover/Source/Mover/Public/DefaultMovementSet/Modes/SwimmingMode.h
--- a/Engine/Plugins/Experimental/Mover/Source/Mover/Public/DefaultMovementSet/Modes/SwimmingMode.h
+++ b/Engine/Plugins/Experimental/Mover/Source/Mover/Public/DefaultMovementSet/Modes/SwimmingMode.h
@@ -198,7 +198,7 @@ class USwimmingMode : public UBaseMovementMode
 
 public:
 
-	UE_API virtual void GenerateMove_Implementation(const FMoverTickStartData& StartState, const FMoverTimeStep& TimeStep, FProposedMove& OutProposedMove) const override;
+	UE_API virtual void GenerateMove_Implementation(const FMoverSimContext& SimContext, const FMoverTickStartData& StartState, const FMoverTimeStep& TimeStep, FProposedMove& OutProposedMove) const override;
 
 	UE_API virtual void SimulationTick_Implementation(const FSimulationTickParams& Params, FMoverTickEndData& OutputState) override;
 
@@ -206,8 +206,8 @@ public:
 	FSwimmingControlSettings SurfaceSwimmingWaterControlSettings;
 
 protected:
-	UE_API virtual void OnRegistered(const FName ModeName) override; 
-	UE_API virtual void OnUnregistered() override;
+	UE_API virtual void OnRegistered(const FName ModeName, const FMoverSimContext& SimContext) override;
+	UE_API virtual void OnUnregistered(const FMoverSimContext& SimContext) override;
 
 	UE_API virtual bool AttemptJump(const FSimulationTickParams& Params, float UpwardsSpeed, FMoverTickEndData& Output);
 
diff --git a/Engine/Plugins/Experimental/Mover/Source/Mover/Public/DefaultMovementSet/Modes/WalkingMode.h b/Engine/Plugins/Experimental/Mover/Source/Mover/Public/DefaultMovementSet/Modes/WalkingMode.h
--- a/Engine/Plugins/Experimental/Mover/Source/Mover/Public/DefaultMovementSet/Modes/WalkingMode.h
+++ b/Engine/Plugins/Experimental/Mover/Source/Mover/Public/DefaultMovementSet/Modes/WalkingMode.h
@@ -37,7 +37,7 @@ class UWalkingMode : public UBaseMovementMode
 
 public:
 	
-	UE_API virtual void GenerateMove_Implementation(const FMoverTickStartData& StartState, const FMoverTimeStep& TimeStep, FProposedMove& OutProposedMove) const override;
+	UE_API virtual void GenerateMove_Implementation(const FMoverSimContext& SimContext, const FMoverTickStartData& StartState, const FMoverTimeStep& TimeStep, FProposedMove& OutProposedMove) const override;
 
 	UE_API virtual void SimulationTick_Implementation(const FSimulationTickParams& Params, FMoverTickEndData& OutputState) override;
 
@@ -59,8 +59,8 @@ protected:
 	UPROPERTY(EditAnywhere, Instanced, Category=Mover, meta=(ObjectMustImplement="/Script/Mover.TurnGeneratorInterface"))
 	TObjectPtr<UObject> TurnGenerator;
 
-	UE_API virtual void OnRegistered(const FName ModeName) override; 
-	UE_API virtual void OnUnregistered() override;
+	UE_API virtual void OnRegistered(const FName ModeName, const FMoverSimContext& SimContext) override;
+	UE_API virtual void OnUnregistered(const FMoverSimContext& SimContext) override;
 
 	UE_API void CaptureFinalState(USceneComponent* UpdatedComponent, bool bDidAttemptMovement, const FFloorCheckResult& FloorResult, const FMovementRecord& Record, const FVector& AngularVelocityDegrees, FMoverDefaultSyncState& OutputSyncState) const;
 
diff --git a/Engine/Plugins/Experimental/Mover/Source/Mover/Public/MovementMode.h b/Engine/Plugins/Experimental/Mover/Source/Mover/Public/MovementMode.h
--- a/Engine/Plugins/Experimental/Mover/Source/Mover/Public/MovementMode.h
+++ b/Engine/Plugins/Experimental/Mover/Source/Mover/Public/MovementMode.h
@@ -6,6 +6,7 @@
 #include "GameplayTagContainer.h"
 #include "MoverSimulationTypes.h"
 #include "MoverTypes.h"
+#include "MoverSimulationTypes.h"
 #include "MoveLibrary/MoverBlackboard.h"
 #include "MovementModeTransition.h"
 #include "UObject/Interface.h"
@@ -43,12 +44,12 @@ class UBaseMovementMode : public UObject
 public:
 	UE_API virtual UWorld* GetWorld() const override;
 	
-	UE_API virtual void OnRegistered(const FName ModeName);
-	UE_API virtual void OnUnregistered();
+	UE_API virtual void OnRegistered(const FName ModeName, const FMoverSimContext& SimContext);
+	UE_API virtual void OnUnregistered(const FMoverSimContext& SimContext);
 	
-	// These functions are called immediately when the state machine switches modes
-	UE_API virtual void Activate(const FMoverEventContext& Context);
-	UE_API virtual void Deactivate(const FMoverEventContext& Context);
+	// These functions are called immediately within the simulation. If using async movement, it may be on a worker thread and unsafe to perform some operations on external objects.
+	UE_API virtual void Activate(const FMoverEventContext& Context, FName PrevModeName, const FMoverSimContext& SimContext, const FMoverTickStartData& StartState, FMoverSyncState* OutSyncState, FMoverAuxStateContext* OutAuxState);
+	UE_API virtual void Deactivate(const FMoverEventContext& Context, FName NextModeName, const FMoverSimContext& SimContext);
 
 	// These functions are called when the sync state is changed on the game thread
 	// and a new mode is activated/deactivated
@@ -56,7 +57,7 @@ public:
 	UE_API virtual void Deactivate_External(const FMoverEventContext& Context);
 	
 	UFUNCTION(BlueprintNativeEvent, meta = (DisplayName = "Generate Move", ForceAsFunction))
-	UE_API void GenerateMove(const FMoverTickStartData& StartState, const FMoverTimeStep& TimeStep, UPARAM(ref) FProposedMove& OutProposedMove) const;
+	UE_API void GenerateMove(const FMoverSimContext& SimContext, const FMoverTickStartData& StartState, const FMoverTimeStep& TimeStep, UPARAM(ref) FProposedMove& OutProposedMove) const;
 
 	UFUNCTION(BlueprintNativeEvent, meta = (DisplayName = "Simulation Tick", ForceAsFunction))
 	UE_API void SimulationTick(const FSimulationTickParams& Params, UPARAM(ref) FMoverTickEndData& OutputState);
diff --git a/Engine/Plugins/Experimental/Mover/Source/Mover/Public/MovementModeStateMachine.h b/Engine/Plugins/Experimental/Mover/Source/Mover/Public/MovementModeStateMachine.h
--- a/Engine/Plugins/Experimental/Mover/Source/Mover/Public/MovementModeStateMachine.h
+++ b/Engine/Plugins/Experimental/Mover/Source/Mover/Public/MovementModeStateMachine.h
@@ -29,8 +29,8 @@ class UMovementModeStateMachine : public UObject
 	 GENERATED_UCLASS_BODY()
 
 public:
-	UE_API void RegisterMovementMode(FName ModeName, TObjectPtr<UBaseMovementMode> Mode, bool bIsDefaultMode=false);
-	UE_API void RegisterMovementMode(FName ModeName, TSubclassOf<UBaseMovementMode> ModeType, bool bIsDefaultMode=false);
+	UE_API void RegisterMovementMode(FName ModeName, TObjectPtr<UBaseMovementMode> Mode, bool bIsDefaultMode);
+	UE_API void RegisterMovementMode(FName ModeName, TSubclassOf<UBaseMovementMode> ModeType, bool bIsDefaultMode);
 
 	UE_API void UnregisterMovementMode(FName ModeName);
 	UE_API void ClearAllMovementModes();
@@ -48,7 +48,7 @@ public:
 	UE_API void SetModeImmediately(FName DesiredModeName, bool bShouldReenter=false);
 	UE_API void ClearQueuedMode();
 
-	UE_API void OnSimulationTick(USceneComponent* UpdatedComponent, UPrimitiveComponent* UpdatedPrimitive, UMoverBlackboard* SimBlackboard, const FMoverTickStartData& StartState, const FMoverTimeStep& TimeStep, FMoverTickEndData& OutputState);
+	UE_API void OnSimulationTick(USceneComponent* UpdatedComponent, UPrimitiveComponent* UpdatedPrimitive, UMoverBlackboard* SimBlackboard, const FMoverSimContext& SimContext, const FMoverTickStartData& StartState, const FMoverTimeStep& TimeStep, FMoverTickEndData& OutputState);
  	UE_API void OnSimulationPreRollback(const FMoverSyncState* InvalidSyncState, const FMoverSyncState* SyncState, const FMoverAuxStateContext* InvalidAuxState, const FMoverAuxStateContext* AuxState, const FMoverTimeStep& NewBaseTimeStep);
 	UE_API void OnSimulationRollback(const FMoverSyncState* SyncState, const FMoverAuxStateContext* AuxState, const FMoverTimeStep& NewBaseTimeStep);
 
@@ -129,6 +129,7 @@ protected:
 	// Internal-use-only tick data structs, for efficiency since they typically have the same contents from frame to frame
 	FMoverTickStartData WorkingSubstepStartData;
 	FSimulationTickParams WorkingSimTickParams;
+	FMoverTickEndData* WorkingSimTickEndData = nullptr;	// Only valid during OnSimulationTick and null otherwise
 
 private:
 	// Locks for thread safety on queueing mechanisms
@@ -152,6 +153,7 @@ private:
 	UE_API bool HasAnyInstantEffectsQueued() const;
  	UE_API bool ApplyInstantEffects(FApplyMovementEffectParams& ApplyEffectParams, FMoverSyncState& OutputState);
 	UE_API AActor* GetOwnerActor() const;
+	UE_API FMoverSimContext GetSimContext() const;
 };
 
 #undef UE_API
diff --git a/Engine/Plugins/Experimental/Mover/Source/Mover/Public/MoverComponent.h b/Engine/Plugins/Experimental/Mover/Source/Mover/Public/MoverComponent.h
--- a/Engine/Plugins/Experimental/Mover/Source/Mover/Public/MoverComponent.h
+++ b/Engine/Plugins/Experimental/Mover/Source/Mover/Public/MoverComponent.h
@@ -104,7 +104,7 @@ public:
 	UPROPERTY(BlueprintAssignable, Category = Mover)
 	FMover_OnPostSimRollback OnPostSimulationRollback;
 
-	// Broadcast when a MovementMode has changed. Happens during a simulation tick if the mode changed that tick or when SetModeImmediately is used to change modes.
+	// Broadcast when a MovementMode has changed. Happens during a simulation tick if the mode changed. 
 	UPROPERTY(BlueprintAssignable, Category = Mover)
 	FMover_OnMovementModeChanged OnMovementModeChanged;
 
@@ -172,7 +172,7 @@ public:
 	MOVER_API void InitializeSimulationState(FMoverSyncState* OutSync, FMoverAuxStateContext* OutAux);
 
 	// Primary movement simulation update. Given an starting state and timestep, produce a new state. Called by backend system.
-	MOVER_API void SimulationTick(const FMoverTimeStep& InTimeStep, const FMoverTickStartData& SimInput, OUT FMoverTickEndData& SimOutput);
+	MOVER_API void SimulationTick(const FMoverSimContext& SimContext, const FMoverTimeStep& InTimeStep, const FMoverTickStartData& SimInput, OUT FMoverTickEndData& SimOutput);
 
 	// Specifies which supporting back end class should drive this Mover actor
 	UPROPERTY(EditDefaultsOnly, Category = Mover, meta = (MustImplement = "/Script/Mover.MoverBackendLiaisonInterface"))
diff --git a/Engine/Plugins/Experimental/Mover/Source/Mover/Public/MoverSimulation.h b/Engine/Plugins/Experimental/Mover/Source/Mover/Public/MoverSimulation.h
--- a/Engine/Plugins/Experimental/Mover/Source/Mover/Public/MoverSimulation.h
+++ b/Engine/Plugins/Experimental/Mover/Source/Mover/Public/MoverSimulation.h
@@ -35,6 +35,10 @@ public:
 	UFUNCTION(BlueprintPure, Category = Mover)
 	UE_API FRollbackBlackboardSimWrapper GetRollbackBlackboardSimWrapper() const;
 
+	// Get accessor for the blackboard to be used predictively. You can read from it, but any writes will be temporary and thrown away once prediction is finished.
+	UFUNCTION(BlueprintPure, Category = Mover)
+	UE_API FRollbackBlackboardSimWrapper GetRollbackBlackboardPredictionWrapper() const;
+
 	/**
 	* Attempt to teleport to TargetTransform. The teleport is not guaranteed to happen. This function is meant to be called by an instant movement effect as part of its effect application.
 	* If it succeeds a FTeleportSucceededEventData will be emitted, if it fails a FTeleportFailedEventData will be sent.
diff --git a/Engine/Plugins/Experimental/Mover/Source/Mover/Public/MoverSimulationTypes.h b/Engine/Plugins/Experimental/Mover/Source/Mover/Public/MoverSimulationTypes.h
--- a/Engine/Plugins/Experimental/Mover/Source/Mover/Public/MoverSimulationTypes.h
+++ b/Engine/Plugins/Experimental/Mover/Source/Mover/Public/MoverSimulationTypes.h
@@ -7,6 +7,7 @@
 #include "Engine/NetSerialization.h"
 #include "MoverTypes.h"
 #include "MoveLibrary/MovementRecord.h"
+#include "MoveLibrary/RollbackBlackboardWrappers.h"
 #include "LayeredMove.h"
 #include "LayeredMoveGroup.h"
 #include "MovementModifier.h"
@@ -439,18 +440,50 @@ struct FMoverTickEndData
 	FMovementRecord MoveRecord;
 };
 
+
+
+/**
+ * FMoverSimContext: contains information about the systems supporting the active simulation. This is commonly
+ * passed in to in-simulation functions and events.
+ * All objects are guaranteed to exist for the lifetime of any function called with FMoverSimContext as an argument.
+ * Do not cache references to them, as their lifetime is not guaranteed outside of the simulation.
+ */
+USTRUCT(BlueprintType)
+struct FMoverSimContext
+{
+	GENERATED_BODY()
+
+	UPROPERTY(BlueprintReadOnly, Category = Mover)
+	TObjectPtr<UMoverSimulation> Simulation = nullptr;
+
+	UPROPERTY(BlueprintReadOnly, Category = Mover)
+	FRollbackBlackboardSimWrapper Blackboard;
+
+	/** 
+	 * Owner of the simulation. True type is dependent on the situation (Actor, MoverComp, SG entity, etc.)
+	 * May be null when its lifetime can't be guaranteed
+	 */
+	UPROPERTY(BlueprintReadOnly, Category=Mover) 
+	TObjectPtr<UObject> SimulationOwner = nullptr;
+};
+
+
 // Input parameters to provide context for SimulationTick functions
 USTRUCT(BlueprintType)
 struct FSimulationTickParams
 {
 	GENERATED_BODY()
 
+	// Information about the systems supporting the active simulation
+	UPROPERTY(BlueprintReadOnly, Category = Mover)
+	FMoverSimContext SimContext;
+
 	// Components involved in movement by the simulation
 	// This will be empty when the simulation is ticked asynchronously
 	UPROPERTY(BlueprintReadOnly, EditAnywhere, Category = Mover)
 	FMovingComponentSet MovingComps;
 
-	// Blackboard
+	// Blackboard (old style, soon to be removed)
 	UPROPERTY(BlueprintReadOnly, EditAnywhere, Category = Mover)
 	TObjectPtr<UMoverBlackboard> SimBlackboard;
 
diff --git a/Engine/Plugins/Experimental/Mover/Source/Mover/Public/PhysicsMover/Modes/PhysicsDrivenFallingMode.h b/Engine/Plugins/Experimental/Mover/Source/Mover/Public/PhysicsMover/Modes/PhysicsDrivenFallingMode.h
--- a/Engine/Plugins/Experimental/Mover/Source/Mover/Public/PhysicsMover/Modes/PhysicsDrivenFallingMode.h
+++ b/Engine/Plugins/Experimental/Mover/Source/Mover/Public/PhysicsMover/Modes/PhysicsDrivenFallingMode.h
@@ -24,7 +24,7 @@ class UPhysicsDrivenFallingMode : public UFallingMode, public IPhysicsCharacterM
 	GENERATED_UCLASS_BODY()
 
 public:
-	UE_API virtual void OnRegistered(const FName ModeName) override;
+	UE_API virtual void OnRegistered(const FName ModeName, const FMoverSimContext& SimContext) override;
 	UE_API virtual void SimulationTick_Implementation(const FSimulationTickParams& Params, FMoverTickEndData& OutputState) override;
 	UE_API virtual void UpdateConstraintSettings(Chaos::FCharacterGroundConstraint& Constraint) const override;
 
diff --git a/Engine/Plugins/Experimental/Mover/Source/Mover/Public/PhysicsMover/Modes/PhysicsDrivenSwimmingMode.h b/Engine/Plugins/Experimental/Mover/Source/Mover/Public/PhysicsMover/Modes/PhysicsDrivenSwimmingMode.h
--- a/Engine/Plugins/Experimental/Mover/Source/Mover/Public/PhysicsMover/Modes/PhysicsDrivenSwimmingMode.h
+++ b/Engine/Plugins/Experimental/Mover/Source/Mover/Public/PhysicsMover/Modes/PhysicsDrivenSwimmingMode.h
@@ -24,7 +24,7 @@ class UPhysicsDrivenSwimmingMode : public USwimmingMode, public IPhysicsCharacte
 	GENERATED_UCLASS_BODY()
 
 public:
-	UE_API virtual void OnRegistered(const FName ModeName) override;
+	UE_API virtual void OnRegistered(const FName ModeName, const FMoverSimContext& SimContext) override;
 	UE_API virtual void SimulationTick_Implementation(const FSimulationTickParams& Params, FMoverTickEndData& OutputState) override;
 	UE_API virtual bool AttemptJump(const FSimulationTickParams& Params, float UpwardsSpeed, FMoverTickEndData& Output) override;
 
diff --git a/Engine/Plugins/Experimental/Mover/Source/Mover/Public/PhysicsMover/Modes/PhysicsDrivenWalkingMode.h b/Engine/Plugins/Experimental/Mover/Source/Mover/Public/PhysicsMover/Modes/PhysicsDrivenWalkingMode.h
--- a/Engine/Plugins/Experimental/Mover/Source/Mover/Public/PhysicsMover/Modes/PhysicsDrivenWalkingMode.h
+++ b/Engine/Plugins/Experimental/Mover/Source/Mover/Public/PhysicsMover/Modes/PhysicsDrivenWalkingMode.h
@@ -25,7 +25,7 @@ class UPhysicsDrivenWalkingMode : public UWalkingMode, public IPhysicsCharacterM
 	GENERATED_UCLASS_BODY()
 
 public:
-	UE_API virtual void OnRegistered(const FName ModeName) override;
+	UE_API virtual void OnRegistered(const FName ModeName, const FMoverSimContext& SimContext) override;
 	UE_API virtual void SimulationTick_Implementation(const FSimulationTickParams& Params, FMoverTickEndData& OutputState) override;
 
 	UE_API virtual void UpdateConstraintSettings(Chaos::FCharacterGroundConstraint& Constraint) const override;
diff --git a/Engine/Plugins/Experimental/MoverExamples/Source/MoverExamples/Private/CharacterVariants/Ziplining/ZipliningMode.cpp b/Engine/Plugins/Experimental/MoverExamples/Source/MoverExamples/Private/CharacterVariants/Ziplining/ZipliningMode.cpp
--- a/Engine/Plugins/Experimental/MoverExamples/Source/MoverExamples/Private/CharacterVariants/Ziplining/ZipliningMode.cpp
+++ b/Engine/Plugins/Experimental/MoverExamples/Source/MoverExamples/Private/CharacterVariants/Ziplining/ZipliningMode.cpp
@@ -71,7 +71,7 @@ UZipliningMode::UZipliningMode(const FObjectInitializer& ObjectInitializer)
 }
 
 
-void UZipliningMode::GenerateMove_Implementation(const FMoverTickStartData& StartState, const FMoverTimeStep& TimeStep, FProposedMove& OutProposedMove) const
+void UZipliningMode::GenerateMove_Implementation(const FMoverSimContext& SimContext, const FMoverTickStartData& StartState, const FMoverTimeStep& TimeStep, FProposedMove& OutProposedMove) const
 {
 	UMoverComponent* MoverComp = GetMoverComponent();
 
diff --git a/Engine/Plugins/Experimental/MoverExamples/Source/MoverExamples/Private/MovementBases/FollowPathMode.cpp b/Engine/Plugins/Experimental/MoverExamples/Source/MoverExamples/Private/MovementBases/FollowPathMode.cpp
--- a/Engine/Plugins/Experimental/MoverExamples/Source/MoverExamples/Private/MovementBases/FollowPathMode.cpp
+++ b/Engine/Plugins/Experimental/MoverExamples/Source/MoverExamples/Private/MovementBases/FollowPathMode.cpp
@@ -13,7 +13,7 @@ UFollowPathMode::UFollowPathMode(const FObjectInitializer& ObjectInitializer)
 {
 }
 
-void UFollowPathMode::GenerateMove_Implementation(const FMoverTickStartData& StartState, const FMoverTimeStep& TimeStep, FProposedMove& OutProposedMove) const
+void UFollowPathMode::GenerateMove_Implementation(const FMoverSimContext& SimContext, const FMoverTickStartData& StartState, const FMoverTimeStep& TimeStep, FProposedMove& OutProposedMove) const
 {
 	// TODO: Rework to maintain the split between generating and executing moves
 }
@@ -105,9 +105,9 @@ void UFollowPathMode::SimulationTick_Implementation(const FSimulationTickParams&
 	UpdatedComponent->ComponentVelocity = Velocity;
 }
 
-void UFollowPathMode::Activate(const FMoverEventContext& Context)
+void UFollowPathMode::Activate(const FMoverEventContext& Context, FName PrevModeName, const FMoverSimContext& SimContext, const FMoverTickStartData& StartState, FMoverSyncState* OutSyncState, FMoverAuxStateContext* OutAuxState)
 {
-	Super::Activate(Context);
+	Super::Activate(Context, PrevModeName, SimContext, StartState, OutSyncState, OutAuxState);
 	
 	UpdateControlPoints(true);
 	TimeMultiplier = 1.f / (Duration > UE_KINDA_SMALL_NUMBER ? Duration : UE_KINDA_SMALL_NUMBER);
diff --git a/Engine/Plugins/Experimental/MoverExamples/Source/MoverExamples/Private/MovementBases/FollowSplineMode.cpp b/Engine/Plugins/Experimental/MoverExamples/Source/MoverExamples/Private/MovementBases/FollowSplineMode.cpp
--- a/Engine/Plugins/Experimental/MoverExamples/Source/MoverExamples/Private/MovementBases/FollowSplineMode.cpp
+++ b/Engine/Plugins/Experimental/MoverExamples/Source/MoverExamples/Private/MovementBases/FollowSplineMode.cpp
@@ -48,7 +48,7 @@ UFollowSplineMode::UFollowSplineMode(const FObjectInitializer& ObjectInitializer
 {
 }
 
-void UFollowSplineMode::GenerateMove_Implementation(const FMoverTickStartData& StartState, const FMoverTimeStep& TimeStep, FProposedMove& OutProposedMove) const
+void UFollowSplineMode::GenerateMove_Implementation(const FMoverSimContext& SimContext, const FMoverTickStartData& StartState, const FMoverTimeStep& TimeStep, FProposedMove& OutProposedMove) const
 {
 	// TODO: Rework to maintain the split between generating and executing moves
 }
@@ -166,9 +166,9 @@ void UFollowSplineMode::SetControlSpline(const AActor* SplineProviderActor, FSpl
 	}
 }
 
-void UFollowSplineMode::OnRegistered(const FName ModeName)
+void UFollowSplineMode::OnRegistered(const FName ModeName, const FMoverSimContext& SimContext)
 {
-	Super::OnRegistered(ModeName);
+	Super::OnRegistered(ModeName, SimContext);
 
 	ConfigureSplineData();
 }
diff --git a/Engine/Plugins/Experimental/MoverExamples/Source/MoverExamples/Public/CharacterVariants/Ziplining/ZipliningMode.h b/Engine/Plugins/Experimental/MoverExamples/Source/MoverExamples/Public/CharacterVariants/Ziplining/ZipliningMode.h
--- a/Engine/Plugins/Experimental/MoverExamples/Source/MoverExamples/Public/CharacterVariants/Ziplining/ZipliningMode.h
+++ b/Engine/Plugins/Experimental/MoverExamples/Source/MoverExamples/Public/CharacterVariants/Ziplining/ZipliningMode.h
@@ -20,7 +20,7 @@ class UZipliningMode : public UBaseMovementMode
 	GENERATED_UCLASS_BODY()
 
 public:
-	virtual void GenerateMove_Implementation(const FMoverTickStartData& StartState, const FMoverTimeStep& TimeStep, FProposedMove& OutProposedMove) const override;
+	virtual void GenerateMove_Implementation(const FMoverSimContext& SimContext, const FMoverTickStartData& StartState, const FMoverTimeStep& TimeStep, FProposedMove& OutProposedMove) const override;
 	virtual void SimulationTick_Implementation(const FSimulationTickParams& Params, FMoverTickEndData& OutputState) override;
 
 	// Maximum speed 
diff --git a/Engine/Plugins/Experimental/MoverExamples/Source/MoverExamples/Public/MovementBases/FollowPathMode.h b/Engine/Plugins/Experimental/MoverExamples/Source/MoverExamples/Public/MovementBases/FollowPathMode.h
--- a/Engine/Plugins/Experimental/MoverExamples/Source/MoverExamples/Public/MovementBases/FollowPathMode.h
+++ b/Engine/Plugins/Experimental/MoverExamples/Source/MoverExamples/Public/MovementBases/FollowPathMode.h
@@ -36,10 +36,10 @@ class MOVEREXAMPLES_API UFollowPathMode : public UBaseMovementMode
 
 
 public:
-	virtual void GenerateMove_Implementation(const FMoverTickStartData& StartState, const FMoverTimeStep& TimeStep, FProposedMove& OutProposedMove) const override;
+	virtual void GenerateMove_Implementation(const FMoverSimContext& SimContext, const FMoverTickStartData& StartState, const FMoverTimeStep& TimeStep, FProposedMove& OutProposedMove) const override;
 	virtual void SimulationTick_Implementation(const FSimulationTickParams& Params, FMoverTickEndData& OutputState) override;
 
-	virtual void Activate(const FMoverEventContext& Context) override;
+	virtual void Activate(const FMoverEventContext& Context, FName PrevModeName, const FMoverSimContext& SimContext, const FMoverTickStartData& StartState, FMoverSyncState* OutSyncState, FMoverAuxStateContext* OutAuxState) override;
 	
 	// List of ordered path locations to visit
 	UPROPERTY(EditAnywhere, BlueprintReadWrite, Category = Pathing)
diff --git a/Engine/Plugins/Experimental/MoverExamples/Source/MoverExamples/Public/MovementBases/FollowSplineMode.h b/Engine/Plugins/Experimental/MoverExamples/Source/MoverExamples/Public/MovementBases/FollowSplineMode.h
--- a/Engine/Plugins/Experimental/MoverExamples/Source/MoverExamples/Public/MovementBases/FollowSplineMode.h
+++ b/Engine/Plugins/Experimental/MoverExamples/Source/MoverExamples/Public/MovementBases/FollowSplineMode.h
@@ -55,7 +55,7 @@ class MOVEREXAMPLES_API UFollowSplineMode : public UBaseMovementMode
 	GENERATED_UCLASS_BODY()
 
 public:
-	virtual void GenerateMove_Implementation(const FMoverTickStartData& StartState, const FMoverTimeStep& TimeStep, FProposedMove& OutProposedMove) const override;
+	virtual void GenerateMove_Implementation(const FMoverSimContext& SimContext, const FMoverTickStartData& StartState, const FMoverTimeStep& TimeStep, FProposedMove& OutProposedMove) const override;
 
 	virtual void SimulationTick_Implementation(const FSimulationTickParams& Params, FMoverTickEndData& OutputState) override;
 
@@ -99,7 +99,7 @@ public:
 	TObjectPtr<UCurveFloat> InterpolationCurve;
 
 protected:
-	virtual void OnRegistered(const FName ModeName) override;
+	virtual void OnRegistered(const FName ModeName, const FMoverSimContext& SimContext) override;
 
 	void ConfigureSplineData();
 	bool CanMove(float MeasuredSplineTime) const;
