Mover: fixed case where SimContext wasn't fully filled out (kinematic cases only)

#rnx
#Synced-CL 51816243

diff --git a/Engine/Plugins/Experimental/Mover/Source/Mover/Private/MovementModeStateMachine.cpp b/Engine/Plugins/Experimental/Mover/Source/Mover/Private/MovementModeStateMachine.cpp
--- a/Engine/Plugins/Experimental/Mover/Source/Mover/Private/MovementModeStateMachine.cpp
+++ b/Engine/Plugins/Experimental/Mover/Source/Mover/Private/MovementModeStateMachine.cpp
@@ -369,6 +369,7 @@ void UMovementModeStateMachine::OnSimulationTick(USceneComponent* UpdatedCompone
 			
 			// Execute the combined proposed move
 			{
+				WorkingSimTickParams.SimContext = SimContext;
 				WorkingSimTickParams.StartState = WorkingSubstepStartData;
 				WorkingSimTickParams.MovingComps.SetFrom(MoverComp);
 				WorkingSimTickParams.SimBlackboard = SimBlackboard;
