Mover: making blackboard wrapper mutable so non-const functions can be called on it even when part of a const SimContext

#rnx
#Synced-CL 51941115

diff --git a/Engine/Plugins/Experimental/Mover/Source/Mover/Public/MoverSimulationTypes.h b/Engine/Plugins/Experimental/Mover/Source/Mover/Public/MoverSimulationTypes.h
--- a/Engine/Plugins/Experimental/Mover/Source/Mover/Public/MoverSimulationTypes.h
+++ b/Engine/Plugins/Experimental/Mover/Source/Mover/Public/MoverSimulationTypes.h
@@ -457,7 +457,7 @@ struct FMoverSimContext
 	TObjectPtr<UMoverSimulation> Simulation = nullptr;
 
 	UPROPERTY(BlueprintReadOnly, Category = Mover)
-	FRollbackBlackboardSimWrapper Blackboard;
+	mutable FRollbackBlackboardSimWrapper Blackboard;
 
 	/** 
 	 * Owner of the simulation. True type is dependent on the situation (Actor, MoverComp, SG entity, etc.)
