Mover: fixed issue where sim frame counts in trajectory prediction were always starting from zero, rather than keeping pace with the start time.
Added optional start time/frame to the prediction params. We will fall back to existing behavior of starting from the active sim time if not set.
Fixed cases where trajectory prediction needed to advance the rollback blackboard's prediction time as it stepped samples.

#Synced-CL 52056680

diff --git a/Engine/Plugins/Experimental/ChaosMover/Source/ChaosMover/Private/ChaosMoverSimulation.cpp b/Engine/Plugins/Experimental/ChaosMover/Source/ChaosMover/Private/ChaosMoverSimulation.cpp
--- a/Engine/Plugins/Experimental/ChaosMover/Source/ChaosMover/Private/ChaosMoverSimulation.cpp
+++ b/Engine/Plugins/Experimental/ChaosMover/Source/ChaosMover/Private/ChaosMoverSimulation.cpp
@@ -328,6 +328,8 @@ void UChaosMoverSimulation::PredictTrajectory(const FMoverTimeStep& TimeStep, co
 
 	for (int32 i = 0; i < NumSteps + 1; ++i)
 	{
+		RollbackBlackboard->BeginPredictionFrame(FutureTimeStep);
+
 		FTrajectorySampleInfo& Sample = TrajOut.Deltas[i];
 		const FTransform CurrentWorldTransform(StepDefaultSyncState->GetOrientation_WorldSpace(), StepDefaultSyncState->GetLocation_WorldSpace());
 		if (i == 0)
@@ -379,6 +381,9 @@ void UChaosMoverSimulation::PredictTrajectory(const FMoverTimeStep& TimeStep, co
 		FutureTimeStep.BaseSimTimeMs += FutureTimeStep.StepMs;
 		++FutureTimeStep.ServerFrame;
 	}
+
+	RollbackBlackboard->EndPrediction();
+
 }
 
 void UChaosMoverSimulation::OnInit()
diff --git a/Engine/Plugins/Experimental/Mover/Source/Mover/Private/MoverComponent.cpp b/Engine/Plugins/Experimental/Mover/Source/Mover/Private/MoverComponent.cpp
--- a/Engine/Plugins/Experimental/Mover/Source/Mover/Private/MoverComponent.cpp
+++ b/Engine/Plugins/Experimental/Mover/Source/Mover/Private/MoverComponent.cpp
@@ -2428,8 +2428,17 @@ TArray<FTrajectorySampleInfo> UMoverComponent::GetPredictedTrajectory(FMoverPred
 
 		FMoverTimeStep FutureTimeStep;
 		FutureTimeStep.StepMs = (PredictionParams.SecondsPerSample * 1000.f);
-		FutureTimeStep.BaseSimTimeMs = CachedLastSimTickTimeStep.BaseSimTimeMs;
-		FutureTimeStep.ServerFrame = 0;
+
+		if (PredictionParams.OptionalStartTime.IsSet())
+		{
+			FutureTimeStep.BaseSimTimeMs = FMath::Max(0.0, PredictionParams.OptionalStartTime.GetValue().TimeMs);
+			FutureTimeStep.ServerFrame   = FMath::Max(0, PredictionParams.OptionalStartTime.GetValue().FrameCount);
+		}
+		else
+		{
+			FutureTimeStep.BaseSimTimeMs = FMath::Max(0.0, CachedLastSimTickTimeStep.BaseSimTimeMs);
+			FutureTimeStep.ServerFrame   = FMath::Max(0, CachedLastSimTickTimeStep.ServerFrame);
+		}
 
 		if (const UBaseMovementMode* CurrentMovementMode = GetMovementMode())
 		{
@@ -2452,6 +2461,8 @@ TArray<FTrajectorySampleInfo> UMoverComponent::GetPredictedTrajectory(FMoverPred
 
 				for (int32 i = 0; i < PredictionParams.NumPredictionSamples; ++i)
 				{
+					Simulation->RollbackBlackboard->BeginPredictionFrame(FutureTimeStep);
+
 					// If no further inputs are specified, the previous input cmd will continue to be used
 					if (i < PredictionParams.OptionalInputCmds.Num())
 					{
@@ -2510,6 +2521,8 @@ TArray<FTrajectorySampleInfo> UMoverComponent::GetPredictedTrajectory(FMoverPred
 					SetGravityOverride(bOrigHasGravityOverride, OrigGravityAccelOverride);
 				}
 
+				Simulation->RollbackBlackboard->EndPrediction();
+
 				return OutSamples;
 			}
 		}
diff --git a/Engine/Plugins/Experimental/Mover/Source/Mover/Public/MoverSimulationTypes.h b/Engine/Plugins/Experimental/Mover/Source/Mover/Public/MoverSimulationTypes.h
--- a/Engine/Plugins/Experimental/Mover/Source/Mover/Public/MoverSimulationTypes.h
+++ b/Engine/Plugins/Experimental/Mover/Source/Mover/Public/MoverSimulationTypes.h
@@ -548,6 +548,10 @@ struct FMoverPredictTrajectoryParams
 	UPROPERTY(EditAnywhere, BlueprintReadWrite, Category = Mover)
 	bool bDisableGravity = false;
 
+	/** Optional starting time. If not set, prediction will begin from the current sim time. */
+	UPROPERTY(EditAnywhere, BlueprintReadWrite, Category = Mover)
+	TOptional<FMoverTime> OptionalStartTime;
+
  	/** Optional starting sync state. If not set, prediction will begin from the current state. */
  	UPROPERTY(EditAnywhere, BlueprintReadWrite, Category = Mover)
  	TOptional<FMoverSyncState> OptionalStartSyncState;
