Mover: Adding GetGameplayTags to Moves, Modifiers and Modes so they can report the all tags they have into a container that accumulate all tags thus applied. This is in the context of driving gameplay or animation by gameplay tags without having to know exactly what move, modifier or mode applied a tag and without requiring a redundant networked state

#Synced-CL 51470979

diff --git a/Engine/Plugins/Experimental/Mover/Source/Mover/Private/DefaultMovementSet/MovementModifiers/StanceModifier.cpp b/Engine/Plugins/Experimental/Mover/Source/Mover/Private/DefaultMovementSet/MovementModifiers/StanceModifier.cpp
--- a/Engine/Plugins/Experimental/Mover/Source/Mover/Private/DefaultMovementSet/MovementModifiers/StanceModifier.cpp
+++ b/Engine/Plugins/Experimental/Mover/Source/Mover/Private/DefaultMovementSet/MovementModifiers/StanceModifier.cpp
@@ -31,6 +31,12 @@ bool FStanceModifier::HasGameplayTag(FGameplayTag TagToFind, bool bExactMatch) c
 	return TagToFind.MatchesTag(Mover_IsCrouching);
 }
 
+void FStanceModifier::GetGameplayTags(FGameplayTagContainer& InOutTags) const
+{
+	// TODO: Eventually add other stance tags here like prone
+	InOutTags.AddTag(Mover_IsCrouching);
+}
+
 void FStanceModifier::OnStart(UMoverComponent* MoverComp, const FMoverTimeStep& TimeStep, const FMoverSyncState& SyncState, const FMoverAuxStateContext& AuxState)
 {
 	if (UStanceSettings* StanceSettings = MoverComp->FindSharedSettings_Mutable<UStanceSettings>())
diff --git a/Engine/Plugins/Experimental/Mover/Source/Mover/Private/MovementMode.cpp b/Engine/Plugins/Experimental/Mover/Source/Mover/Private/MovementMode.cpp
--- a/Engine/Plugins/Experimental/Mover/Source/Mover/Private/MovementMode.cpp
+++ b/Engine/Plugins/Experimental/Mover/Source/Mover/Private/MovementMode.cpp
@@ -140,6 +140,11 @@ bool UBaseMovementMode::HasGameplayTag(FGameplayTag TagToFind, bool bExactMatch)
 	return GameplayTags.HasTag(TagToFind);
 }
 
+void UBaseMovementMode::GetGameplayTags(FGameplayTagContainer& InOutTags) const
+{
+	InOutTags.AppendTags(GameplayTags);
+}
+
 const FName UNullMovementMode::NullModeName(TEXT("Null"));
 
 UNullMovementMode::UNullMovementMode(const FObjectInitializer& ObjectInitializer)
diff --git a/Engine/Plugins/Experimental/Mover/Source/Mover/Public/DefaultMovementSet/MovementModifiers/StanceModifier.h b/Engine/Plugins/Experimental/Mover/Source/Mover/Public/DefaultMovementSet/MovementModifiers/StanceModifier.h
--- a/Engine/Plugins/Experimental/Mover/Source/Mover/Public/DefaultMovementSet/MovementModifiers/StanceModifier.h
+++ b/Engine/Plugins/Experimental/Mover/Source/Mover/Public/DefaultMovementSet/MovementModifiers/StanceModifier.h
@@ -38,6 +38,7 @@ public:
 	EStanceMode ActiveStance;
 	
 	UE_API virtual bool HasGameplayTag(FGameplayTag TagToFind, bool bExactMatch) const override;
+	UE_API virtual void GetGameplayTags(FGameplayTagContainer& InOutTags) const override;
 	
 	/** Fired when this modifier is activated. */
 	UE_API virtual void OnStart(UMoverComponent* MoverComp, const FMoverTimeStep& TimeStep, const FMoverSyncState& SyncState, const FMoverAuxStateContext& AuxState) override;
diff --git a/Engine/Plugins/Experimental/Mover/Source/Mover/Public/LayeredMove.h b/Engine/Plugins/Experimental/Mover/Source/Mover/Public/LayeredMove.h
--- a/Engine/Plugins/Experimental/Mover/Source/Mover/Public/LayeredMove.h
+++ b/Engine/Plugins/Experimental/Mover/Source/Mover/Public/LayeredMove.h
@@ -111,7 +111,10 @@ struct FLayeredMoveBase
 	{
 		return false;
 	}
-	
+
+	/** Adds all gameplay tags currently active on this layered move to InOutTags. */
+	virtual void GetGameplayTags(FGameplayTagContainer& InOutTags) const {}
+
 	// Kicks off this move, allowing any initialization to occur.
 	UE_API void StartMove(const UMoverComponent* MoverComp, UMoverBlackboard* SimBlackboard, double CurrentSimTimeMs);
 	// Async version of the above
diff --git a/Engine/Plugins/Experimental/Mover/Source/Mover/Public/MovementMode.h b/Engine/Plugins/Experimental/Mover/Source/Mover/Public/MovementMode.h
--- a/Engine/Plugins/Experimental/Mover/Source/Mover/Public/MovementMode.h
+++ b/Engine/Plugins/Experimental/Mover/Source/Mover/Public/MovementMode.h
@@ -94,6 +94,9 @@ public:
    	 */
 	UE_API virtual bool HasGameplayTag(FGameplayTag TagToFind, bool bExactMatch) const;
 
+	/** Adds all gameplay tags currently active on this movement mode to InOutTags. */
+	UE_API virtual void GetGameplayTags(FGameplayTagContainer& InOutTags) const;
+
 #if WITH_EDITOR
 	UE_API virtual EDataValidationResult IsDataValid(class FDataValidationContext& Context) const override;
 #endif // WITH_EDITOR
diff --git a/Engine/Plugins/Experimental/Mover/Source/Mover/Public/MovementModifier.h b/Engine/Plugins/Experimental/Mover/Source/Mover/Public/MovementModifier.h
--- a/Engine/Plugins/Experimental/Mover/Source/Mover/Public/MovementModifier.h
+++ b/Engine/Plugins/Experimental/Mover/Source/Mover/Public/MovementModifier.h
@@ -197,6 +197,9 @@ struct FMovementModifierBase
 		return false;
 	}
 
+	/** Adds all gameplay tags currently active on this modifier to InOutTags. */
+	virtual void GetGameplayTags(FGameplayTagContainer& InOutTags) const {}
+
 protected:
 	/**
 	 * Modifier handle local to this client or server instance. Used to cancel or query for a active/queued modifier.
