Mover: adds blackboard helpers for common entry creation cases, to reduce the verbosity required.
Added new current-frame-only persistence policy for entries that are only valid within the sim frame they were written.

#Synced-CL 52056680

diff --git a/Engine/Plugins/Experimental/Mover/Source/Mover/Private/MoveLibrary/RollbackBlackboard.cpp b/Engine/Plugins/Experimental/Mover/Source/Mover/Private/MoveLibrary/RollbackBlackboard.cpp
--- a/Engine/Plugins/Experimental/Mover/Source/Mover/Private/MoveLibrary/RollbackBlackboard.cpp
+++ b/Engine/Plugins/Experimental/Mover/Source/Mover/Private/MoveLibrary/RollbackBlackboard.cpp
@@ -90,7 +90,10 @@ bool URollbackBlackboard::BlackboardEntryBase::CanReadEntryAtTime(const EntryTim
 
 	switch (Settings.PersistencePolicy)
 	{
+		PRAGMA_DISABLE_DEPRECATION_WARNINGS;
 		case EBlackboardPersistencePolicy::NextFrameOnly:
+		PRAGMA_ENABLE_DEPRECATION_WARNINGS;
+		case EBlackboardPersistencePolicy::ThroughNextFrame:
 		{
 			// if value wasn't last set during the current or immediate prior sim frame, then it can't be read
 			if (!(ReaderTimeStamp.Frame == EntryTimeStamp.Frame ||
@@ -101,6 +104,16 @@ bool URollbackBlackboard::BlackboardEntryBase::CanReadEntryAtTime(const EntryTim
 		}
 		break;
 
+		case EBlackboardPersistencePolicy::CurrentFrameOnly:
+		{
+			// if value wasn't last set during the current sim frame, then it can't be read
+			if (ReaderTimeStamp.Frame != EntryTimeStamp.Frame)
+			{
+				return false;
+			}
+		}
+		break;
+
 		case EBlackboardPersistencePolicy::Forever:
 			// If the reader is in the past compared to the time of the entry, then it can't be read
 			if (ReaderTimeStamp.Frame < EntryTimeStamp.Frame)
diff --git a/Engine/Plugins/Experimental/Mover/Source/Mover/Private/MoveLibrary/RollbackBlackboardLibrary.cpp b/Engine/Plugins/Experimental/Mover/Source/Mover/Private/MoveLibrary/RollbackBlackboardLibrary.cpp
new file mode 100644
--- /dev/null
+++ b/Engine/Plugins/Experimental/Mover/Source/Mover/Private/MoveLibrary/RollbackBlackboardLibrary.cpp
@@ -0,0 +1,26 @@
+// Copyright Epic Games, Inc. All Rights Reserved.
+
+#include "MoveLibrary/RollbackBlackboardLibrary.h"
+
+#include UE_INLINE_GENERATED_CPP_BY_NAME(RollbackBlackboardLibrary)
+
+
+URollbackBlackboard::EntrySettings URollbackBlackboardLibrary::MakeEventEntrySettings(int32 MaxHistoryCount)
+{
+	URollbackBlackboard::EntrySettings EventEntrySettings;
+	EventEntrySettings.SizingPolicy = EBlackboardSizingPolicy::FixedDeclaredSize;
+	EventEntrySettings.FixedSize = MaxHistoryCount;
+	EventEntrySettings.PersistencePolicy = EBlackboardPersistencePolicy::Forever;
+	EventEntrySettings.RollbackPolicy = EBlackboardRollbackPolicy::InvalidatedOnRollback;
+	return EventEntrySettings;
+}
+
+
+URollbackBlackboard::EntrySettings URollbackBlackboardLibrary::MakeSingleFrameEntrySettings()
+{
+	URollbackBlackboard::EntrySettings SingleFrameEntrySettings;
+	SingleFrameEntrySettings.SizingPolicy = EBlackboardSizingPolicy::SingleEntry;
+	SingleFrameEntrySettings.PersistencePolicy = EBlackboardPersistencePolicy::CurrentFrameOnly;
+	SingleFrameEntrySettings.RollbackPolicy = EBlackboardRollbackPolicy::InvalidatedOnRollback;
+	return SingleFrameEntrySettings;
+}
\ No newline at end of file
diff --git a/Engine/Plugins/Experimental/Mover/Source/Mover/Public/MoveLibrary/RollbackBlackboard.h b/Engine/Plugins/Experimental/Mover/Source/Mover/Public/MoveLibrary/RollbackBlackboard.h
--- a/Engine/Plugins/Experimental/Mover/Source/Mover/Public/MoveLibrary/RollbackBlackboard.h
+++ b/Engine/Plugins/Experimental/Mover/Source/Mover/Public/MoveLibrary/RollbackBlackboard.h
@@ -61,19 +61,25 @@ enum class EBlackboardSizingPolicy : uint8
 };
 
 /**
- * Determines how long a blackboard entry is valid once its been set.
+ * Determines how long a blackboard entry can be read after its value is set. Time is in terms of simulation frames or simulation time.
  */
 UENUM()
 enum class EBlackboardPersistencePolicy : uint8
 {
 	/** Any edits to the entry remain valid and readable indefinitely */
-	Forever,
+	Forever = 0,
 
 	/** Any edits remain valid until the end of the next simulation frame. If an edit occurs during frame N, the entry will be readable for the remainder of frame N and frame N+1. */
-	NextFrameOnly,
+	NextFrameOnly UE_DEPRECATED(5.8, "Deprecated in favor of ThroughNextFrame, which has identical behavior") = 1,
+
+	/** Any edits remain valid until the end of the next simulation frame. If an edit occurs during frame N, the entry will be readable for the remainder of frame N and frame N+1. */
+	ThroughNextFrame = 2,
+
+	/** Any edits remain valid until the end of the current simulation frame. If an edit occurs during frame N, the entry will be readable for the remainder of frame N, until frame N+1 begins. */
+	CurrentFrameOnly = 3,
 
 	/** NOT YET IMPLEMENTED. Any edits remain valid until a specified amount of simulation time has passed. */
-	LimitedTime UMETA(Hidden),
+	LimitedTime = 4 UMETA(Hidden)
 };
 
 
diff --git a/Engine/Plugins/Experimental/Mover/Source/Mover/Public/MoveLibrary/RollbackBlackboardLibrary.h b/Engine/Plugins/Experimental/Mover/Source/Mover/Public/MoveLibrary/RollbackBlackboardLibrary.h
new file mode 100644
--- /dev/null
+++ b/Engine/Plugins/Experimental/Mover/Source/Mover/Public/MoveLibrary/RollbackBlackboardLibrary.h
@@ -0,0 +1,43 @@
+// Copyright Epic Games, Inc. All Rights Reserved.
+
+#pragma once
+
+#include "CoreMinimal.h"
+#include "Kismet/BlueprintFunctionLibrary.h"
+#include "MoveLibrary/RollbackBlackboard.h"
+
+#include "RollbackBlackboardLibrary.generated.h"
+
+#define UE_API MOVER_API
+
+
+/**
+ * RollbackBlackboardLibrary: a collection of static functions to help working with a rollback blackboard
+ */
+UCLASS(MinimalAPI)
+class URollbackBlackboardLibrary : public UBlueprintFunctionLibrary
+{
+	GENERATED_BODY()
+
+public:
+	/**
+	 * An 'event' entry is good for data that is rarely authored each frame and may be useful for some time into the future.
+	 * These are commonly used to denote gameplay events that need to be referenced a short time later.
+	 * Examples: 
+	 *     - recording a jump apex's time & height, to be used upon landing to see how far a character fell
+	 *     - denoting a state like "was ducking" at the start of navigating under an obstacle, to be used later to restore the ducking behavior after clearing the obstacle
+	 * For MaxHistoryCount, consider how many times this entry might change during your worst-case prediction window, which is typically 1000 milliseconds or less.
+	 */ 
+	static UE_API URollbackBlackboard::EntrySettings MakeEventEntrySettings(int32 MaxHistoryCount);
+
+	/**
+	 * Single-frame entries are good for data that only needs to be communicated between different parts of the 
+	 * movement sim during the current frame and has no use beyond that.
+	 * Example: information about acceleration used in a mode's GenerateMove that you want to capture and include in the SyncState during SimulationTick
+	 */
+	static UE_API URollbackBlackboard::EntrySettings MakeSingleFrameEntrySettings();
+};
+
+
+
+#undef UE_API
\ No newline at end of file
