Mover: shifting deprecation macros to be outside a switch statement, to appease static analysis warnings

#rnx
#Synced-CL 52125990

diff --git a/Engine/Plugins/Experimental/Mover/Source/Mover/Private/MoveLibrary/RollbackBlackboard.cpp b/Engine/Plugins/Experimental/Mover/Source/Mover/Private/MoveLibrary/RollbackBlackboard.cpp
--- a/Engine/Plugins/Experimental/Mover/Source/Mover/Private/MoveLibrary/RollbackBlackboard.cpp
+++ b/Engine/Plugins/Experimental/Mover/Source/Mover/Private/MoveLibrary/RollbackBlackboard.cpp
@@ -88,11 +88,11 @@ bool URollbackBlackboard::BlackboardEntryBase::CanReadEntryAtTime(const EntryTim
 		return false;	// entry isn't initialized yet or never set
 	}
 
+	PRAGMA_DISABLE_DEPRECATION_WARNINGS;	// supporting deprecation of EBlackboardPersistencePolicy::NextFrameOnly
+
 	switch (Settings.PersistencePolicy)
 	{
-		PRAGMA_DISABLE_DEPRECATION_WARNINGS;
-		case EBlackboardPersistencePolicy::NextFrameOnly:
-		PRAGMA_ENABLE_DEPRECATION_WARNINGS;
+		case EBlackboardPersistencePolicy::NextFrameOnly:	// deprecated
 		case EBlackboardPersistencePolicy::ThroughNextFrame:
 		{
 			// if value wasn't last set during the current or immediate prior sim frame, then it can't be read
@@ -126,6 +126,9 @@ bool URollbackBlackboard::BlackboardEntryBase::CanReadEntryAtTime(const EntryTim
 			break;
 	}
 
+	PRAGMA_ENABLE_DEPRECATION_WARNINGS;
+
+
 	return true;
 }
 
