Mover: adding additional "rolling" entry settings type to help common blackboard cases

#rnx

#Synced-CL 52291027

diff --git a/Engine/Plugins/Experimental/Mover/Source/Mover/Private/MoveLibrary/RollbackBlackboardLibrary.cpp b/Engine/Plugins/Experimental/Mover/Source/Mover/Private/MoveLibrary/RollbackBlackboardLibrary.cpp
--- a/Engine/Plugins/Experimental/Mover/Source/Mover/Private/MoveLibrary/RollbackBlackboardLibrary.cpp
+++ b/Engine/Plugins/Experimental/Mover/Source/Mover/Private/MoveLibrary/RollbackBlackboardLibrary.cpp
@@ -23,4 +23,15 @@ URollbackBlackboard::EntrySettings URollbackBlackboardLibrary::MakeSingleFrameEn
 	SingleFrameEntrySettings.PersistencePolicy = EBlackboardPersistencePolicy::CurrentFrameOnly;
 	SingleFrameEntrySettings.RollbackPolicy = EBlackboardRollbackPolicy::InvalidatedOnRollback;
 	return SingleFrameEntrySettings;
-}
\ No newline at end of file
+}
+
+
+URollbackBlackboard::EntrySettings URollbackBlackboardLibrary::MakeRollingEntrySettings()
+{
+	URollbackBlackboard::EntrySettings RollingEntrySettings;
+	RollingEntrySettings.SizingPolicy = EBlackboardSizingPolicy::FixedDeclaredSize;	// TODO: use FixedBackendBufferSize policy once it's supported
+	RollingEntrySettings.FixedSize = 30;
+	RollingEntrySettings.PersistencePolicy = EBlackboardPersistencePolicy::Forever;
+	RollingEntrySettings.RollbackPolicy = EBlackboardRollbackPolicy::InvalidatedOnRollback;
+	return RollingEntrySettings;
+}
diff --git a/Engine/Plugins/Experimental/Mover/Source/Mover/Public/MoveLibrary/RollbackBlackboardLibrary.h b/Engine/Plugins/Experimental/Mover/Source/Mover/Public/MoveLibrary/RollbackBlackboardLibrary.h
--- a/Engine/Plugins/Experimental/Mover/Source/Mover/Public/MoveLibrary/RollbackBlackboardLibrary.h
+++ b/Engine/Plugins/Experimental/Mover/Source/Mover/Public/MoveLibrary/RollbackBlackboardLibrary.h
@@ -36,6 +36,15 @@ public:
 	 * Example: information about acceleration used in a mode's GenerateMove that you want to capture and include in the SyncState during SimulationTick
 	 */
 	static UE_API URollbackBlackboard::EntrySettings MakeSingleFrameEntrySettings();
+
+	/**
+	 * A 'rolling' entry is good for data that changes every frame.
+	 * Examples:
+	 *     - a dynamic homing target location, or motion warping position, that is moving as we relate to it
+	 *     - information about animation state from an associated skeletal mesh
+	 */
+	static UE_API URollbackBlackboard::EntrySettings MakeRollingEntrySettings();
+
 };
 
 
