Mover spring model states: fixing issue where interrupted move generation, such as during a root motion montage, would keep residual spring state around.  Now it resets.

#virtualized
#Synced-CL 52291027

diff --git a/Engine/Plugins/Experimental/Mover/Source/Mover/Private/DefaultMovementSet/Modes/SimpleSpringWalkingMode.cpp b/Engine/Plugins/Experimental/Mover/Source/Mover/Private/DefaultMovementSet/Modes/SimpleSpringWalkingMode.cpp
--- a/Engine/Plugins/Experimental/Mover/Source/Mover/Private/DefaultMovementSet/Modes/SimpleSpringWalkingMode.cpp
+++ b/Engine/Plugins/Experimental/Mover/Source/Mover/Private/DefaultMovementSet/Modes/SimpleSpringWalkingMode.cpp
@@ -2,27 +2,33 @@
 
 #include "DefaultMovementSet/Modes/SimpleSpringWalkingMode.h"
 #include "DefaultMovementSet/Modes/SimpleSpringState.h"
-
 #include "MoverComponent.h"
 #include "DefaultMovementSet/Settings/CommonLegacyMovementSettings.h"
 #include "Math/SpringMath.h"
+#include "MoveLibrary/RollbackBlackboard.h"
 
 
 #include UE_INLINE_GENERATED_CPP_BY_NAME(SimpleSpringWalkingMode)
 
+
+
 void USimpleSpringWalkingMode::SimulationTick_Implementation(const FSimulationTickParams& Params, FMoverTickEndData& OutputState)
 {
 	Super::SimulationTick_Implementation(Params, OutputState);
 
-	// We've already updated the spring state during GenerateMove, and just need to copy it into the output simulation state
-	if (const FSimpleSpringState* InSpringState = Params.StartState.SyncState.SyncStateCollection.FindDataByType<FSimpleSpringState>())
+	// If we've updated the spring state during GenerateMove, we need to copy it into the output simulation state. Otherwise, the spring state is effectively reset.
+	bool bDidGenerateMove = false;
+	if (Params.SimContext.Blackboard.TryGet(DidGenerateMoveThisFrame, bDidGenerateMove) && bDidGenerateMove)
 	{
-		FSimpleSpringState& OutputSpringState = OutputState.SyncState.SyncStateCollection.FindOrAddMutableDataByType<FSimpleSpringState>();
-		OutputSpringState = *InSpringState;
+		if (const FSimpleSpringState* InSpringState = Params.StartState.SyncState.SyncStateCollection.FindDataByType<FSimpleSpringState>())
+		{
+			FSimpleSpringState& OutputSpringState = OutputState.SyncState.SyncStateCollection.FindOrAddMutableDataByType<FSimpleSpringState>();
+			OutputSpringState = *InSpringState;
+		}
 	}
 }
 
-void USimpleSpringWalkingMode::GenerateWalkMove_Implementation(FMoverTickStartData& StartState, float DeltaSeconds, const FVector& DesiredVelocity,
+void USimpleSpringWalkingMode::GenerateWalkMove_Implementation(FMoverTickStartData& StartState, float DeltaSeconds, const FMoverSimContext& SimContext, const FVector& DesiredVelocity,
 	const FQuat& DesiredFacing, const FQuat& CurrentFacing, FVector& InOutAngularVelocityDegrees, FVector& InOutVelocity)
 {
 	FSimpleSpringState& SpringState = StartState.SyncState.SyncStateCollection.FindOrAddMutableDataByType<FSimpleSpringState>();
diff --git a/Engine/Plugins/Experimental/Mover/Source/Mover/Private/DefaultMovementSet/Modes/SimpleWalkingMode.cpp b/Engine/Plugins/Experimental/Mover/Source/Mover/Private/DefaultMovementSet/Modes/SimpleWalkingMode.cpp
--- a/Engine/Plugins/Experimental/Mover/Source/Mover/Private/DefaultMovementSet/Modes/SimpleWalkingMode.cpp
+++ b/Engine/Plugins/Experimental/Mover/Source/Mover/Private/DefaultMovementSet/Modes/SimpleWalkingMode.cpp
@@ -4,9 +4,21 @@
 
 #include "MoverComponent.h"
 #include "DefaultMovementSet/Settings/CommonLegacyMovementSettings.h"
+#include "MoveLibrary/RollbackBlackboardLibrary.h"
 
 #include UE_INLINE_GENERATED_CPP_BY_NAME(SimpleWalkingMode)
 
+const FName USimpleWalkingMode::DidGenerateMoveThisFrame = TEXT("DidGenerateMoveThisFrame");
+
+
+void USimpleWalkingMode::OnRegistered(const FName ModeName, const FMoverSimContext& SimContext)
+{
+	Super::OnRegistered(ModeName, SimContext);
+
+	SimContext.Blackboard.CreateEntry<bool>(DidGenerateMoveThisFrame, URollbackBlackboardLibrary::MakeSingleFrameEntrySettings());
+}
+
+
 void USimpleWalkingMode::GenerateMove_Implementation(const FMoverSimContext& SimContext, const FMoverTickStartData& StartState, const FMoverTimeStep& TimeStep,
                                                    FProposedMove& OutProposedMove) const
 {
@@ -103,10 +115,12 @@ void USimpleWalkingMode::GenerateMove_Implementation(const FMoverSimContext& Sim
 	// Why is this needed?
 	// Because some modes want to mutate their data inside the generate walk move
 	USimpleWalkingMode* MutableSimpleWalkMode  = const_cast<USimpleWalkingMode*>(this);
-	MutableSimpleWalkMode->GenerateWalkMove(const_cast<FMoverTickStartData&>(StartState), DeltaSeconds, DesiredVelocity, DesiredFacing, CurrentFacing, OutProposedMove.AngularVelocityDegrees, OutProposedMove.LinearVelocity);
+	MutableSimpleWalkMode->GenerateWalkMove(const_cast<FMoverTickStartData&>(StartState), DeltaSeconds, SimContext, DesiredVelocity, DesiredFacing, CurrentFacing, OutProposedMove.AngularVelocityDegrees, OutProposedMove.LinearVelocity);
+
+	SimContext.Blackboard.TrySet(DidGenerateMoveThisFrame, true);
 }
 
-void USimpleWalkingMode::GenerateWalkMove_Implementation(FMoverTickStartData& StartState, float DeltaSeconds, const FVector& DesiredVelocity,
+void USimpleWalkingMode::GenerateWalkMove_Implementation(FMoverTickStartData& StartState, float DeltaSeconds, const FMoverSimContext& SimContext, const FVector& DesiredVelocity,
 	const FQuat& DesiredFacing, const FQuat& CurrentFacing, FVector& InOutAngularVelocityDegrees, FVector& InOutVelocity)
 {
 	InOutVelocity = DesiredVelocity;
diff --git a/Engine/Plugins/Experimental/Mover/Source/Mover/Private/DefaultMovementSet/Modes/SmoothWalkingMode.cpp b/Engine/Plugins/Experimental/Mover/Source/Mover/Private/DefaultMovementSet/Modes/SmoothWalkingMode.cpp
--- a/Engine/Plugins/Experimental/Mover/Source/Mover/Private/DefaultMovementSet/Modes/SmoothWalkingMode.cpp
+++ b/Engine/Plugins/Experimental/Mover/Source/Mover/Private/DefaultMovementSet/Modes/SmoothWalkingMode.cpp
@@ -13,15 +13,19 @@ void USmoothWalkingMode::SimulationTick_Implementation(const FSimulationTickPara
 {
 	Super::SimulationTick_Implementation(Params, OutputState);
 
-	// We've already updated the spring state during GenerateMove, and just need to copy it into the output simulation state
-	if (const FSmoothWalkingState* InSpringState = Params.StartState.SyncState.SyncStateCollection.FindDataByType<FSmoothWalkingState>())
+	// If we've updated the spring state during GenerateMove, we need to copy it into the output simulation state. Otherwise, the spring state is effectively reset.
+	bool bDidGenerateMove = false;
+	if (Params.SimContext.Blackboard.TryGet(DidGenerateMoveThisFrame, bDidGenerateMove) && bDidGenerateMove)
 	{
-		FSmoothWalkingState& OutputSpringState = OutputState.SyncState.SyncStateCollection.FindOrAddMutableDataByType<FSmoothWalkingState>();
-		OutputSpringState = *InSpringState;
+		if (const FSmoothWalkingState* InSpringState = Params.StartState.SyncState.SyncStateCollection.FindDataByType<FSmoothWalkingState>())
+		{
+			FSmoothWalkingState& OutputSpringState = OutputState.SyncState.SyncStateCollection.FindOrAddMutableDataByType<FSmoothWalkingState>();
+			OutputSpringState = *InSpringState;
+		}
 	}
 }
 
-void USmoothWalkingMode::GenerateWalkMove_Implementation(FMoverTickStartData& StartState, float DeltaSeconds, const FVector& DesiredVelocity,
+void USmoothWalkingMode::GenerateWalkMove_Implementation(FMoverTickStartData& StartState, float DeltaSeconds, const FMoverSimContext& SimContext, const FVector& DesiredVelocity,
 	const FQuat& DesiredFacing, const FQuat& CurrentFacing, FVector& InOutAngularVelocityDegrees, FVector& InOutVelocity) 
 {
 	if (DeltaSeconds <= FLT_EPSILON)
diff --git a/Engine/Plugins/Experimental/Mover/Source/Mover/Public/DefaultMovementSet/Modes/SimpleSpringWalkingMode.h b/Engine/Plugins/Experimental/Mover/Source/Mover/Public/DefaultMovementSet/Modes/SimpleSpringWalkingMode.h
--- a/Engine/Plugins/Experimental/Mover/Source/Mover/Public/DefaultMovementSet/Modes/SimpleSpringWalkingMode.h
+++ b/Engine/Plugins/Experimental/Mover/Source/Mover/Public/DefaultMovementSet/Modes/SimpleSpringWalkingMode.h
@@ -19,7 +19,7 @@ class USimpleSpringWalkingMode : public USimpleWalkingMode
 
 public:
 	UE_API virtual void SimulationTick_Implementation(const FSimulationTickParams& Params, FMoverTickEndData& OutputState) override;
-	UE_API virtual void GenerateWalkMove_Implementation(FMoverTickStartData& StartState, float DeltaSeconds, const FVector& DesiredVelocity,
+	UE_API virtual void GenerateWalkMove_Implementation(FMoverTickStartData& StartState, float DeltaSeconds, const FMoverSimContext& SimContext, const FVector& DesiredVelocity,
 									 const FQuat& DesiredFacing, const FQuat& CurrentFacing, FVector& InOutAngularVelocityDegrees, FVector& InOutVelocity) override;
 
 protected:
diff --git a/Engine/Plugins/Experimental/Mover/Source/Mover/Public/DefaultMovementSet/Modes/SimpleWalkingMode.h b/Engine/Plugins/Experimental/Mover/Source/Mover/Public/DefaultMovementSet/Modes/SimpleWalkingMode.h
--- a/Engine/Plugins/Experimental/Mover/Source/Mover/Public/DefaultMovementSet/Modes/SimpleWalkingMode.h
+++ b/Engine/Plugins/Experimental/Mover/Source/Mover/Public/DefaultMovementSet/Modes/SimpleWalkingMode.h
@@ -17,17 +17,20 @@ class USimpleWalkingMode : public UWalkingMode
 	GENERATED_BODY()
 
 public:
-
+	UE_API virtual void OnRegistered(const FName ModeName, const FMoverSimContext& SimContext) override;
 	UE_API virtual void GenerateMove_Implementation(const FMoverSimContext& SimContext, const FMoverTickStartData& StartState, const FMoverTimeStep& TimeStep, FProposedMove& OutProposedMove) const override;
 
 	// Override this to make a simple walking mode
 	UFUNCTION(BlueprintNativeEvent, meta = (DisplayName = "Generate Simple Walk Move"))
-	UE_API void GenerateWalkMove(UPARAM(ref) FMoverTickStartData& StartState, float DeltaSeconds, const FVector& DesiredVelocity,
+	UE_API void GenerateWalkMove(UPARAM(ref) FMoverTickStartData& StartState, float DeltaSeconds, const FMoverSimContext& SimContext, const FVector& DesiredVelocity,
 	                                     const FQuat& DesiredFacing, const FQuat& CurrentFacing, UPARAM(ref) FVector& InOutAngularVelocityDegrees, UPARAM(ref) FVector& InOutVelocity);
 
 	// If this value is greater or equal to 0, this will override the max speed read from the common legacy shared walk settings
 	UPROPERTY(EditAnywhere, BlueprintReadWrite, Category = "Mover|Walking Settings", meta = (ForceUnits = "cm/s"))
 	float MaxSpeedOverride = -1.0f;
+
+protected:
+	static const FName DidGenerateMoveThisFrame;	// boolean blackboard entry name, set whenever GenerateMove finishes
 };
 
 #undef UE_API
diff --git a/Engine/Plugins/Experimental/Mover/Source/Mover/Public/DefaultMovementSet/Modes/SmoothWalkingMode.h b/Engine/Plugins/Experimental/Mover/Source/Mover/Public/DefaultMovementSet/Modes/SmoothWalkingMode.h
--- a/Engine/Plugins/Experimental/Mover/Source/Mover/Public/DefaultMovementSet/Modes/SmoothWalkingMode.h
+++ b/Engine/Plugins/Experimental/Mover/Source/Mover/Public/DefaultMovementSet/Modes/SmoothWalkingMode.h
@@ -19,7 +19,7 @@ class USmoothWalkingMode : public USimpleWalkingMode
 
 public:
 	UE_API virtual void SimulationTick_Implementation(const FSimulationTickParams& Params, FMoverTickEndData& OutputState) override;
-	UE_API virtual void GenerateWalkMove_Implementation(FMoverTickStartData& StartState, float DeltaSeconds, const FVector& DesiredVelocity,
+	UE_API virtual void GenerateWalkMove_Implementation(FMoverTickStartData& StartState, float DeltaSeconds, const FMoverSimContext& SimContext, const FVector& DesiredVelocity,
 									 const FQuat& DesiredFacing, const FQuat& CurrentFacing, FVector& InOutAngularVelocityDegrees, FVector& InOutVelocity) override;
 
 
