Mover: Adding an ensure to navwalking modes noting that the bAllowClientSideNavigation navigation system setting needs to be true if nav walking mode is running on a client so a navmesh is available for nav walking to use. This is the case when running a pawn in nav walking mode with Fixed Tick Forward Predicted Network prediction simulation.

#Synced-CL 52331429

diff --git a/Engine/Plugins/Experimental/Mover/Source/Mover/Private/DefaultMovementSet/Modes/AsyncNavWalkingMode.cpp b/Engine/Plugins/Experimental/Mover/Source/Mover/Private/DefaultMovementSet/Modes/AsyncNavWalkingMode.cpp
--- a/Engine/Plugins/Experimental/Mover/Source/Mover/Private/DefaultMovementSet/Modes/AsyncNavWalkingMode.cpp
+++ b/Engine/Plugins/Experimental/Mover/Source/Mover/Private/DefaultMovementSet/Modes/AsyncNavWalkingMode.cpp
@@ -435,11 +435,17 @@ const INavigationDataInterface* UAsyncNavWalkingMode::GetNavData() const
 
 	if (const UWorld* World = GetWorld())
 	{
-		const UNavigationSystemV1* NavSys = Cast<UNavigationSystemV1>(World->GetNavigationSystem());
-		if (NavSys && NavMoverComponent.IsValid())
+		if (const UNavigationSystemV1* NavSys = Cast<UNavigationSystemV1>(World->GetNavigationSystem()))
 		{
-			const FNavAgentProperties& AgentProps = NavMoverComponent->GetNavAgentPropertiesRef();
-			NavData = NavSys->GetNavDataForProps(AgentProps, NavMoverComponent->GetNavLocation());
+			if (NavMoverComponent.IsValid())
+			{
+				const FNavAgentProperties& AgentProps = NavMoverComponent->GetNavAgentPropertiesRef();
+				NavData = NavSys->GetNavDataForProps(AgentProps, NavMoverComponent->GetNavLocation());
+			}
+		}
+		else
+		{
+			NavWalkingUtils::CheckClientSideNavigationSettings(GetMoverComponent());
 		}
 	}
 
diff --git a/Engine/Plugins/Experimental/Mover/Source/Mover/Private/DefaultMovementSet/Modes/NavWalkingMode.cpp b/Engine/Plugins/Experimental/Mover/Source/Mover/Private/DefaultMovementSet/Modes/NavWalkingMode.cpp
--- a/Engine/Plugins/Experimental/Mover/Source/Mover/Private/DefaultMovementSet/Modes/NavWalkingMode.cpp
+++ b/Engine/Plugins/Experimental/Mover/Source/Mover/Private/DefaultMovementSet/Modes/NavWalkingMode.cpp
@@ -10,6 +10,7 @@
 #include "Components/ShapeComponent.h"
 #include "DefaultMovementSet/NavMoverComponent.h"
 #include "DefaultMovementSet/Settings/CommonLegacyMovementSettings.h"
+#include "Misc/ConfigCacheIni.h"
 #include "MoveLibrary/FloorQueryUtils.h"
 #include "MoveLibrary/GroundMovementUtils.h"
 #include "MoveLibrary/ModularMovement.h"
@@ -18,6 +19,17 @@
 #include UE_INLINE_GENERATED_CPP_BY_NAME(NavWalkingMode)
 
 
+void NavWalkingUtils::CheckClientSideNavigationSettings(const UMoverComponent* MoverComponent)
+{
+	if (MoverComponent && MoverComponent->GetNetMode() == NM_Client)
+	{
+		bool bClientSideNavigationAllowed = false;
+		GConfig->GetBool(TEXT("/Script/NavigationSystem"), TEXT("bAllowClientSideNavigation"), bClientSideNavigationAllowed, GEngineIni);
+				
+		ensureMsgf(bClientSideNavigationAllowed, TEXT("NavWalking Mode on %s expected a valid navigation system on clients but bAllowClientSideNavigation wasn't enabled. Enable bAllowClientSideNavigation if NavWalking needs to run on Clients."), *GetNameSafe(MoverComponent->GetOwner()));
+	}
+}
+
 UNavWalkingMode::UNavWalkingMode()
 	: bSweepWhileNavWalking(true)
 	, bProjectNavMeshWalking(false)
@@ -436,11 +448,17 @@ const INavigationDataInterface* UNavWalkingMode::GetNavData() const
 	
 	if (const UWorld* World = GetWorld())
 	{
-		const UNavigationSystemV1* NavSys = Cast<UNavigationSystemV1>(World->GetNavigationSystem());
-		if (NavSys && NavMoverComponent)
+		if (const UNavigationSystemV1* NavSys = Cast<UNavigationSystemV1>(World->GetNavigationSystem()))
+		{
+			if (NavMoverComponent)
+			{
+				const FNavAgentProperties& AgentProps = NavMoverComponent->GetNavAgentPropertiesRef();
+				NavData = NavSys->GetNavDataForProps(AgentProps, NavMoverComponent->GetNavLocation());
+			}
+		}
+		else
 		{
-			const FNavAgentProperties& AgentProps = NavMoverComponent->GetNavAgentPropertiesRef();
-			NavData = NavSys->GetNavDataForProps(AgentProps, NavMoverComponent->GetNavLocation());
+			NavWalkingUtils::CheckClientSideNavigationSettings(GetMoverComponent());
 		}
 	}
 	
diff --git a/Engine/Plugins/Experimental/Mover/Source/Mover/Public/DefaultMovementSet/Modes/NavWalkingMode.h b/Engine/Plugins/Experimental/Mover/Source/Mover/Public/DefaultMovementSet/Modes/NavWalkingMode.h
--- a/Engine/Plugins/Experimental/Mover/Source/Mover/Public/DefaultMovementSet/Modes/NavWalkingMode.h
+++ b/Engine/Plugins/Experimental/Mover/Source/Mover/Public/DefaultMovementSet/Modes/NavWalkingMode.h
@@ -5,6 +5,7 @@
 #include "CoreMinimal.h"
 #include "Engine/EngineTypes.h"
 #include "MovementMode.h"
+#include "MoverComponent.h"
 #include "MoveLibrary/BasedMovementUtils.h"
 #include "AI/Navigation/NavigationTypes.h"
 #include "NavWalkingMode.generated.h"
@@ -33,7 +34,10 @@ enum class EOffNavMeshBehavior : uint8
 	RotateOnly = 3			   UMETA(DisplayName = "Rotate Only"),
 };
 
-
+namespace NavWalkingUtils
+{
+	void CheckClientSideNavigationSettings(const UMoverComponent* MoverComponent);
+}
 
 
 /**
