Mover:
- Virtual function for dispatching events, to give a place for event filtering (in case users want to accumulate changes from events and only dispatch final state)
- CachedNewestSimTickTimeStep also updated in the ChaosMover case (SetSimulationOutput), for parity

#Synced-CL 52217694

diff --git a/Engine/Plugins/Experimental/Mover/Source/Mover/Private/MoverComponent.cpp b/Engine/Plugins/Experimental/Mover/Source/Mover/Private/MoverComponent.cpp
--- a/Engine/Plugins/Experimental/Mover/Source/Mover/Private/MoverComponent.cpp
+++ b/Engine/Plugins/Experimental/Mover/Source/Mover/Private/MoverComponent.cpp
@@ -2786,6 +2786,13 @@ void UMoverComponent::SetSimulationOutput(const FMoverTimeStep& TimeStep, const
 {
 	CachedLastSimTickTimeStep = TimeStep;
 
+	// Keep CachedNewestSimTickTimeStep up to date for backends (e.g. ChaosMover) that drive
+	// the mover via SetSimulationOutput rather than SimulationTick.
+	if (TimeStep.ServerFrame > CachedNewestSimTickTimeStep.ServerFrame || TimeStep.BaseSimTimeMs > CachedNewestSimTickTimeStep.BaseSimTimeMs)
+	{
+		CachedNewestSimTickTimeStep = TimeStep;
+	}
+
 	CachedLastUsedInputCmd = OutputData.LastUsedInputCmd;
 
 	FMoverSyncState& BufferedSyncState = MoverSyncStateDoubleBuffer.GetWritable();
@@ -2793,6 +2800,15 @@ void UMoverComponent::SetSimulationOutput(const FMoverTimeStep& TimeStep, const
 	LastMoverDefaultSyncState = BufferedSyncState.SyncStateCollection.FindDataByType<FMoverDefaultSyncState>();
 	MoverSyncStateDoubleBuffer.Flip();
 
+	DispatchSimulationEvents(OutputData);
+
+	// This is for things like the ground info that we want to cache and interpolate but isn't part of the networked sync state.
+	// AdditionalOutputData is generic because ground info might not be useful for platforms, say, but platforms may want to store something else.
+	SetAdditionalSimulationOutput(OutputData.AdditionalOutputData);
+}
+
+void UMoverComponent::DispatchSimulationEvents(const UE::Mover::FSimulationOutputData& OutputData)
+{
 	for (const TSharedPtr<FMoverSimulationEventData>& EventData : OutputData.Events)
 	{
 		if (const FMoverSimulationEventData* Data = EventData.Get())
@@ -2800,10 +2816,6 @@ void UMoverComponent::SetSimulationOutput(const FMoverTimeStep& TimeStep, const
 			DispatchSimulationEvent(*Data);
 		}
 	}
-
-	// This is for things like the ground info that we want to cache and interpolate but isn't part of the networked sync state.
-	// AdditionalOutputData is generic because ground info might not be useful for platforms, say, but platforms may want to store something else.
-	SetAdditionalSimulationOutput(OutputData.AdditionalOutputData);
 }
 
 void UMoverComponent::DispatchSimulationEvent(const FMoverSimulationEventData& EventData)
diff --git a/Engine/Plugins/Experimental/Mover/Source/Mover/Public/MoverComponent.h b/Engine/Plugins/Experimental/Mover/Source/Mover/Public/MoverComponent.h
--- a/Engine/Plugins/Experimental/Mover/Source/Mover/Public/MoverComponent.h
+++ b/Engine/Plugins/Experimental/Mover/Source/Mover/Public/MoverComponent.h
@@ -892,6 +892,7 @@ public:
 	MOVER_API void DispatchSimulationEvent(const FMoverSimulationEventData& EventData);
 
 protected:
+	MOVER_API virtual void DispatchSimulationEvents(const UE::Mover::FSimulationOutputData& OutputData);
 	MOVER_API virtual void ProcessSimulationEvent(const FMoverSimulationEventData& EventData);
 	MOVER_API virtual void SetAdditionalSimulationOutput(const FMoverDataCollection& Data);
 	MOVER_API virtual void CheckForExternalMovement(const FMoverTickStartData& SimStartingData);
