Mover: GT-side gameplay tag change event infrastructure

New MoverGameplayTagLog.cpp/h: MOVER_TAG_LOG macro with EnableGameplayTagLog CVar;
LogGameplayTagDriftDetect log category; actor debug names include world package path
for PIE session disambiguation.

MoverSimulationTypes: FMoverGameplayTagChangeEventData with bReEmitOnResim=true;
bReEmitOnResim field on FMoverSimulationEventData; InternalGameplayTags in
FSimulationOutputData; DebugOwnerName/DebugOwnerRole debug fields; PruneEventsFromFrame;
CreateInterpolatedResult applies tag events during interpolation with LastInterpolatedTags
generation caching to handle repeated calls in the same window.

MoverComponent.h: OnGameplayTagAdded/OnGameplayTagRemoved delegates (bIsExternalTag);
GetGameplayTags(); InternalGameplayTags/AllGameplayTags; CollectGameplayTagsFromSyncState/
RebuildAllGameplayTags; DispatchGameplayTagTransitionEvents/DispatchUnexpectedTagDriftEvents
declarations.

MoverComponent.cpp: DispatchGameplayTagTransitionEvents/DispatchUnexpectedTagDriftEvents
implementations; drift detection suppressed during full resim window (bIsResimulating guard)
to prevent false positives from rollback tag state discontinuities; GetGameplayTags/
CollectGameplayTagsFromSyncState/RebuildAllGameplayTags/SetSimulationOutput/
OnSimulationRollback implementations; UseSnapshotGameplayTagsForSyncBackend CVar.

MoverBackendLiaison.h: OnSimulationRollback virtual with default empty body.

#Synced-CL 52152484

diff --git a/Engine/Plugins/Experimental/Mover/Source/Mover/Private/MoverComponent.cpp b/Engine/Plugins/Experimental/Mover/Source/Mover/Private/MoverComponent.cpp
--- a/Engine/Plugins/Experimental/Mover/Source/Mover/Private/MoverComponent.cpp
+++ b/Engine/Plugins/Experimental/Mover/Source/Mover/Private/MoverComponent.cpp
@@ -2,6 +2,7 @@
 
 
 #include "MoverComponent.h"
+#include "MoverGameplayTagLog.h"
 #include "MoverSimulationTypes.h"
 #include "MovementModeStateMachine.h"
 #include "MoverSimulation.h"
@@ -39,6 +40,17 @@
 
 #define LOCTEXT_NAMESPACE "Mover"
 
+namespace MoverCVars
+{
+	bool EnableGameplayTagLog = false;
+	FAutoConsoleVariableRef CVarEnableGameplayTagLog(
+		TEXT("mover.debug.EnableGameplayTagLog"),
+		EnableGameplayTagLog,
+		TEXT("Enable Mover gameplay tag event logging across the worker thread (WT:*) and game thread (GT:*)."),
+		ECVF_Default
+	);
+}
+
 namespace MoverComponentCVars
 {
 	static int32 WarnOnPostSimDifference = 0;
@@ -48,6 +60,16 @@ namespace MoverComponentCVars
 		TEXT("If != 0, then any differences between the sim sync state and the component locations just after movement simulation will emit warnings.\n")
 	);
 
+	// When enabled, GetGameplayTags() and HasGameplayTag() use the AllGameplayTags snapshot in the synchronous (non-async) backend,
+	// matching the behavior of the async (ChaosMover) backend. The snapshot is rebuilt each FinalizeFrame and whenever external tags
+	// are added or removed. When disabled (default), both functions perform a live walk of the active sync state instead.
+	static int32 UseSnapshotGameplayTagsForSyncBackend = 0;
+	FAutoConsoleVariableRef CVarUseSnapshotGameplayTagsForSyncBackend(
+		TEXT("mover.UseSnapshotGameplayTagsForSyncBackend"),
+		UseSnapshotGameplayTagsForSyncBackend,
+		TEXT("If != 0, GetGameplayTags() and HasGameplayTag() use the cached AllGameplayTags snapshot in the synchronous backend rather than performing a live walk.\n")
+	);
+
 } // end MoverComponentCVars
 
 
@@ -392,6 +414,20 @@ void UMoverComponent::FinalizeFrame(const FMoverSyncState* SyncState, const FMov
 		UpdateCachedFrameState(SyncState, AuxState);
 	}
 
+	// In the NP and standalone backends, tag state is not accumulated via simulation events during SimulationTick.
+	// Instead, tags are derived by walking the committed sync state once per finalized frame, and transition
+	// events (add/remove) are then fired based on the delta from the previous frame.
+	// (ChaosMover has its own FinalizeFrame with a different signature that calls SetSimulationOutput
+	// directly and never reaches this code path.)
+	{
+		FGameplayTagContainer CollectedTags;
+		CollectGameplayTagsFromSyncState(*SyncState, CollectedTags);
+		const FGameplayTagContainer PreviousGameplayTags = InternalGameplayTags;
+		InternalGameplayTags = CollectedTags;
+		RebuildAllGameplayTags();
+		DispatchGameplayTagTransitionEvents(PreviousGameplayTags, InternalGameplayTags, {}, CachedLastSimTickTimeStep);
+	}
+
 	if (PrimaryVisualComponent)
 	{
 		if (!PrimaryVisualComponent->GetRelativeTransform().Equals(BaseVisualComponentTransform))
@@ -399,7 +435,7 @@ void UMoverComponent::FinalizeFrame(const FMoverSyncState* SyncState, const FMov
 			PrimaryVisualComponent->SetRelativeTransform(BaseVisualComponentTransform);
 		}
 	}
-	
+
 	if (OnPostFinalize.IsBound())
 	{
 		OnPostFinalize.Broadcast(MoverSyncStateDoubleBuffer.GetReadable(), CachedLastAuxState);
@@ -842,6 +878,11 @@ const FMovementModifierBase* UMoverComponent::FindMovementModifierByType(const U
 
 bool UMoverComponent::HasGameplayTag(FGameplayTag TagToFind, bool bExactMatch) const
 {
+	if ((BackendLiaisonComp && BackendLiaisonComp->IsAsync()) || MoverComponentCVars::UseSnapshotGameplayTagsForSyncBackend != 0)
+	{
+		return bExactMatch ? AllGameplayTags.HasTagExact(TagToFind) : AllGameplayTags.HasTag(TagToFind);
+	}
+
 	return HasGameplayTagInState(MoverSyncStateDoubleBuffer.GetReadable(), TagToFind, bExactMatch);
 }
 
@@ -899,22 +940,68 @@ bool UMoverComponent::HasGameplayTagInState(const FMoverSyncState& SyncState, FG
 
 void UMoverComponent::AddGameplayTag(FGameplayTag TagToAdd)
 {
-	ExternalGameplayTags.AddTag(TagToAdd);
+	if (!ExternalGameplayTags.HasTagExact(TagToAdd))
+	{
+		ExternalGameplayTags.AddTag(TagToAdd);
+		RebuildAllGameplayTags();
+
+		if (InternalGameplayTags.HasTagExact(TagToAdd))
+		{
+			UE_LOGF(LogMover, Verbose, "External gameplay tag '%ls' added but tag is already active via the simulation - event suppressed to avoid a duplicate added event.", *TagToAdd.ToString());
+		}
+		else
+		{
+			FMoverEventContext EventContext;
+			if (BackendLiaisonComp)
+			{
+				EventContext.EventTimeMs = BackendLiaisonComp->GetCurrentSimTimeMs();
+				EventContext.ServerFrame = BackendLiaisonComp->GetCurrentSimFrame();
+			}
+
+			OnGameplayTagAdded.Broadcast(TagToAdd, EventContext, /* bIsExternalTag = */ true);
+		}
+	}
 }
 
 void UMoverComponent::AddGameplayTags(const FGameplayTagContainer& TagsToAdd)
 {
-	ExternalGameplayTags.AppendTags(TagsToAdd);
+	for (const FGameplayTag& Tag : TagsToAdd)
+	{
+		AddGameplayTag(Tag);
+	}
 }
 
 void UMoverComponent::RemoveGameplayTag(FGameplayTag TagToRemove)
 {
-	ExternalGameplayTags.RemoveTag(TagToRemove);
+	if (ExternalGameplayTags.HasTagExact(TagToRemove))
+	{
+		ExternalGameplayTags.RemoveTag(TagToRemove);
+		RebuildAllGameplayTags();
+
+		if (InternalGameplayTags.HasTagExact(TagToRemove))
+		{
+			UE_LOGF(LogMover, Verbose, "External gameplay tag '%ls' removed but tag is still active via the simulation - event suppressed to avoid a duplicate removed event.", *TagToRemove.ToString());
+		}
+		else
+		{
+			FMoverEventContext EventContext;
+			if (BackendLiaisonComp)
+			{
+				EventContext.EventTimeMs = BackendLiaisonComp->GetCurrentSimTimeMs();
+				EventContext.ServerFrame = BackendLiaisonComp->GetCurrentSimFrame();
+			}
+
+			OnGameplayTagRemoved.Broadcast(TagToRemove, EventContext, /* bIsExternalTag = */ true);
+		}
+	}
 }
 
 void UMoverComponent::RemoveGameplayTags(const FGameplayTagContainer& TagsToRemove)
 {
-	ExternalGameplayTags.RemoveTags(TagsToRemove);
+	for (const FGameplayTag& Tag : TagsToRemove)
+	{
+		RemoveGameplayTag(Tag);
+	}
 }
 
 void UMoverComponent::PreSimulationTick(const FMoverTimeStep& TimeStep, const FMoverInputCmdContext& InputCmd)
@@ -986,7 +1073,7 @@ void UMoverComponent::SetFrameStateFromContext(const FMoverSyncState* SyncState,
 		if (IsUsingDeferredGroupMovement())
 		{
 			// Signal to the USceneComponent that we are moving that this should be in a grouped update
-			// and not apply changes on the physics thread immediately
+			// and not apply changes on the worker thread immediately
 			FScopedMovementUpdate MovementUpdate(
 				UpdatedComponent,
 				EScopedUpdate::DeferredGroupUpdates,
@@ -1148,27 +1235,41 @@ void UMoverComponent::OnSimulationPreRollback(const FMoverSyncState* InvalidSync
 	ModeFSM->OnSimulationPreRollback(InvalidSyncState, SyncState, InvalidAuxState, AuxState, NewBaseTimeStep);
 }
 
-
 void UMoverComponent::OnSimulationRollback(const FMoverSyncState* SyncState, const FMoverAuxStateContext* AuxState, const FMoverTimeStep& NewBaseTimeStep)
 {
+	if (!ensure(SyncState))
+	{
+		return;
+	}
+
+	// Mover Gameplay Tag Logging
+	const FString DebugActorName = FString::Printf(TEXT("%s[%s]"), *GetNameSafe(GetOwner()), GetOwner() && GetOwner()->GetWorld() ? *GetOwner()->GetWorld()->GetOutermost()->GetName() : TEXT("null"));
+	const FString DebugActorRole = GetOwner() ? UEnum::GetValueAsString(TEXT("Engine.ENetRole"), GetOwner()->GetLocalRole()) : FString(TEXT("null"));
+	MOVER_TAG_LOG(LogMover, "[GT:RollbackNotify] Actor=%s Role=%s Frame=%d IsResim=%d IsRollback=%d | FromFrame=%d FromTimeMs=%.3f",
+		*DebugActorName, *DebugActorRole, NewBaseTimeStep.ServerFrame, (int32)NewBaseTimeStep.bIsResimulating, (int32)true,
+		CachedLastSimTickTimeStep.ServerFrame, CachedLastSimTickTimeStep.BaseSimTimeMs);
+
 	SimBlackboard->Invalidate(EInvalidationReason::Rollback);
 
 	RollbackBlackboard->BeginRollback(NewBaseTimeStep);
 
+	if (BackendLiaisonComp)
+	{
+		BackendLiaisonComp->OnSimulationRollback(*SyncState, NewBaseTimeStep, CachedLastSimTickTimeStep);
+	}
+
 	ModeFSM->OnSimulationRollback(SyncState, AuxState, NewBaseTimeStep);
 
 	RollbackBlackboard->EndRollback();
 	bHasJustRolledBack = true;
 }
 
-
 void UMoverComponent::ProcessFirstSimTickAfterRollback(const FMoverTimeStep& TimeStep)
 {
 	OnPostSimulationRollback.Broadcast(TimeStep, CachedLastSimTickTimeStep);
 	bHasJustRolledBack = false;
 }
 
-
 #if WITH_EDITOR
 
 void UMoverComponent::PreSave(FObjectPreSaveContext ObjectSaveContext)
@@ -2800,6 +2901,12 @@ void UMoverComponent::SetSimulationOutput(const FMoverTimeStep& TimeStep, const
 	LastMoverDefaultSyncState = BufferedSyncState.SyncStateCollection.FindDataByType<FMoverDefaultSyncState>();
 	MoverSyncStateDoubleBuffer.Flip();
 
+	// Replace the tag snapshot wholesale before dispatching events, so GetGameplayTags() is always consistent.
+	// Preserve the previous snapshot to detect drift, i.e. tags that changed due to a correction and have no corresponding event.
+	const FGameplayTagContainer PreviousGameplayTags = InternalGameplayTags;
+	InternalGameplayTags = OutputData.InternalGameplayTags;
+	RebuildAllGameplayTags();
+
 	DispatchSimulationEvents(OutputData);
 
 	// This is for things like the ground info that we want to cache and interpolate but isn't part of the networked sync state.
@@ -2809,6 +2916,8 @@ void UMoverComponent::SetSimulationOutput(const FMoverTimeStep& TimeStep, const
 
 void UMoverComponent::DispatchSimulationEvents(const UE::Mover::FSimulationOutputData& OutputData)
 {
+	DispatchUnexpectedTagDriftEvents(PreviousGameplayTags, InternalGameplayTags, OutputData.Events, TimeStep);
+
 	for (const TSharedPtr<FMoverSimulationEventData>& EventData : OutputData.Events)
 	{
 		if (const FMoverSimulationEventData* Data = EventData.Get())
@@ -2818,6 +2927,128 @@ void UMoverComponent::DispatchSimulationEvents(const UE::Mover::FSimulationOutpu
 	}
 }
 
+void UMoverComponent::DispatchGameplayTagTransitionEvents(const FGameplayTagContainer& PreviousTags, const FGameplayTagContainer& NewTags, const TArray<TSharedPtr<FMoverSimulationEventData>>& EventsToSkip, const FMoverTimeStep& TimeStep)
+{
+	const FMoverEventContext EventContext(TimeStep);
+
+	for (const FGameplayTag& Tag : PreviousTags)
+	{
+		if (!NewTags.HasTagExact(Tag))
+		{
+			const bool bAlreadyCovered = EventsToSkip.ContainsByPredicate([&Tag](const TSharedPtr<FMoverSimulationEventData>& E)
+			{
+				const FMoverGameplayTagChangeEventData* TagEvent = E.IsValid() ? E->CastTo<FMoverGameplayTagChangeEventData>() : nullptr;
+				return TagEvent && !TagEvent->bWasAdded && TagEvent->Tag == Tag;
+			});
+
+			if (!bAlreadyCovered)
+			{
+				if (ExternalGameplayTags.HasTagExact(Tag))
+				{
+					UE_LOGF(LogMover, Verbose, "Gameplay tag '%ls' removed event suppressed on frame %d - tag is also held externally.", *Tag.ToString(), TimeStep.ServerFrame);
+				}
+				else
+				{
+					OnGameplayTagRemoved.Broadcast(Tag, EventContext, /* bIsExternalTag = */ false);
+				}
+			}
+		}
+	}
+
+	for (const FGameplayTag& Tag : NewTags)
+	{
+		if (!PreviousTags.HasTagExact(Tag))
+		{
+			const bool bAlreadyCovered = EventsToSkip.ContainsByPredicate([&Tag](const TSharedPtr<FMoverSimulationEventData>& E)
+			{
+				const FMoverGameplayTagChangeEventData* TagEvent = E.IsValid() ? E->CastTo<FMoverGameplayTagChangeEventData>() : nullptr;
+				return TagEvent && TagEvent->bWasAdded && TagEvent->Tag == Tag;
+			});
+
+			if (!bAlreadyCovered)
+			{
+				if (ExternalGameplayTags.HasTagExact(Tag))
+				{
+					UE_LOGF(LogMover, Verbose, "Gameplay tag '%ls' added event suppressed on frame %d - tag is also held externally.", *Tag.ToString(), TimeStep.ServerFrame);
+				}
+				else
+				{
+					OnGameplayTagAdded.Broadcast(Tag, EventContext, /* bIsExternalTag = */ false);
+				}
+			}
+		}
+	}
+}
+
+void UMoverComponent::DispatchUnexpectedTagDriftEvents(const FGameplayTagContainer& PreviousTags, const FGameplayTagContainer& NewTags, const TArray<TSharedPtr<FMoverSimulationEventData>>& Events, const FMoverTimeStep& TimeStep)
+{
+	const AActor* DebugOwner = GetOwner();
+	const UWorld* DebugWorld = DebugOwner ? DebugOwner->GetWorld() : nullptr;
+	const FString DebugActorName = FString::Printf(TEXT("%s[%s]"), *GetNameSafe(DebugOwner), DebugWorld ? *DebugWorld->GetOutermost()->GetName() : TEXT("null"));
+	const FString DebugActorRole = DebugOwner ? UEnum::GetValueAsString(TEXT("Engine.ENetRole"), DebugOwner->GetLocalRole()) : FString(TEXT("null"));
+
+	// Every change between PreviousTags and NewTags should ideally be accounted for by a corresponding
+	// FMoverGameplayTagChangeEventData in the simulation output events. These two sources of truth -
+	// the PT snapshot and the PT event log - are both produced by the same simulation tick and should
+	// always agree. Any divergence is flagged as PT/GT drift via an ensure and a warning log.
+	//
+	// Exception: during resimulation the GT's tag history may be discontinuous. Events that were
+	// in flight before the rollback are pruned, and the PT may not re-emit them if its own
+	// LastKnownGameplayTags was already at the post-transition state when the rollback fired
+	// (e.g. a prediction that reached IsOnGround is rolled back to a frame where the authority
+	// also has IsOnGround, so the PT sees no diff and produces no events). The GT then observes
+	// a tag change with no covering event, which is an expected consequence of the rollback
+	// correction rather than a real PT/GT mismatch.
+	//
+	// In all cases we recover by emitting the missing events via DispatchGameplayTagTransitionEvents.
+	if (!TimeStep.bIsResimulating)
+	{
+		for (const FGameplayTag& Tag : PreviousTags)
+		{
+			if (!NewTags.HasTagExact(Tag))
+			{
+				const bool bAlreadyCovered = Events.ContainsByPredicate([&Tag](const TSharedPtr<FMoverSimulationEventData>& E)
+				{
+					const FMoverGameplayTagChangeEventData* TagEvent = E.IsValid() ? E->CastTo<FMoverGameplayTagChangeEventData>() : nullptr;
+					return TagEvent && !TagEvent->bWasAdded && TagEvent->Tag == Tag;
+				});
+
+				if (!bAlreadyCovered)
+				{
+					ensureMsgf(false, TEXT("Gameplay tag '%s' was removed on frame %d without a corresponding simulation event - possible PT/GT drift [Actor=%s Role=%s]"), *Tag.ToString(), TimeStep.ServerFrame, *DebugActorName, *DebugActorRole);
+					UE_LOGF(LogMover, Warning, "Gameplay tag '%ls' was removed on frame %d without a corresponding simulation event - possible PT/GT drift [Actor=%ls Role=%ls]", *Tag.ToString(), TimeStep.ServerFrame, *DebugActorName, *DebugActorRole);
+
+					// Mover Gameplay Tag Logging
+					LogGameplayTagDriftDetect(DebugActorName, *DebugActorRole, TimeStep, Tag, /* bWasAdded= */ false, Events, PreviousTags, NewTags);
+				}
+			}
+		}
+
+		for (const FGameplayTag& Tag : NewTags)
+		{
+			if (!PreviousTags.HasTagExact(Tag))
+			{
+				const bool bAlreadyCovered = Events.ContainsByPredicate([&Tag](const TSharedPtr<FMoverSimulationEventData>& E)
+				{
+					const FMoverGameplayTagChangeEventData* TagEvent = E.IsValid() ? E->CastTo<FMoverGameplayTagChangeEventData>() : nullptr;
+					return TagEvent && TagEvent->bWasAdded && TagEvent->Tag == Tag;
+				});
+
+				if (!bAlreadyCovered)
+				{
+					ensureMsgf(false, TEXT("Gameplay tag '%s' was added on frame %d without a corresponding simulation event - possible PT/GT drift [Actor=%s Role=%s]"), *Tag.ToString(), TimeStep.ServerFrame, *DebugActorName, *DebugActorRole);
+					UE_LOGF(LogMover, Warning, "Gameplay tag '%ls' was added on frame %d without a corresponding simulation event - possible PT/GT drift [Actor=%ls Role=%ls]", *Tag.ToString(), TimeStep.ServerFrame, *DebugActorName, *DebugActorRole);
+
+					// Mover Gameplay Tag Logging
+					LogGameplayTagDriftDetect(DebugActorName, *DebugActorRole, TimeStep, Tag, /* bWasAdded= */ true, Events, PreviousTags, NewTags);
+				}
+			}
+		}
+	}
+
+	DispatchGameplayTagTransitionEvents(PreviousTags, NewTags, Events, TimeStep);
+}
+
 void UMoverComponent::DispatchSimulationEvent(const FMoverSimulationEventData& EventData)
 {
 	// This gives the event a callback when it is processed on the game thread
@@ -2874,6 +3105,74 @@ void UMoverComponent::ProcessSimulationEvent(const FMoverSimulationEventData& Ev
 	{
 		OnTeleportFailed.Broadcast(TeleportFailedEventData->FromLocation, TeleportFailedEventData->FromRotation, TeleportFailedEventData->ToLocation, TeleportFailedEventData->ToRotation, TeleportFailedEventData->TeleportFailureReason);
 	}
+	else if (const FMoverGameplayTagChangeEventData* TagStateChangeData = EventData.CastTo<FMoverGameplayTagChangeEventData>())
+	{
+		if (TagStateChangeData->bWasAdded)
+		{
+			if (ExternalGameplayTags.HasTagExact(TagStateChangeData->Tag))
+			{
+				UE_LOGF(LogMover, Verbose, "Simulation gameplay tag '%ls' added but tag is already active externally - event suppressed to avoid a duplicate added event.", *TagStateChangeData->Tag.ToString());
+			}
+			else
+			{
+				OnGameplayTagAdded.Broadcast(TagStateChangeData->Tag, TagStateChangeData->Context, TagStateChangeData->bIsExternalTag);
+			}
+		}
+		else
+		{
+			if (ExternalGameplayTags.HasTagExact(TagStateChangeData->Tag))
+			{
+				UE_LOGF(LogMover, Verbose, "Simulation gameplay tag '%ls' removed but tag is still active externally - event suppressed to avoid a duplicate removed event.", *TagStateChangeData->Tag.ToString());
+			}
+			else
+			{
+				OnGameplayTagRemoved.Broadcast(TagStateChangeData->Tag, TagStateChangeData->Context, TagStateChangeData->bIsExternalTag);
+			}
+		}
+	}
+}
+
+FGameplayTagContainer UMoverComponent::GetGameplayTags() const
+{
+	if ((BackendLiaisonComp && BackendLiaisonComp->IsAsync()) || MoverComponentCVars::UseSnapshotGameplayTagsForSyncBackend != 0)
+	{
+		return AllGameplayTags;
+	}
+
+	FGameplayTagContainer LiveTags;
+	CollectGameplayTagsFromSyncState(MoverSyncStateDoubleBuffer.GetReadable(), LiveTags);
+	LiveTags.AppendTags(ExternalGameplayTags);
+	return LiveTags;
+}
+
+void UMoverComponent::RebuildAllGameplayTags()
+{
+	AllGameplayTags = InternalGameplayTags;
+	AllGameplayTags.AppendTags(ExternalGameplayTags);
+}
+
+void UMoverComponent::CollectGameplayTagsFromSyncState(const FMoverSyncState& SyncState, FGameplayTagContainer& OutTags) const
+{
+	if (const UBaseMovementMode* ActiveMovementMode = FindMovementModeByName(SyncState.MovementMode))
+	{
+		ActiveMovementMode->GetGameplayTags(OutTags);
+	}
+
+	for (auto ModifierIt = SyncState.MovementModifiers.GetActiveModifiersIterator(); ModifierIt; ++ModifierIt)
+	{
+		if (const TSharedPtr<FMovementModifierBase> Modifier = *ModifierIt)
+		{
+			Modifier->GetGameplayTags(OutTags);
+		}
+	}
+
+	for (const TSharedPtr<FLayeredMoveBase>& LayeredMove : SyncState.LayeredMoves.GetActiveMoves())
+	{
+		if (LayeredMove.IsValid())
+		{
+			LayeredMove->GetGameplayTags(OutTags);
+		}
+	}
 }
 
 void UMoverComponent::SetAdditionalSimulationOutput(const FMoverDataCollection& Data)
diff --git a/Engine/Plugins/Experimental/Mover/Source/Mover/Private/MoverGameplayTagLog.cpp b/Engine/Plugins/Experimental/Mover/Source/Mover/Private/MoverGameplayTagLog.cpp
new file mode 100644
--- /dev/null
+++ b/Engine/Plugins/Experimental/Mover/Source/Mover/Private/MoverGameplayTagLog.cpp
@@ -0,0 +1,51 @@
+// Copyright Epic Games, Inc. All Rights Reserved.
+
+#include "MoverGameplayTagLog.h"
+#include "MoverLog.h"
+
+void LogGameplayTagDriftDetect(
+	const FString& ActorName,
+	const TCHAR* RoleStr,
+	const FMoverTimeStep& TimeStep,
+	const FGameplayTag& Tag,
+	bool bWasAdded,
+	const TArray<TSharedPtr<FMoverSimulationEventData>>& Events,
+	const FGameplayTagContainer& PreviousTags,
+	const FGameplayTagContainer& NewTags)
+{
+	if (!MoverCVars::EnableGameplayTagLog)
+	{
+		return;
+	}
+
+	UE_LOG(LogMover, Log, TEXT("[GameplayTagLog][GT:DriftDetect] Actor=%s Role=%s Frame=%d TimeMs=%.3f IsResim=%d IsFirstResim=%d Tag=%s Change=%s Status=MissingEvent EventQueueSize=%d PreviousTags=[%s] NewTags=[%s]"),
+		*ActorName, RoleStr,
+		TimeStep.ServerFrame, TimeStep.BaseSimTimeMs,
+		(int32)TimeStep.bIsResimulating, (int32)TimeStep.bIsFirstResimFrame,
+		*Tag.ToString(), bWasAdded ? TEXT("Added") : TEXT("Removed"),
+		Events.Num(),
+		*PreviousTags.ToStringSimple(), *NewTags.ToStringSimple());
+
+	for (const TSharedPtr<FMoverSimulationEventData>& E : Events)
+	{
+		if (!E.IsValid())
+		{
+			continue;
+		}
+
+		const FMoverGameplayTagChangeEventData* TagEvent = E->CastTo<FMoverGameplayTagChangeEventData>();
+		if (TagEvent)
+		{
+			UE_LOG(LogMover, Log, TEXT("[GameplayTagLog][GT:DriftDetect:Event] Actor=%s Role=%s Frame=%d TimeMs=%.3f IsResim=%d IsRollback=%d Tag=%s Change=%s"),
+				*ActorName, RoleStr, E->Context.ServerFrame, E->Context.EventTimeMs,
+				(int32)E->Context.bIsDuringResimulation, (int32)E->Context.bIsCausedByRollback,
+				*TagEvent->Tag.ToString(), TagEvent->bWasAdded ? TEXT("Added") : TEXT("Removed"));
+		}
+		else
+		{
+			UE_LOG(LogMover, Log, TEXT("[GameplayTagLog][GT:DriftDetect:Event] Actor=%s Role=%s Frame=%d TimeMs=%.3f IsResim=%d IsRollback=%d"),
+				*ActorName, RoleStr, E->Context.ServerFrame, E->Context.EventTimeMs,
+				(int32)E->Context.bIsDuringResimulation, (int32)E->Context.bIsCausedByRollback);
+		}
+	}
+}
diff --git a/Engine/Plugins/Experimental/Mover/Source/Mover/Private/MoverSimulationTypes.cpp b/Engine/Plugins/Experimental/Mover/Source/Mover/Private/MoverSimulationTypes.cpp
--- a/Engine/Plugins/Experimental/Mover/Source/Mover/Private/MoverSimulationTypes.cpp
+++ b/Engine/Plugins/Experimental/Mover/Source/Mover/Private/MoverSimulationTypes.cpp
@@ -1,6 +1,7 @@
 // Copyright Epic Games, Inc. All Rights Reserved.
 
 #include "MoverSimulationTypes.h"
+#include "MoverGameplayTagLog.h"
 #include "Physics/Experimental/PhysScene_Chaos.h"
 #include "PBDRigidsSolver.h"
 #include "GameFramework/GameStateBase.h"
@@ -41,6 +42,7 @@ namespace UE::Mover
 		LastUsedInputCmd.Reset();
 		AdditionalOutputData.Empty();
 		Events.Empty();
+		InternalGameplayTags.Reset();
 		bIsValid = false;
 	}
 
@@ -49,6 +51,8 @@ namespace UE::Mover
 		SyncState.Interpolate(&From.SyncState, &To.SyncState, Alpha);
 		LastUsedInputCmd.Interpolate(&From.LastUsedInputCmd, &To.LastUsedInputCmd, Alpha);
 		AdditionalOutputData.Interpolate(From.AdditionalOutputData, To.AdditionalOutputData, Alpha);
+		InternalGameplayTags = From.InternalGameplayTags;
+		Events.Empty();
 	}
 
 	void FSimulationOutputRecord::FData::Reset()
@@ -68,6 +72,7 @@ namespace UE::Mover
 			Events.Add(Event);
 		}
 		Data[CurrentIndex].SimOutputData.Events.Empty();
+		NextInterpolatedGeneration++;
 	}
 
 	const FSimulationOutputData& FSimulationOutputRecord::GetLatest() const
@@ -75,12 +80,38 @@ namespace UE::Mover
 		return Data[CurrentIndex].SimOutputData;
 	}
 
+	void FSimulationOutputRecord::LogTagInterpolate(uint8 PrevIndex, double AtBaseTimeMs, bool bIsInterpolated) const
+	{
+		if (!MoverCVars::EnableGameplayTagLog)
+		{
+			return;
+		}
+
+		MOVER_TAG_LOG(LogMover, "[GT:Interpolate] Actor=%s Role=%s Frame=%d..%d TimeMs=%.3f IsInterpolated=%d PrevIsResim=%d CurrIsResim=%d | PendingEvents=%d",
+			*DebugOwnerName, *DebugOwnerRole, Data[PrevIndex].TimeStep.ServerFrame, Data[CurrentIndex].TimeStep.ServerFrame, AtBaseTimeMs, (int32)bIsInterpolated, (int32)Data[PrevIndex].TimeStep.bIsResimulating, (int32)Data[CurrentIndex].TimeStep.bIsResimulating, Events.Num());
+
+		for (const TSharedPtr<FMoverSimulationEventData>& PendingEvt : Events)
+		{
+			if (!PendingEvt.IsValid())
+			{
+				continue;
+			}
+			if (const FMoverGameplayTagChangeEventData* TagEvt = PendingEvt->CastTo<FMoverGameplayTagChangeEventData>())
+			{
+				MOVER_TAG_LOG(LogMover, "[GT:Interpolate:PendingEvent] Actor=%s Role=%s Frame=%d TimeMs=%.3f IsResim=%d IsRollback=%d Tag=%s Change=%s",
+					*DebugOwnerName, *DebugOwnerRole, PendingEvt->Context.ServerFrame, PendingEvt->Context.EventTimeMs, (int32)PendingEvt->Context.bIsDuringResimulation, (int32)PendingEvt->Context.bIsCausedByRollback, *TagEvt->Tag.ToString(), TagEvt->bWasAdded ? TEXT("Added") : TEXT("Removed"));
+			}
+		}
+	}
+
 	void FSimulationOutputRecord::CreateInterpolatedResult(double AtBaseTimeMs, FMoverTimeStep& OutTimeStep, FSimulationOutputData& OutData)
 	{
 		const uint8 PrevIndex = (CurrentIndex + 1) % 2;
 		const double PrevTimeMs = Data[PrevIndex].TimeStep.BaseSimTimeMs;
 		const double CurrTimeMs = Data[CurrentIndex].TimeStep.BaseSimTimeMs;
 
+		bool bIsInterpolated = false;
+
 		if (FMath::IsNearlyEqual(PrevTimeMs, CurrTimeMs) || (AtBaseTimeMs >= CurrTimeMs))
 		{
 			OutData = Data[CurrentIndex].SimOutputData;
@@ -96,14 +127,66 @@ namespace UE::Mover
 			const float Alpha = FMath::Clamp((AtBaseTimeMs - PrevTimeMs) / (CurrTimeMs - PrevTimeMs), 0.0f, 1.0f);
 			OutData.Interpolate(Data[PrevIndex].SimOutputData, Data[CurrentIndex].SimOutputData, Alpha, AtBaseTimeMs);
 			OutTimeStep = Data[PrevIndex].TimeStep;
+			bIsInterpolated = true;
 		}
 
 		OutTimeStep.BaseSimTimeMs = AtBaseTimeMs;
 
+		// Mover Gameplay Tag Logging
+		LogTagInterpolate(PrevIndex, AtBaseTimeMs, bIsInterpolated);
+
+		// If we are re-entering the same interpolation window (same NextInterpolatedGeneration -- meaning
+		// Add() has not been called since the last visit), start from the accumulated
+		// post-event tag state rather than the original From snapshot. This prevents
+		// SetSimulationOutput from seeing a stale InternalGameplayTags (the From state)
+		// after events have already been consumed, which would trigger spurious drift
+		// detection for tags that were correctly dispatched in a prior call.
+		// Otherwise (first visit, or bIsInterpolated is false), OutData.InternalGameplayTags
+		// already holds Data[PrevIndex].SimOutputData.InternalGameplayTags (the From snapshot),
+		// set by FSimulationOutputData::Interpolate, and events are applied on top of it below.
+		if (bIsInterpolated && LastInterpolatedGeneration == NextInterpolatedGeneration)
+		{
+			OutData.InternalGameplayTags = LastInterpolatedTags;
+		}
+
 		for (TSharedPtr<FMoverSimulationEventData>& Event : Events)
 		{
+			if (!Event.IsValid())
+			{
+				continue;
+			}
+
 			if (Event->Context.EventTimeMs <= AtBaseTimeMs)
 			{
+				if (bIsInterpolated && Event->Context.EventTimeMs > PrevTimeMs)
+				{
+					if (const FMoverGameplayTagChangeEventData* TagStateChangeEvent = Event->CastTo<FMoverGameplayTagChangeEventData>())
+					{
+						if (TagStateChangeEvent->bWasAdded)
+						{
+							UE_CLOG(OutData.InternalGameplayTags.HasTagExact(TagStateChangeEvent->Tag), LogMover, Verbose,
+								TEXT("Interpolated gameplay tag add for '%s' but tag is already present in the From snapshot - possible drift between recorded tag snapshot and event log"),
+								*TagStateChangeEvent->Tag.ToString());
+
+							OutData.InternalGameplayTags.AddTag(TagStateChangeEvent->Tag);
+						}
+						else
+						{
+							UE_CLOG(!OutData.InternalGameplayTags.HasTagExact(TagStateChangeEvent->Tag), LogMover, Verbose,
+								TEXT("Interpolated gameplay tag remove for '%s' but tag is not present in the From snapshot - possible drift between recorded tag snapshot and event log"),
+								*TagStateChangeEvent->Tag.ToString());
+
+							OutData.InternalGameplayTags.RemoveTag(TagStateChangeEvent->Tag);
+						}
+					}
+				}
+
+				if (const FMoverGameplayTagChangeEventData* TagDispatch = Event->CastTo<FMoverGameplayTagChangeEventData>())
+				{
+					MOVER_TAG_LOG(LogMover, "[GT:Interpolate:TagEventOutput] Actor=%s Role=%s Frame=%d TimeMs=%.3f Tag=%s Change=%s InternalTags=[%s]",
+						*DebugOwnerName, *DebugOwnerRole, Event->Context.ServerFrame, Event->Context.EventTimeMs, *TagDispatch->Tag.ToString(), TagDispatch->bWasAdded ? TEXT("Added") : TEXT("Removed"), *OutData.InternalGameplayTags.ToStringSimple());
+				}
+
 				OutData.Events.Add(MoveTemp(Event));
 			}
 		}
@@ -111,6 +194,15 @@ namespace UE::Mover
 		Events.RemoveAll([](const TSharedPtr<FMoverSimulationEventData> Event) {
 			return !Event.IsValid();
 			});
+
+		if (bIsInterpolated)
+		{
+			MOVER_TAG_LOG(LogMover, "[GT:Interpolate:TagState] Actor=%s Role=%s Frame=%d..%d TimeMs=%.3f PrevTags=[%s] InternalTags=[%s]",
+				*DebugOwnerName, *DebugOwnerRole, Data[PrevIndex].TimeStep.ServerFrame, Data[CurrentIndex].TimeStep.ServerFrame, AtBaseTimeMs, *LastInterpolatedTags.ToStringSimple(), *OutData.InternalGameplayTags.ToStringSimple());
+
+			LastInterpolatedGeneration = NextInterpolatedGeneration;
+			LastInterpolatedTags = OutData.InternalGameplayTags;
+		}
 	}
 
 	void FSimulationOutputRecord::Clear()
@@ -119,6 +211,16 @@ namespace UE::Mover
 		Data[0].Reset();
 		Data[1].Reset();
 		Events.Empty();
+		LastInterpolatedGeneration = INDEX_NONE;
+		LastInterpolatedTags.Reset();
+	}
+
+	void FSimulationOutputRecord::PruneEventsFromFrame(int32 FromServerFrame)
+	{
+		Events.RemoveAll([FromServerFrame](const TSharedPtr<FMoverSimulationEventData>& Event)
+		{
+			return Event.IsValid() && Event->Context.ServerFrame >= FromServerFrame;
+		});
 	}
 
 } // namespace UE::Mover
diff --git a/Engine/Plugins/Experimental/Mover/Source/Mover/Public/Backends/MoverBackendLiaison.h b/Engine/Plugins/Experimental/Mover/Source/Mover/Public/Backends/MoverBackendLiaison.h
--- a/Engine/Plugins/Experimental/Mover/Source/Mover/Public/Backends/MoverBackendLiaison.h
+++ b/Engine/Plugins/Experimental/Mover/Source/Mover/Public/Backends/MoverBackendLiaison.h
@@ -57,6 +57,13 @@ public:
 	virtual bool ReadPrevPresentationSyncState(OUT FMoverSyncState& OutSyncState) { return false; }
 	virtual bool WritePrevPresentationSyncState(const FMoverSyncState& SyncStateToWrite) { return false; }
 
+	/** Called on the game thread when a rollback occurs, after simulation state has been restored.
+	 *  NewBaseTimeStep is the frame we are about to re-simulate from.
+	 *  PreRollbackTimeStep is the last forward-simulated frame before the rollback -- together they define
+	 *  the window of time that was invalidated (mirrors the FMover_OnPostSimRollback delegate).
+	 *  NewSyncState is the sync state at NewBaseTimeStep, i.e., the state the simulation has been restored to. */
+	virtual void OnSimulationRollback(const FMoverSyncState& NewSyncState, const FMoverTimeStep& NewBaseTimeStep, const FMoverTimeStep& PreRollbackTimeStep) {}
+
 #if WITH_EDITOR
 	virtual EDataValidationResult ValidateData(FDataValidationContext& Context, const UMoverComponent& ValidationMoverComp) const { return EDataValidationResult::Valid; }
 #endif // WITH_EDITOR
diff --git a/Engine/Plugins/Experimental/Mover/Source/Mover/Public/Backends/MoverNetworkPhysicsLiaisonBase.h b/Engine/Plugins/Experimental/Mover/Source/Mover/Public/Backends/MoverNetworkPhysicsLiaisonBase.h
--- a/Engine/Plugins/Experimental/Mover/Source/Mover/Public/Backends/MoverNetworkPhysicsLiaisonBase.h
+++ b/Engine/Plugins/Experimental/Mover/Source/Mover/Public/Backends/MoverNetworkPhysicsLiaisonBase.h
@@ -181,9 +181,9 @@ public:
  * WARNING - This class will be removed. Please use UChaosMoverBackend instead
  *
  * Base class for liaison components to act as a middleman between an actor's Mover component and the Chaos physics network prediction system to move
- * the Mover's updated component on the physics thread (PT).
- * 
- * This is accomplished by registering a FPhysicsMoverManagerAsyncCallback and exposing the game thread (GT) and PT methods for the TSimCallback to
+ * the Mover's updated component on a worker thread (WT).
+ *
+ * This is accomplished by registering a FPhysicsMoverManagerAsyncCallback and exposing the game thread (GT) and WT methods for the TSimCallback to
  * virtual methods on this liaison.
  */
 UCLASS(MinimalAPI, Abstract, Within = MoverComponent)
@@ -228,13 +228,13 @@ public:
 	UE_API virtual void ConsumeOutput_External(const FPhysicsMoverAsyncOutput& Output, double OutputTimeInSeconds);
 	UE_API virtual void PostPhysicsUpdate_External();
 
-	// "Internal" = Called on the Physics Thread (PT)
+	// "Internal" = Called on a Worker Thread (WT)
 	UE_API void ProcessInputs_Internal(int32 PhysicsStep, float DeltaTime, const FPhysicsMoverAsyncInput& Input) const;
 	UE_API void PreSimulate_Internal(const FPhysicsMoverSimulationTickParams& TickParams, const FPhysicsMoverAsyncInput& Input, OUT FPhysicsMoverAsyncOutput& Output) const;
 	UE_API virtual void OnContactModification_Internal(const FPhysicsMoverAsyncInput& Input, Chaos::FCollisionContactModifier& Modifier) const;
 	
 protected:
-	/** Catchall to verify that the state of this object is valid (called from both the PT and the GT) */
+	/** Catchall to verify that the state of this object is valid (called from both WT and the GT) */
 	UE_API virtual bool HasValidState() const;
 	
 	UE_API virtual void PerformProduceInput_External(float DeltaTime, OUT FPhysicsMoverAsyncInput& Input);
@@ -245,7 +245,7 @@ protected:
 	 */
 	UE_API virtual bool CanProcessInputs_Internal(const FPhysicsMoverAsyncInput& Input) const;
 
-	/** Actually do the ProcessInputs work on the PT. CanProcessInputs_Internal is guaranteed to have returned true if this is called. */
+	/** Actually do the ProcessInputs work on a WT. CanProcessInputs_Internal is guaranteed to have returned true if this is called. */
 	UE_API virtual void PerformProcessInputs_Internal(int32 PhysicsStep, float DeltaTime, const FPhysicsMoverAsyncInput& Input) const;
 
 	/** 
@@ -255,7 +255,7 @@ protected:
 	 */
 	UE_API virtual bool CanSimulate_Internal(const FPhysicsMoverSimulationTickParams& TickParams, const FPhysicsMoverAsyncInput& Input) const;
 	
-	/** Actually do the PreSimulate work on the PT. CanSimulate_Internal is guaranteed to have returned true if this is called. */
+	/** Actually do the PreSimulate work on the WT. CanSimulate_Internal is guaranteed to have returned true if this is called. */
 	UE_API virtual void PerformPreSimulate_Internal(const FPhysicsMoverSimulationTickParams& TickParams, const FPhysicsMoverAsyncInput& Input, OUT FPhysicsMoverAsyncOutput& Output) const;
 	
 	/** Override in child classes to return the chaos particle for the component the Mover wants to move. */
@@ -279,7 +279,7 @@ protected:
 	UE_API int32 GetNetworkPhysicsTickOffset_Internal() const;
 	UE_API int32 GetNetworkPhysicsTickOffset_External() const;
 
-	// Time step on the physics thread when using async physics
+	// Time step on the worker thread when using async physics
 	UE_API FMoverTimeStep GetCurrentAsyncMoverTimeStep_Internal() const;
 
 	// Time step on the game thread when using async physics. Uses physics results time
@@ -309,7 +309,7 @@ protected:
 	UPROPERTY(Transient)
 	TObjectPtr<UNetworkPhysicsComponent> NetworkPhysicsComponent;
 
-	// This stores the previous two outputs from the physics thread
+	// This stores the previous two outputs from the simulation (worker) thread
 	// which can be interpolated to get the new output at a time
 	// that matches the interpolated physics particle.
 	FSimulationOutputRecord SimOutputRecord;
diff --git a/Engine/Plugins/Experimental/Mover/Source/Mover/Public/MovementModeStateMachine.h b/Engine/Plugins/Experimental/Mover/Source/Mover/Public/MovementModeStateMachine.h
--- a/Engine/Plugins/Experimental/Mover/Source/Mover/Public/MovementModeStateMachine.h
+++ b/Engine/Plugins/Experimental/Mover/Source/Mover/Public/MovementModeStateMachine.h
@@ -71,10 +71,10 @@ public:
 	UE_API void CancelFeaturesWithTag(FGameplayTag TagToCancel, bool bRequireExactMatch=false);
 	UE_API bool HasFeaturesWithTag(FGameplayTag TagToMatch, bool bRequireExactMatch = false) const;
 
-	// This function is meant to be used only in async mode on the physics thread, not on the game thread
+	// This function is meant to be used only in async mode on a worker thread, not on the game thread
 	void QueueInstantMovementEffect_Internal(const FScheduledInstantMovementEffect& ScheduledInstantMovementEffect);
 
-	// This function is meant to be used only in async mode on the physics thread, not on the game thread
+	// This function is meant to be used only in async mode on a worker thread, not on the game thread
 	void QueueLayeredMove_Internal(const FScheduledLayeredMove& ScheduledLayeredMove);
 
 protected:
diff --git a/Engine/Plugins/Experimental/Mover/Source/Mover/Public/MoverComponent.h b/Engine/Plugins/Experimental/Mover/Source/Mover/Public/MoverComponent.h
--- a/Engine/Plugins/Experimental/Mover/Source/Mover/Public/MoverComponent.h
+++ b/Engine/Plugins/Experimental/Mover/Source/Mover/Public/MoverComponent.h
@@ -44,6 +44,7 @@ DECLARE_DYNAMIC_MULTICAST_DELEGATE_ThreeParams(FMover_OnPostMovement, const FMov
 DECLARE_DYNAMIC_MULTICAST_DELEGATE_OneParam(FMover_OnPostSimTick, const FMoverTimeStep&, TimeStep);
 
 // Fired after a rollback. First param is the time step we've rolled back to. Second param is when we rolled back from, and represents a later frame that is no longer valid.
+// Note: ExpungedTimeStep is called PreRollbackTimeStep in C++ code
 DECLARE_DYNAMIC_MULTICAST_DELEGATE_TwoParams(FMover_OnPostSimRollback, const FMoverTimeStep&, CurrentTimeStep, const FMoverTimeStep&, ExpungedTimeStep);
 
 // Fired after changing movement modes. First param is the name of the previous movement mode. Second is the name of the new movement mode. 
@@ -55,6 +56,14 @@ DECLARE_DYNAMIC_MULTICAST_DELEGATE_FourParams(FMover_OnTeleportSucceeded, const
 // Fired when a teleport has failed
 DECLARE_DYNAMIC_MULTICAST_DELEGATE_FiveParams(FMover_OnTeleportFailed, const FVector&, FromLocation, const FQuat&, FromRotation, const FVector&, ToLocation, const FQuat&, ToRotation, ETeleportFailureReason, TeleportFailureReason);
 
+// Fired when a gameplay tag is added to the simulation (e.g. Mover_IsCrouching becomes set). Context carries bIsCausedByRollback and bIsDuringResimulation.
+// bIsExternalTag is true if the tag was applied via AddGameplayTag rather than from within the simulation.
+DECLARE_DYNAMIC_MULTICAST_DELEGATE_ThreeParams(FMover_OnGameplayTagAdded, const FGameplayTag&, Tag, const FMoverEventContext&, EventContext, bool, bIsExternalTag);
+
+// Fired when a gameplay tag is removed from the simulation. Context carries bIsCausedByRollback and bIsDuringResimulation.
+// bIsExternalTag is true if the tag was removed via RemoveGameplayTag rather than from within the simulation.
+DECLARE_DYNAMIC_MULTICAST_DELEGATE_ThreeParams(FMover_OnGameplayTagRemoved, const FGameplayTag&, Tag, const FMoverEventContext&, EventContext, bool, bIsExternalTag);
+
 // Fired after a transition has been triggered.
 DECLARE_DYNAMIC_MULTICAST_DELEGATE_OneParam(FMover_OnTransitionTriggered, const UBaseMovementModeTransition*, ModeTransition);
 
@@ -116,6 +125,14 @@ public:
 	UPROPERTY(BlueprintAssignable, Category = Mover)
 	FMover_OnTeleportFailed OnTeleportFailed;
 
+	// Broadcast when a gameplay tag becomes active, whether driven by the simulation internally or applied externally via AddGameplayTag. Fired on the game thread.
+	UPROPERTY(BlueprintAssignable, Category = Mover)
+	FMover_OnGameplayTagAdded OnGameplayTagAdded;
+
+	// Broadcast when a gameplay tag is no longer active, whether driven by the simulation internally or removed externally via RemoveGameplayTag. Fired on the game thread.
+	UPROPERTY(BlueprintAssignable, Category = Mover)
+	FMover_OnGameplayTagRemoved OnGameplayTagRemoved;
+
 	// Broadcast when a Transition has been triggered.
 	UPROPERTY(BlueprintAssignable, Category = Mover)
 	FMover_OnTransitionTriggered OnMovementTransitionTriggered;
@@ -657,12 +674,18 @@ public:	// Queries
 	template <typename ModifierT = FMovementModifierBase UE_REQUIRES(std::is_base_of_v<FMovementModifierBase, ModifierT>)>
 	const ModifierT* FindMovementModifierByType() const { return static_cast<const ModifierT*>(FindMovementModifierByType(ModifierT::StaticStruct())); }
 	
+	// Get the gameplay tags currently active, combining simulation-reported tags and externally-added tags.
+	// In the async backend (ChaosMover), or when mover.UseSnapshotGameplayTagsForSyncBackend is set,
+	// returns the cached AllGameplayTags snapshot. Otherwise performs a live walk of the current sync state.
+	UFUNCTION(BlueprintPure, Category = Mover)
+	MOVER_API FGameplayTagContainer GetGameplayTags() const;
+
 	/**
  	 * Check Mover systems for a gameplay tag.
  	 *
  	 * @param TagToFind			Tag to check on the Mover systems
  	 * @param bExactMatch		If true, the tag has to be exactly present, if false then TagToFind will include it's parent tags while matching
- 	 * 
+ 	 *
  	 * @return True if the TagToFind was found
  	 */
 	UFUNCTION(BlueprintPure, Category = Mover, meta = (Keywords = "HasTag"))
@@ -794,9 +817,21 @@ protected:
 	FTransform BaseVisualComponentTransform = FTransform::Identity;
 
 	// TODO: Look at possibility of replacing this with a FGameplayTagCountContainer that could possibly represent both internal and external tags
+	/** Gameplay tags currently active on the simulation. */
+	FGameplayTagContainer InternalGameplayTags;
+
 	/** A list of gameplay tags associated with this Mover Component added from sources outside of Mover */
 	FGameplayTagContainer ExternalGameplayTags;
-	
+
+	/** Combined cache of InternalGameplayTags and ExternalGameplayTags. Used by GetGameplayTags() and HasGameplayTag() in the async backend path, and in the sync backend path when mover.UseSnapshotGameplayTagsForSyncBackend is set. */
+	FGameplayTagContainer AllGameplayTags;
+
+	/** Rebuilds AllGameplayTags from InternalGameplayTags and ExternalGameplayTags. Must be called after either source changes. */
+	void RebuildAllGameplayTags();
+
+	/** Collects all gameplay tags currently active on the given sync state (movement mode, modifiers, layered moves) into OutTags. */
+	void CollectGameplayTagsFromSyncState(const FMoverSyncState& SyncState, FGameplayTagContainer& OutTags) const;
+
 	FMoverInputCmdContext CachedLastProducedInputCmd;
 	
 	FMoverInputCmdContext CachedLastUsedInputCmd;
@@ -848,7 +883,23 @@ protected:
 	
 	/** Helper function for making Queueing and making Layered Moves */
 	bool MakeAndQueueLayeredMove(const TSubclassOf<ULayeredMoveLogic>& MoveLogicClass, const FLayeredMoveActivationParams* ActivationParams);
-	
+
+	/**
+	 * Fires OnGameplayTagAdded/OnGameplayTagRemoved for every tag that is in NewTags but not PreviousTags (added),
+	 * or in PreviousTags but not NewTags (removed). Tags already covered by an entry in EventsToSkip are suppressed
+	 * to avoid double-firing alongside explicit simulation events. Tags also present in ExternalGameplayTags are
+	 * suppressed with a log, as concurrent simulation and external ownership of the same tag is unexpected.
+	 */
+	void DispatchGameplayTagTransitionEvents(const FGameplayTagContainer& PreviousTags, const FGameplayTagContainer& NewTags, const TArray<TSharedPtr<FMoverSimulationEventData>>& EventsToSkip, const FMoverTimeStep& TimeStep);
+
+	/**
+	 * Called after applying a new tag snapshot from simulation output. In principle, every queued simulation tag event
+	 * should account for every change between PreviousTags and NewTags. If that invariant is violated (i.e. a tag
+	 * changed state without a corresponding event), this is a bug - we recover by emitting the missing events and
+	 * logging a warning for each discrepancy. Delegates to DispatchGameplayTagTransitionEvents for the actual dispatch.
+	 */
+	void DispatchUnexpectedTagDriftEvents(const FGameplayTagContainer& PreviousTags, const FGameplayTagContainer& NewTags, const TArray<TSharedPtr<FMoverSimulationEventData>>& Events, const FMoverTimeStep& TimeStep);
+
 private:
 	/** Collection of settings objects that are shared between movement modes. This list is automatically managed based on the @MovementModes contents. */
 	UPROPERTY(EditDefaultsOnly, EditFixedSize, Instanced, Category = Mover, meta = (NoResetToDefault, ObjectMustImplement = "/Script/Mover.MovementSettingsInterface"))
diff --git a/Engine/Plugins/Experimental/Mover/Source/Mover/Public/MoverGameplayTagLog.h b/Engine/Plugins/Experimental/Mover/Source/Mover/Public/MoverGameplayTagLog.h
new file mode 100644
--- /dev/null
+++ b/Engine/Plugins/Experimental/Mover/Source/Mover/Public/MoverGameplayTagLog.h
@@ -0,0 +1,37 @@
+// Copyright Epic Games, Inc. All Rights Reserved.
+
+#pragma once
+
+#include "MoverSimulationTypes.h"
+
+// Set mover.debug.EnableGameplayTagLog to enable Mover gameplay tag event logging
+// across the worker thread (WT:*) and game thread (GT:*) steps.
+namespace MoverCVars
+{
+	MOVER_API extern bool EnableGameplayTagLog;
+}
+
+// Logs the [GT:DriftDetect] header line and a [GT:DriftDetect:Event] line per queued event.
+// The CVar check is included; safe to call unconditionally.
+MOVER_API void LogGameplayTagDriftDetect(
+	const FString& ActorName,
+	const TCHAR* RoleStr,
+	const FMoverTimeStep& TimeStep,
+	const FGameplayTag& Tag,
+	bool bWasAdded,
+	const TArray<TSharedPtr<FMoverSimulationEventData>>& Events,
+	const FGameplayTagContainer& PreviousTags,
+	const FGameplayTagContainer& NewTags);
+
+// Conditional log macro: gates on MoverCVars::EnableGameplayTagLog and prepends
+// the [GameplayTagLog] prefix. Accepts any UE log category.
+//   MOVER_TAG_LOG(LogMover, "[WT:TagDiff] Actor=%s Role=%s ...", *Name, *Role, ...)
+// Under NO_LOGGING the macro expands to a true no-op that discards all arguments,
+// so call sites do not need to guard themselves and DebugOwnerName-bearing members
+// can be stripped with a single #if !NO_LOGGING at the declaration.
+#if NO_LOGGING
+	#define MOVER_TAG_LOG(Category, Format, ...) do {} while (0)
+#else
+	#define MOVER_TAG_LOG(Category, StepAndFormat, ...) \
+		UE_CLOG(MoverCVars::EnableGameplayTagLog, Category, Log, TEXT("[GameplayTagLog]" StepAndFormat), ##__VA_ARGS__)
+#endif
diff --git a/Engine/Plugins/Experimental/Mover/Source/Mover/Public/MoverSimulationTypes.h b/Engine/Plugins/Experimental/Mover/Source/Mover/Public/MoverSimulationTypes.h
--- a/Engine/Plugins/Experimental/Mover/Source/Mover/Public/MoverSimulationTypes.h
+++ b/Engine/Plugins/Experimental/Mover/Source/Mover/Public/MoverSimulationTypes.h
@@ -634,6 +634,12 @@ private:
 
 public:
 	FMoverEventContext Context;
+
+	// If true, this event will always re-emit during resimulation and is never added to the
+	// ProcessedEvents dedup list. Set this on event types that represent state changes which
+	// must be faithfully replayed on resim (e.g. gameplay tag adds/removes). Leave false for
+	// punctual one-shot events that the Game Thread has already received and must not see again.
+	bool bReEmitOnResim = false;
 };
 
 USTRUCT()
@@ -759,6 +765,45 @@ struct FTeleportFailedEventData : public FMoverSimulationEventData
 	ETeleportFailureReason TeleportFailureReason;
 };
 
+USTRUCT()
+struct FMoverGameplayTagChangeEventData : public FMoverSimulationEventData
+{
+	GENERATED_BODY()
+
+	FMoverGameplayTagChangeEventData(const FMoverTimeStep& InEventTime, const FGameplayTag& InTag, bool bInWasAdded, FEventProcessedCallbackPtr InEventProcessedCallback = nullptr)
+		: FMoverSimulationEventData(InEventTime, InEventProcessedCallback)
+		, Tag(InTag)
+		, bWasAdded(bInWasAdded)
+	{
+		bReEmitOnResim = true;
+	}
+	FMoverGameplayTagChangeEventData()
+	{
+		bReEmitOnResim = true;
+	}
+
+	virtual UScriptStruct* GetScriptStruct() const override
+	{
+		return FMoverGameplayTagChangeEventData::StaticStruct();
+	}
+
+	virtual bool IsEqual(const FMoverSimulationEventData& Other, int32 ServerFrameDiffTolerance = 0) const override
+	{
+		if (Super::IsEqual(Other, ServerFrameDiffTolerance))
+		{
+			const FMoverGameplayTagChangeEventData& TypedOther = static_cast<const FMoverGameplayTagChangeEventData&>(Other);
+			return Tag == TypedOther.Tag && bWasAdded == TypedOther.bWasAdded;
+		}
+		return false;
+	}
+
+	FGameplayTag Tag;
+	bool bWasAdded = false;
+
+	// Was this event triggered by an externally-added tag (via AddGameplayTag/RemoveGameplayTag), rather than from within the simulation?
+	bool bIsExternalTag = false;
+};
+
 namespace UE::Mover
 {
 	struct FSimulationOutputData
@@ -774,6 +819,8 @@ namespace UE::Mover
 		FMoverInputCmdContext LastUsedInputCmd;
 		FMoverDataCollection AdditionalOutputData;
 		TArray<TSharedPtr<FMoverSimulationEventData>> Events;
+		/** Snapshot of all gameplay tags active on the simulation at the end of this tick, as collected from the simulation. */
+		FGameplayTagContainer InternalGameplayTags;
 		bool bIsValid = false;
 	};
 
@@ -797,10 +844,30 @@ namespace UE::Mover
 
 		MOVER_API void Clear();
 
+		/** Remove all queued events whose ServerFrame is >= FromServerFrame. Call this on the first resim frame after a rollback to discard stale predicted events. */
+		MOVER_API void PruneEventsFromFrame(int32 FromServerFrame);
+
+#if !NO_LOGGING
+		// Mover Gameplay Tag Logging
+		FString DebugOwnerName;
+		FString DebugOwnerRole;
+#endif
+
 	private:
 		FData Data[2];
 		TArray<TSharedPtr<FMoverSimulationEventData>> Events;
 		uint8 CurrentIndex = 1;
+
+		// Tracks accumulated post-event InternalGameplayTags across multiple calls to
+		// CreateInterpolatedResult for the same interpolation window. Incremented by
+		// Add() whenever new PT data arrives; compared in CreateInterpolatedResult to
+		// detect whether the Prev..Curr window has changed since the last call.
+		int32 NextInterpolatedGeneration = 0;
+		int32 LastInterpolatedGeneration = INDEX_NONE;
+		FGameplayTagContainer LastInterpolatedTags;
+
+		// Mover Gameplay Tag Logging
+		void LogTagInterpolate(uint8 PrevIndex, double AtBaseTimeMs, bool bIsInterpolated) const;
 	};
 
 } // namespace UE::Mover
