Mover - fix for null MoverComp in UFallingMode::GenerateMove_Implementation after recompiling BPs

#Synced-CL 52360741

diff --git a/Engine/Plugins/Experimental/Mover/Source/Mover/Private/DefaultMovementSet/Modes/FallingMode.cpp b/Engine/Plugins/Experimental/Mover/Source/Mover/Private/DefaultMovementSet/Modes/FallingMode.cpp
--- a/Engine/Plugins/Experimental/Mover/Source/Mover/Private/DefaultMovementSet/Modes/FallingMode.cpp
+++ b/Engine/Plugins/Experimental/Mover/Source/Mover/Private/DefaultMovementSet/Modes/FallingMode.cpp
@@ -42,6 +42,11 @@ void UFallingMode::GenerateMove_Implementation(const FMoverSimContext& SimContex
 
 	const float DeltaSeconds = TimeStep.StepMs * 0.001f;
 
+	if (!ensure(MoverComp))
+	{
+		return;
+	}
+
 	FVector UpDirection = MoverComp->GetUpDirection();
 	
 	// We don't want velocity limits to take the falling velocity component into account, since it is handled 
