Mover: Add FMoverTime util functions to make it easier to write code that can work in either fixed (frame count based) or variable (float) time. Basic FMoverTime arithmetic calculating seconds to another FMoverTime.

#Synced-CL 52152484

diff --git a/Engine/Plugins/Experimental/Mover/Source/Mover/Public/MoverTypes.h b/Engine/Plugins/Experimental/Mover/Source/Mover/Public/MoverTypes.h
--- a/Engine/Plugins/Experimental/Mover/Source/Mover/Public/MoverTypes.h
+++ b/Engine/Plugins/Experimental/Mover/Source/Mover/Public/MoverTypes.h
@@ -120,7 +120,19 @@ struct FMoverTime
 	{
 		return (FrameCount != INDEX_NONE) ? FString::Printf(TEXT("FrameCount = %d, TimeMs = %.3f"), FrameCount, TimeMs) : FString::Printf(TEXT("TimeMs = %.3f"), TimeMs);
 	}
-	
+
+	// Returns elapsed time in seconds from this time to Other.
+	// In fixed time steps (FrameCount != INDEX_NONE), uses integer frame arithmetic for exactness.
+	// (but intra step time differences, smaller than FixedFrameDtSeconds, will have to use other functions)
+	// In variable time steps (FrameCount == INDEX_NONE), uses floating-point ms arithmetic.
+	double ElapsedSecondsTo(const FMoverTime& Other, double FixedFrameDtSeconds) const
+	{
+		if (FrameCount != INDEX_NONE)
+			return (double)(Other.FrameCount - FrameCount) * FixedFrameDtSeconds;
+		else
+			return (Other.TimeMs - TimeMs) * 0.001;
+	}
+
 	// Number of fixed time steps that have occurred since time 0. Only valid when time advances in fixed time steps.
 	UPROPERTY(EditAnywhere, BlueprintReadWrite, Category=Mover)
 	int32 FrameCount=INDEX_NONE;
@@ -156,6 +168,29 @@ struct FMoverTimeStep
 	// If resimulating, is this the first frame after the rollback?
 	UPROPERTY(EditAnywhere, BlueprintReadOnly, Category=Mover)
 	uint8 bIsFirstResimFrame : 1 = 0;
+
+	// Returns an FMoverTime representing the start of this time step.
+	// Both FrameCount and TimeMs are always populated; ElapsedSecondsTo uses whichever is appropriate
+	// based on the reference FMoverTime's FrameCount (INDEX_NONE = variable dt, otherwise fixed dt).
+	FMoverTime ToStartTime() const
+	{
+		FMoverTime T;
+		T.FrameCount = ServerFrame;
+		T.TimeMs = BaseSimTimeMs;
+		return T;
+	}
+
+	// Returns an FMoverTime representing the start of the next step.
+	// With regard to time difference calculation, it also represents the time at the end of the current 
+	// step, but with a frame count incremented by one. Otherwise the distance between ToStartTime and
+	// ToNextTime in async mode would be 0, since it base time arithmetic on frame counts.
+	FMoverTime ToNextTime() const
+	{
+		FMoverTime T;
+		T.FrameCount = ServerFrame + 1;
+		T.TimeMs = BaseSimTimeMs + (double)StepMs;
+		return T;
+	}
 };
 
 
