Mover: UMovementUtils::ComputeVelocity now limits velocity input to max speed properly

#Synced-CL 52438879

diff --git a/Engine/Plugins/Experimental/Mover/Source/Mover/Private/MoveLibrary/MovementUtils.cpp b/Engine/Plugins/Experimental/Mover/Source/Mover/Private/MoveLibrary/MovementUtils.cpp
--- a/Engine/Plugins/Experimental/Mover/Source/Mover/Private/MoveLibrary/MovementUtils.cpp
+++ b/Engine/Plugins/Experimental/Mover/Source/Mover/Private/MoveLibrary/MovementUtils.cpp
@@ -43,6 +43,7 @@ FVector UMovementUtils::ComputeVelocity(const FComputeVelocityParams& InParams)
 		const float RequestedSpeed = FMath::Min(InParams.MaxSpeed,  InParams.MoveInput.Size());
 		const FVector RequestedMoveDir = InParams.MoveInput.GetSafeNormal();
 		DesiredSpeed = RequestedSpeed;
+		const FVector MoveVelocity = RequestedMoveDir * DesiredSpeed;
 
 		// Compute acceleration. Use MaxAccel to limit speed increase
 		if (InParams.bUseAccelerationForVelocityMove && InParams.InitialVelocity.Size() < RequestedSpeed * UE::MoverUtils::VELOCITY_INPUT_NO_ACCELERATION_DIFFERENCE)
@@ -51,14 +52,14 @@ FVector UMovementUtils::ComputeVelocity(const FComputeVelocityParams& InParams)
 			Velocity = Velocity - (Velocity - RequestedMoveDir * RequestedSpeed) * FMath::Min(InParams.DeltaSeconds * InParams.Friction, 1.f);
 
 			// How much do we need to accelerate to get to the new velocity?
-			Acceleration = (InParams.MoveInput - Velocity) / InParams.DeltaSeconds;
+			Acceleration = (MoveVelocity - Velocity) / InParams.DeltaSeconds;
 			Acceleration = Acceleration.GetClampedToMaxSize(InParams.Acceleration);
 		}
 		else
 		{
 			// Just set velocity directly.
 			// If decelerating we do so instantly, so we don't slide through the destination if we can't brake fast enough.
-			Velocity = InParams.MoveInput;
+			Velocity = MoveVelocity;
 		}
 	}
 	else if (InParams.MoveInputType == EMoveInputType::DirectionalIntent)
