Mover: Adding a projection setting to nav walking modes that allows users to specify what geometry to project to in relation to the navmesh. Also adding proper support for move off navmesh option for NavWalkingMode

#Synced-CL 52438879

diff --git a/Engine/Plugins/Experimental/Mover/Source/Mover/Private/DefaultMovementSet/Modes/AsyncNavWalkingMode.cpp b/Engine/Plugins/Experimental/Mover/Source/Mover/Private/DefaultMovementSet/Modes/AsyncNavWalkingMode.cpp
--- a/Engine/Plugins/Experimental/Mover/Source/Mover/Private/DefaultMovementSet/Modes/AsyncNavWalkingMode.cpp
+++ b/Engine/Plugins/Experimental/Mover/Source/Mover/Private/DefaultMovementSet/Modes/AsyncNavWalkingMode.cpp
@@ -486,13 +486,26 @@ void UAsyncNavWalkingMode::FindBestNavMeshLocation(const FVector& TraceStart, co
 	MultiTraceHits.RemoveAllSwap(FRemoveNotBlockingResponseNavMeshTrace(!bProjectNavMeshOnBothWorldChannels), EAllowShrinking::No);
 	if (MultiTraceHits.Num() > 0)
 	{
-		FVector UpDirection = GetMoverComponent()->GetUpDirection();
-		MultiTraceHits.Sort([UpDirection](const FHitResult& A, const FHitResult& B)
+		if (ProjectNavMeshToGeometryBehavior == EProjectNavMeshToGeometryBehavior::UseTopGeometry)
 		{
-			return A.ImpactPoint.Dot(UpDirection) > B.ImpactPoint.Dot(UpDirection);
-		});
+			FVector UpDirection = (TraceStart - TraceEnd).GetSafeNormal();
+			MultiTraceHits.Sort([UpDirection](const FHitResult& A, const FHitResult& B)
+			{
+				return A.ImpactPoint.Dot(UpDirection) > B.ImpactPoint.Dot(UpDirection);
+			});
+		}
+		else if (ProjectNavMeshToGeometryBehavior == EProjectNavMeshToGeometryBehavior::UseClosestGeometry)
+		{
+			// Sort the hits by the closest to our origin.
+			MultiTraceHits.Sort([TargetNavLocation](const FHitResult& A, const FHitResult& B)
+			{
+				const float ADistSqr = (TargetNavLocation - A.ImpactPoint).SizeSquared();
+				const float BDistSqr = (TargetNavLocation - B.ImpactPoint).SizeSquared();
+
+				return (ADistSqr < BDistSqr);	
+			});
+		}
 
-		// Cache the closest hit and treat it as a blocking hit (we used an overlap to get all the world static hits so we could sort them ourselves)
 		OutHitResult = MultiTraceHits[0];
 		OutHitResult.bBlockingHit = true;
 	}
diff --git a/Engine/Plugins/Experimental/Mover/Source/Mover/Private/DefaultMovementSet/Modes/NavWalkingMode.cpp b/Engine/Plugins/Experimental/Mover/Source/Mover/Private/DefaultMovementSet/Modes/NavWalkingMode.cpp
--- a/Engine/Plugins/Experimental/Mover/Source/Mover/Private/DefaultMovementSet/Modes/NavWalkingMode.cpp
+++ b/Engine/Plugins/Experimental/Mover/Source/Mover/Private/DefaultMovementSet/Modes/NavWalkingMode.cpp
@@ -280,6 +280,7 @@ void UNavWalkingMode::SimulationTick_Implementation(const FSimulationTickParams&
 
 		if (!bFoundPointOnNavMesh)
 		{
+			FHitResult MoveHitResult;
 			// Can't find nav mesh at this location, so we need to do something else
 			switch (BehaviorOffNavMesh)
 			{
@@ -293,8 +294,7 @@ void UNavWalkingMode::SimulationTick_Implementation(const FSimulationTickParams&
 
 				case EOffNavMeshBehavior::MoveWithoutNavMesh:
 					// allow the full move to occur 
-					// TODO: Need to actually move the updated component and add the movement to the move record
-					ensureMsgf(false, TEXT("NavWalkingMode does not yet support MoveWithoutNavMesh"));
+					UMovementUtils::TrySafeMoveUpdatedComponent(Params.MovingComps, OrigMoveDelta, TargetOrientQuat, /*bSweep?*/ false, MoveHitResult, ETeleportType::None, MoveRecord);
 					break;
 
 				case EOffNavMeshBehavior::DoNotMove:
@@ -303,7 +303,6 @@ void UNavWalkingMode::SimulationTick_Implementation(const FSimulationTickParams&
 					break;
 
 				case EOffNavMeshBehavior::RotateOnly:
-					FHitResult MoveHitResult;
 					UMovementUtils::TrySafeMoveUpdatedComponent(Params.MovingComps, FVector::ZeroVector, TargetOrientQuat, /*bSweep?*/ false, MoveHitResult, ETeleportType::None, MoveRecord);
 					break;
 			}
@@ -480,23 +479,6 @@ void UNavWalkingMode::FindBestNavMeshLocation(const FVector& TraceStart, const F
 	TArray<FHitResult> MultiTraceHits;
 	GetWorld()->LineTraceMultiByChannel(MultiTraceHits, TraceStart, TraceEnd, ECC_WorldStatic, Params, ResponseParams);
 
-	struct FCompareFHitResultNavMeshTrace
-	{
-		explicit FCompareFHitResultNavMeshTrace(const FVector& inSourceLocation) : SourceLocation(inSourceLocation)
-		{
-		}
-
-		FORCEINLINE bool operator()(const FHitResult& A, const FHitResult& B) const
-		{
-			const float ADistSqr = (SourceLocation - A.ImpactPoint).SizeSquared();
-			const float BDistSqr = (SourceLocation - B.ImpactPoint).SizeSquared();
-
-			return (ADistSqr < BDistSqr);
-		}
-
-		const FVector& SourceLocation;
-	};
-
 	struct FRemoveNotBlockingResponseNavMeshTrace
 	{
 		FRemoveNotBlockingResponseNavMeshTrace(bool bInCheckOnlyWorldStatic) : bCheckOnlyWorldStatic(bInCheckOnlyWorldStatic) {}
@@ -516,10 +498,26 @@ void UNavWalkingMode::FindBestNavMeshLocation(const FVector& TraceStart, const F
 	MultiTraceHits.RemoveAllSwap(FRemoveNotBlockingResponseNavMeshTrace(!bProjectNavMeshOnBothWorldChannels), EAllowShrinking::No);
 	if (MultiTraceHits.Num() > 0)
 	{
-		// Sort the hits by the closest to our origin.
-		MultiTraceHits.Sort(FCompareFHitResultNavMeshTrace(TargetNavLocation));
+		if (ProjectNavMeshToGeometryBehavior == EProjectNavMeshToGeometryBehavior::UseTopGeometry)
+		{
+			FVector UpDirection = (TraceStart - TraceEnd).GetSafeNormal();
+			MultiTraceHits.Sort([UpDirection](const FHitResult& A, const FHitResult& B)
+			{
+				return A.ImpactPoint.Dot(UpDirection) > B.ImpactPoint.Dot(UpDirection);
+			});
+		}
+		else if (ProjectNavMeshToGeometryBehavior == EProjectNavMeshToGeometryBehavior::UseClosestGeometry)
+		{
+			// Sort the hits by the closest to our origin.
+			MultiTraceHits.Sort([TargetNavLocation](const FHitResult& A, const FHitResult& B)
+			{
+				const float ADistSqr = (TargetNavLocation - A.ImpactPoint).SizeSquared();
+				const float BDistSqr = (TargetNavLocation - B.ImpactPoint).SizeSquared();
+
+				return (ADistSqr < BDistSqr);	
+			});
+		}
 
-		// Cache the closest hit and treat it as a blocking hit (we used an overlap to get all the world static hits so we could sort them ourselves)
 		OutHitResult = MultiTraceHits[0];
 		OutHitResult.bBlockingHit = true;
 	}
diff --git a/Engine/Plugins/Experimental/Mover/Source/Mover/Public/DefaultMovementSet/Modes/AsyncNavWalkingMode.h b/Engine/Plugins/Experimental/Mover/Source/Mover/Public/DefaultMovementSet/Modes/AsyncNavWalkingMode.h
--- a/Engine/Plugins/Experimental/Mover/Source/Mover/Public/DefaultMovementSet/Modes/AsyncNavWalkingMode.h
+++ b/Engine/Plugins/Experimental/Mover/Source/Mover/Public/DefaultMovementSet/Modes/AsyncNavWalkingMode.h
@@ -64,7 +64,11 @@ public:
 	/** Speed at which to interpolate agent navmesh offset between traces. 0: Instant (no interp) > 0: Interp speed") */
 	UPROPERTY(Category = "NavMesh Movement", EditAnywhere, BlueprintReadWrite, meta = (editcondition = "bProjectNavMeshWalking", ClampMin = "0", UIMin = "0"))
 	float NavMeshProjectionInterpSpeed;
-
+	
+	/** What to do when we are projecting nav movement onto geometry from the nav mesh */
+	UPROPERTY(Category = "NavMesh Movement", EditAnywhere, BlueprintReadWrite, meta=(editcondition = "bProjectNavMeshWalking"))
+	EProjectNavMeshToGeometryBehavior ProjectNavMeshToGeometryBehavior = EProjectNavMeshToGeometryBehavior::UseTopGeometry;
+	
 	/** What should we do if we stray off the nav mesh? */
 	UPROPERTY(Category = "NavMesh Movement", EditAnywhere, BlueprintReadWrite)
 	EOffNavMeshBehavior BehaviorOffNavMesh = EOffNavMeshBehavior::SwitchToWalking;
diff --git a/Engine/Plugins/Experimental/Mover/Source/Mover/Public/DefaultMovementSet/Modes/NavWalkingMode.h b/Engine/Plugins/Experimental/Mover/Source/Mover/Public/DefaultMovementSet/Modes/NavWalkingMode.h
--- a/Engine/Plugins/Experimental/Mover/Source/Mover/Public/DefaultMovementSet/Modes/NavWalkingMode.h
+++ b/Engine/Plugins/Experimental/Mover/Source/Mover/Public/DefaultMovementSet/Modes/NavWalkingMode.h
@@ -5,7 +5,6 @@
 #include "CoreMinimal.h"
 #include "Engine/EngineTypes.h"
 #include "MovementMode.h"
-#include "MoverComponent.h"
 #include "MoveLibrary/BasedMovementUtils.h"
 #include "AI/Navigation/NavigationTypes.h"
 #include "NavWalkingMode.generated.h"
@@ -13,6 +12,7 @@
 #define UE_API MOVER_API
 
 class UNavMoverComponent;
+class UMoverComponent;
 class INavigationDataInterface;
 class UCommonLegacyMovementSettings;
 
@@ -34,12 +34,22 @@ enum class EOffNavMeshBehavior : uint8
 	RotateOnly = 3			   UMETA(DisplayName = "Rotate Only"),
 };
 
+/** Options for what to do when we are projecting nav movement onto geometry from the nav mesh */
+UENUM(BlueprintType)
+enum class EProjectNavMeshToGeometryBehavior : uint8
+{
+	/** Projects movement onto the closest geometry to the target location projected onto the navmesh. */
+	UseClosestGeometry = 0        UMETA(DisplayName = "Use closest Geometry"),
+
+	/** Projects movement onto the top most geometry in projection range of the navmesh */
+	UseTopGeometry = 1     UMETA(DisplayName = "Use Top Geometry"),
+};
+
 namespace NavWalkingUtils
 {
 	void CheckClientSideNavigationSettings(const UMoverComponent* MoverComponent);
 }
 
-
 /**
  * NavWalkingMode: a default movement mode for traversing surfaces and movement bases by using an active navmesh when moving the actor rather than collision checks.
  * Note: This movement mode requires a NavMoverComponent be on the actor to function properly. This mode also contains some randomization to avoid navmesh look ups
@@ -90,6 +100,10 @@ public:
 	UPROPERTY(Category="NavMesh Movement", EditAnywhere, BlueprintReadWrite, meta=(editcondition = "bProjectNavMeshWalking", ClampMin="0", UIMin="0"))
 	float NavMeshProjectionInterpSpeed;
 
+	/** What to do when we are projecting nav movement onto geometry from the nav mesh */
+	UPROPERTY(Category = "NavMesh Movement", EditAnywhere, BlueprintReadWrite, meta=(editcondition = "bProjectNavMeshWalking"))
+	EProjectNavMeshToGeometryBehavior ProjectNavMeshToGeometryBehavior = EProjectNavMeshToGeometryBehavior::UseTopGeometry;
+	
 	/** What should we do if we stray off the nav mesh? */
 	UPROPERTY(Category = "NavMesh Movement", EditAnywhere, BlueprintReadWrite)
 	EOffNavMeshBehavior BehaviorOffNavMesh = EOffNavMeshBehavior::SwitchToWalking;
