Mover CVD: trace local sim data using the particle extra data API

Switches Mover's CVD tracing to TraceChaosVDParticleExtraData() /
FChaosVDExtraDataStructEntry::SerializeFrom(), defaulting to SerializeBin
(property-tagged, fault-tolerant across struct field additions and removals).

- MoverCVDRuntimeTrace.cpp: replaces the manual FChaosVDStructCollectionMemoryWriter +
SerializeItem call with Entry.SerializeFrom() for each local-sim-data struct
- MoverCVDSimDataProcessor.cpp: reads back the extra-data container from the solver
frame's CustomData and surfaces it in the CVD particle details panel alongside
the existing Mover-specific data

#Synced-CL 52655663

diff --git a/Engine/Plugins/Experimental/Mover/Source/Mover/Private/ChaosVisualDebugger/MoverCVDRuntimeTrace.cpp b/Engine/Plugins/Experimental/Mover/Source/Mover/Private/ChaosVisualDebugger/MoverCVDRuntimeTrace.cpp
--- a/Engine/Plugins/Experimental/Mover/Source/Mover/Private/ChaosVisualDebugger/MoverCVDRuntimeTrace.cpp
+++ b/Engine/Plugins/Experimental/Mover/Source/Mover/Private/ChaosVisualDebugger/MoverCVDRuntimeTrace.cpp
@@ -6,9 +6,11 @@
 
 #include "MoverCVDDataWrappers.h"
 #include "MoverSimulationTypes.h"
+#include "MoverTypes.h"
 #include "Serialization/MemoryWriter.h"
 #include "Serialization/MemoryReader.h"
 #include "ChaosVisualDebugger/ChaosVisualDebuggerTrace.h"
+#include "ChaosVisualDebugger/ChaosVDParticleExtraData.h"
 #include "MoverComponent.h"
 #include "Chaos/PhysicsObjectInterface.h"
 #include "Chaos/PhysicsObjectInternalInterface.h"
@@ -44,7 +46,10 @@ public:
 
 	virtual FArchive& operator<<(struct FObjectPtr& Value) override
 	{
-		if (UScriptStruct* ScriptStruct = Cast<UScriptStruct>(Value.Get()))
+		UScriptStruct* ScriptStruct = Cast<UScriptStruct>(Value.Get());
+		bool bIsScriptStruct = (ScriptStruct != nullptr);
+		*this << bIsScriptStruct;
+		if (bIsScriptStruct)
 		{
 			// The FullName of the script struct will be something like "ScriptStruct /Script/Mover.FCharacterDefaultInputs"
 			FString FullStructName = ScriptStruct->GetFullName(nullptr);
@@ -84,8 +89,16 @@ public:
 	{
 	}
 
-	virtual FArchive& operator<<(struct FObjectPtr& Value)
+	virtual FArchive& operator<<(struct FObjectPtr& Value) override
 	{
+		bool bIsScriptStruct = false;
+		*this << bIsScriptStruct;
+		if (bIsScriptStruct)
+		{
+			FString StructName;
+			*this << StructName;
+			Value = Cast<UScriptStruct>(FindObject<UStruct>(nullptr, *StructName));
+		}
 		return *this;
 	}
 
@@ -127,7 +140,7 @@ void FMoverCVDRuntimeTrace::UnwrapSimData(const FMoverCVDSimDataWrapper& InSimDa
 		OutSyncState->SyncStateCollection.SerializeDebugData(ArSyncStateCollectionReader);
 	}
 
-	// Local sim data — one deserialized FMoverDataCollection per named section
+	// Local sim data  - one deserialized FMoverDataCollection per named section
 	for (const TPair<FName, TArray<uint8>>& Section : InSimDataWrapper.LocalSimDataSections)
 	{
 		TSharedPtr<FMoverDataCollection> DataCollection = MakeShared<FMoverDataCollection>();
@@ -165,7 +178,7 @@ void FMoverCVDRuntimeTrace::WrapSimData(uint32 SolverID, uint32 ParticleID, cons
 		const_cast<FMoverSyncState&>(InSyncState).SyncStateCollection.SerializeDebugData(ArWriter);
 	}
 
-	// Local sim data — serialize each named collection into its own section to preserve origin in the CVD viewer
+	// Local sim data  - serialize each named collection into its own section to preserve origin in the CVD viewer
 	if (LocalSimDataCollections)
 	{
 		for (const TPair<FName, const FMoverDataCollection*>& NamedCollection : *LocalSimDataCollections)
@@ -246,6 +259,87 @@ void FMoverCVDRuntimeTrace::TraceMoverDataPrivate(uint32 SolverID, uint32 Partic
 	Chaos::VisualDebugger::WriteDataToBuffer(TLSDataBuffer.BufferRef, SimDataWrapper);
 
 	FChaosVisualDebuggerTrace::TraceBinaryData(TLSDataBuffer.BufferRef, FMoverCVDSimDataWrapper::WrapperTypeName);
+
+	// Particle extra data: structured view in the CVD particle details panel.
+	// Each category name becomes a collapsible section header in the panel.
+	// MoverNetworkedData channel is guaranteed enabled by callers; LocalSimDataCollections
+	// is already nullptr when the MoverLocalSimData channel is disabled.
+	using namespace Chaos::VisualDebugger;
+
+	static const FName SyncStateCategoryName(TEXT("Sync State"));
+	static const FName InputCategoryName(TEXT("Input"));
+	static const FName NetworkedChannelName(TEXT("MoverNetworkedData"));
+	static const FName LocalSimChannelName(TEXT("MoverLocalSimData"));
+
+	FChaosVDParticleExtraData ExtraData;
+
+	// Sync State: fixed UPROPERTY fields (MovementMode, LayeredMoves, etc.) as the first entry,
+	// then each collection item (e.g. FMoverDefaultSyncState with Location/Velocity) individually.
+	// Note: SyncStateCollection appears empty in the fixed-fields entry because FMoverDataCollection::DataArray
+	// is not a UPROPERTY; the actual per-item data is in the entries that follow.
+	{
+		FChaosVDExtraDataCategory& SyncCat = ExtraData.Categories.AddDefaulted_GetRef();
+		SyncCat.CategoryName = SyncStateCategoryName;
+		SyncCat.SourceChannelId = NetworkedChannelName;
+		SyncCat.AddEntry(FMoverSyncState::StaticStruct(), SyncState);
+		for (const TSharedPtr<FMoverDataStructBase>& Item : SyncState->SyncStateCollection.GetDataArray())
+		{
+			if (Item.IsValid())
+			{
+				if (UScriptStruct* Struct = Item->GetScriptStruct())
+				{
+					SyncCat.AddEntry(Struct, Item.Get());
+				}
+			}
+		}
+	}
+
+	// Input: FMoverInputCmdContext has no meaningful fixed UPROPERTY fields beyond InputCollection
+	// itself, so only the collection items (e.g. FCharacterDefaultInputs) are added as entries.
+	{
+		FChaosVDExtraDataCategory& InputCat = ExtraData.Categories.AddDefaulted_GetRef();
+		InputCat.CategoryName = InputCategoryName;
+		InputCat.SourceChannelId = NetworkedChannelName;
+		for (const TSharedPtr<FMoverDataStructBase>& Item : InputCmd->InputCollection.GetDataArray())
+		{
+			if (Item.IsValid())
+			{
+				if (UScriptStruct* Struct = Item->GetScriptStruct())
+				{
+					InputCat.AddEntry(Struct, Item.Get());
+				}
+			}
+		}
+	}
+
+	// Local sim data: one named category per section (names are set by the caller, e.g. "LocalSimInput").
+	if (LocalSimDataCollections)
+	{
+		for (const TPair<FName, const FMoverDataCollection*>& NamedCollection : *LocalSimDataCollections)
+		{
+			if (!NamedCollection.Value)
+			{
+				continue;
+			}
+
+			FChaosVDExtraDataCategory& LocalCat = ExtraData.Categories.AddDefaulted_GetRef();
+			LocalCat.CategoryName = NamedCollection.Key;
+			LocalCat.SourceChannelId = LocalSimChannelName;
+
+			for (const TSharedPtr<FMoverDataStructBase>& Item : NamedCollection.Value->GetDataArray())
+			{
+				if (Item.IsValid())
+				{
+					if (UScriptStruct* Struct = Item->GetScriptStruct())
+					{
+						LocalCat.AddEntry(Struct, Item.Get());
+					}
+				}
+			}
+		}
+	}
+
+	TraceChaosVDParticleExtraData(static_cast<int32>(SolverID), static_cast<int32>(ParticleID), ExtraData);
 }
 
 } // namespace UE::MoverUtils
diff --git a/Engine/Plugins/Experimental/Mover/Source/MoverCVDEditor/Private/MoverCVDSimDataProcessor.cpp b/Engine/Plugins/Experimental/Mover/Source/MoverCVDEditor/Private/MoverCVDSimDataProcessor.cpp
--- a/Engine/Plugins/Experimental/Mover/Source/MoverCVDEditor/Private/MoverCVDSimDataProcessor.cpp
+++ b/Engine/Plugins/Experimental/Mover/Source/MoverCVDEditor/Private/MoverCVDSimDataProcessor.cpp
@@ -6,6 +6,8 @@
 
 FMoverCVDSimDataProcessor::FMoverCVDSimDataProcessor() : FChaosVDDataProcessorBase(FMoverCVDSimDataWrapper::WrapperTypeName)
 {
+	// Uses SerializeBin on raw structs with no explicit versioning  - not safe to read across binary-incompatible builds.
+	bBackwardsCompatible = false;
 }
 
 bool FMoverCVDSimDataProcessor::ProcessRawData(const TArray<uint8>& InData)
