Mover: fixes sim proxy montage time drift that was causing divergence with root motion animation playback (misaligned root motion & skeletal animation), and adds debug logging to include montage positions

#Synced-CL 52673010

diff --git a/Engine/Plugins/Experimental/Mover/Source/Mover/Private/DefaultMovementSet/CharacterMoverComponent.cpp b/Engine/Plugins/Experimental/Mover/Source/Mover/Private/DefaultMovementSet/CharacterMoverComponent.cpp
--- a/Engine/Plugins/Experimental/Mover/Source/Mover/Private/DefaultMovementSet/CharacterMoverComponent.cpp
+++ b/Engine/Plugins/Experimental/Mover/Source/Mover/Private/DefaultMovementSet/CharacterMoverComponent.cpp
@@ -252,6 +252,29 @@ void UCharacterMoverComponent::UpdateSyncedMontageState(const FMoverTimeStep& Ti
 					bShouldStartNewMontage = true;
 					bShouldStopSyncedMontage = true;
 				}
+				else
+				{
+					// We are continuing a montage. Check for the local mesh's montage drifting and subtly adjust playback rate to keep it close to the movement state.
+					const USkeletalMeshComponent* MeshComp = Cast<USkeletalMeshComponent>(GetPrimaryVisualComponent());
+					if (UAnimInstance* MeshAnimInstance = MeshComp ? MeshComp->GetAnimInstance() : nullptr)
+					{
+						const float LocalMontagePos = MeshAnimInstance->Montage_GetPosition(SyncedMontageState.Montage);
+
+						const float MontagePosError = NewMontageState.CurrentPosition - LocalMontagePos;
+
+						constexpr float MaxDriftCorrectionPct = 0.1f;	// +- 10% play rate adjustment
+						constexpr float DriftSecsForMaxCorrection = 0.010f;	// if more than 10 ms away, apply full drift correction
+
+						const float CorrectivePlayRate = NewMontageState.PlayRate + (MaxDriftCorrectionPct * (FMath::Clamp(MontagePosError, -DriftSecsForMaxCorrection, DriftSecsForMaxCorrection) / DriftSecsForMaxCorrection));
+					
+						MeshAnimInstance->Montage_SetPlayRate(SyncedMontageState.Montage, CorrectivePlayRate);
+
+						#if !UE_BUILD_SHIPPING
+						UE_CLOG(CVarLogSimProxyMontageReplication->GetBool(), LogMover, Log, TEXT("Mover SP montage continued (SimF %i SimT: %.3f) LocalMontagePos %.3f (Drift %.3f) => corrective playrate: %.2f"),
+							TimeStep.ServerFrame, TimeStep.BaseSimTimeMs, LocalMontagePos, MontagePosError, CorrectivePlayRate);
+						#endif // !UE_BUILD_SHIPPING
+					}
+				}
 			}
 			else
 			{
diff --git a/Engine/Plugins/Experimental/Mover/Source/Mover/Private/DefaultMovementSet/LayeredMoves/AnimRootMotionLayeredMove.cpp b/Engine/Plugins/Experimental/Mover/Source/Mover/Private/DefaultMovementSet/LayeredMoves/AnimRootMotionLayeredMove.cpp
--- a/Engine/Plugins/Experimental/Mover/Source/Mover/Private/DefaultMovementSet/LayeredMoves/AnimRootMotionLayeredMove.cpp
+++ b/Engine/Plugins/Experimental/Mover/Source/Mover/Private/DefaultMovementSet/LayeredMoves/AnimRootMotionLayeredMove.cpp
@@ -31,6 +31,7 @@ FLayeredMove_AnimRootMotion::FLayeredMove_AnimRootMotion()
 
 bool FLayeredMove_AnimRootMotion::GenerateMove(const FMoverTickStartData& SimState, const FMoverTimeStep& TimeStep, const UMoverComponent* MoverComp, UMoverBlackboard* SimBlackboard, FProposedMove& OutProposedMove)
 {
+	float MeshMontagePos = -1.f;
 	// Stop this move if the montage is no longer playing on the mesh
 	if (!TimeStep.bIsResimulating)
 	{
@@ -41,6 +42,10 @@ bool FLayeredMove_AnimRootMotion::GenerateMove(const FMoverTickStartData& SimSta
 			if (const UAnimInstance* MeshAnimInstance = MeshComp->GetAnimInstance())
 			{
 				bIsMontageStillPlaying = MontageState.Montage && MeshAnimInstance->Montage_IsPlaying(MontageState.Montage);
+				if (bIsMontageStillPlaying)
+				{
+					MeshMontagePos = MeshAnimInstance->Montage_GetPosition(MontageState.Montage);
+				}
 			}
 		}
 
@@ -100,9 +105,9 @@ bool FLayeredMove_AnimRootMotion::GenerateMove(const FMoverTickStartData& SimSta
 	MontageState.CurrentPosition = ExtractionStartPosition;
 
 #if !UE_BUILD_SHIPPING
-	UE_CLOGF(CVarLogAnimRootMotionSteps->GetBool(), LogMover, Log, "AnimRootMotion. SimF %i (dt %.3f) Range [%.3f, %.3f] => LocalT: %ls (WST: %ls)  Vel: %.3f",
-	        TimeStep.ServerFrame, DeltaSeconds, ExtractionStartPosition, ExtractionEndPosition, 
-	        *LocalRootMotion.GetTranslation().ToString(), *WorldSpaceRootMotion.GetTranslation().ToString(), OutProposedMove.LinearVelocity.Length());
+	UE_CLOGF(CVarLogAnimRootMotionSteps->GetBool(), LogMover, Log, "AnimRootMotion. SimF %i / SimT %.1f (Resim? %i) (dt %.3f) Range [%.3f, %.3f] => LocalT: %ls (WST: %ls)  Vel: %.3f  MntPos: %.3f",
+	        TimeStep.ServerFrame, TimeStep.BaseSimTimeMs, TimeStep.bIsResimulating, DeltaSeconds, ExtractionStartPosition, ExtractionEndPosition,
+	        *LocalRootMotion.GetTranslation().ToString(), *WorldSpaceRootMotion.GetTranslation().ToString(), OutProposedMove.LinearVelocity.Length(), MeshMontagePos);
 #endif // !UE_BUILD_SHIPPING
 
 	return true;
