Build fix: Fix for an overprotective static analysis issue

#Synced-CL 52666112

diff --git a/Engine/Plugins/Experimental/Mover/Source/Mover/Private/ChaosVisualDebugger/MoverCVDRuntimeTrace.cpp b/Engine/Plugins/Experimental/Mover/Source/Mover/Private/ChaosVisualDebugger/MoverCVDRuntimeTrace.cpp
--- a/Engine/Plugins/Experimental/Mover/Source/Mover/Private/ChaosVisualDebugger/MoverCVDRuntimeTrace.cpp
+++ b/Engine/Plugins/Experimental/Mover/Source/Mover/Private/ChaosVisualDebugger/MoverCVDRuntimeTrace.cpp
@@ -49,7 +49,7 @@ public:
 		UScriptStruct* ScriptStruct = Cast<UScriptStruct>(Value.Get());
 		bool bIsScriptStruct = (ScriptStruct != nullptr);
 		*this << bIsScriptStruct;
-		if (bIsScriptStruct)
+		if (ScriptStruct)
 		{
 			// The FullName of the script struct will be something like "ScriptStruct /Script/Mover.FCharacterDefaultInputs"
 			FString FullStructName = ScriptStruct->GetFullName(nullptr);
