Mover: suppress warnings about premature trajectory queries for Editor worlds

#rnx
#Synced-CL 52673010

diff --git a/Engine/Plugins/Experimental/Mover/Source/Mover/Private/MoverComponent.cpp b/Engine/Plugins/Experimental/Mover/Source/Mover/Private/MoverComponent.cpp
--- a/Engine/Plugins/Experimental/Mover/Source/Mover/Private/MoverComponent.cpp
+++ b/Engine/Plugins/Experimental/Mover/Source/Mover/Private/MoverComponent.cpp
@@ -2630,12 +2630,12 @@ TArray<FTrajectorySampleInfo> UMoverComponent::GetPredictedTrajectory(FMoverPred
 	}
 
 	
-	// Editor preview worlds i.e. thumbnails/blueprint actor previews have undefined init orders. We can skip warning in those cases.
+	// Editor/editor preview worlds i.e. thumbnails/blueprint actor previews have undefined init orders. We can skip warning in those cases.
 	// TODO Revisit if we expect/need trajectory to work in an editor preview world.
 	const UWorld* World = GetWorld();
-	if (World && World->WorldType != EWorldType::EditorPreview)
+	if (World && World->WorldType != EWorldType::EditorPreview && World->WorldType != EWorldType::Editor)
 	{
-		UE_LOGF(LogMover, Warning, "Cannot predict trajectory on a MoverComponent that is not initialized. Returning defaulted samples. Actor: %ls  Initialized? %i BegunPlay? %i BeingDestroyed? %i", *GetNameSafe(GetOwner()), HasBeenInitialized(), HasBegunPlay(), IsBeingDestroyed());
+		UE_LOGF(LogMover, Warning, "Cannot predict trajectory on a MoverComponent that is not initialized. Returning defaulted samples. Actor: %ls  Initialized? %i BegunPlay? %i BeingDestroyed? %i WorldType: %ls", *GetNameSafe(GetOwner()), HasBeenInitialized(), HasBegunPlay(), IsBeingDestroyed(), LexToString(World->WorldType));
 	}
 
 	TArray<FTrajectorySampleInfo> BlankDefaultSamples;
