Mover: reworked kinematic based movement to eliminate drift on rotating, pitching platforms. This enables characters to stand solidly on pitching, heaving platforms like a boat.

#Synced-CL 52695783

diff --git a/Engine/Plugins/Experimental/Mover/Source/Mover/Private/DefaultMovementSet/Modes/AsyncFallingMode.cpp b/Engine/Plugins/Experimental/Mover/Source/Mover/Private/DefaultMovementSet/Modes/AsyncFallingMode.cpp
--- a/Engine/Plugins/Experimental/Mover/Source/Mover/Private/DefaultMovementSet/Modes/AsyncFallingMode.cpp
+++ b/Engine/Plugins/Experimental/Mover/Source/Mover/Private/DefaultMovementSet/Modes/AsyncFallingMode.cpp
@@ -284,7 +284,7 @@ void UAsyncFallingMode::OnUnregistered(const FMoverSimContext& SimContext)
 	Super::OnUnregistered(SimContext);
 }
 
-void UAsyncFallingMode::ProcessLanded(const FFloorCheckResult& FloorResult, FVector& Velocity, FRelativeBaseInfo& BaseInfo, FMoverTickEndData& TickEndData) const
+void UAsyncFallingMode::ProcessLanded(const FFloorCheckResult& FloorResult, FVector Location, FVector& Velocity, FRelativeBaseInfo& BaseInfo, FMoverTickEndData& TickEndData) const
 {
 	// TODO: Refactor landed events for async movement. For now, leave the immediate event, but listeners are limited in what they can do from a worker thread.
 
@@ -295,9 +295,11 @@ void UAsyncFallingMode::ProcessLanded(const FFloorCheckResult& FloorResult, FVec
 	// if we can walk on the floor we landed on
 	if (FloorResult.IsWalkableFloor())
 	{
+		const FVector UpDir = MoverComp->GetUpDirection();
+
 		if (bCancelVerticalSpeedOnLanding)
 		{
-			const FPlane MovementPlane(FVector::ZeroVector, MoverComp->GetUpDirection());
+			const FPlane MovementPlane(FVector::ZeroVector, UpDir);
 			Velocity = UMovementUtils::ConstrainToPlane(Velocity, MovementPlane, false);
 		}
 		else
@@ -312,7 +314,7 @@ void UAsyncFallingMode::ProcessLanded(const FFloorCheckResult& FloorResult, FVec
 
 		if (UBasedMovementUtils::IsADynamicBase(FloorResult.HitResult.GetComponent()))
 		{
-			BaseInfo.SetFromFloorResult(FloorResult);
+			BaseInfo.SetFromFloorResult(FloorResult, Location, UpDir);
 		}
 	}
 	// we could check for other surfaces here (i.e. when swimming is implemented we can check the floor hit here and see if we need to go into swimming)
@@ -350,7 +352,7 @@ void UAsyncFallingMode::CaptureFinalState(const FMoverDefaultSyncState* StartSyn
 	FVector EffectiveVelocity = DeltaSecondsUsed <= UE_SMALL_NUMBER ? StartSyncState->GetVelocity_WorldSpace() : Record.GetRelevantVelocity();
 
 	FRelativeBaseInfo MovementBaseInfo;
-	ProcessLanded(FloorResult, EffectiveVelocity, MovementBaseInfo, TickEndData);
+	ProcessLanded(FloorResult, FinalLocation, EffectiveVelocity, MovementBaseInfo, TickEndData);
 
 	if (MovementBaseInfo.HasRelativeInfo())
 	{
diff --git a/Engine/Plugins/Experimental/Mover/Source/Mover/Private/DefaultMovementSet/Modes/AsyncNavWalkingMode.cpp b/Engine/Plugins/Experimental/Mover/Source/Mover/Private/DefaultMovementSet/Modes/AsyncNavWalkingMode.cpp
--- a/Engine/Plugins/Experimental/Mover/Source/Mover/Private/DefaultMovementSet/Modes/AsyncNavWalkingMode.cpp
+++ b/Engine/Plugins/Experimental/Mover/Source/Mover/Private/DefaultMovementSet/Modes/AsyncNavWalkingMode.cpp
@@ -674,8 +674,6 @@ void UAsyncNavWalkingMode::CaptureOutputState(const FMoverDefaultSyncState& Star
 {
 	UMoverBlackboard* SimBlackboard = GetMoverComponent()->GetSimBlackboard_Mutable();
 
-	// If we're on a dynamic base and we're not trying to move, keep using the same relative actor location. This prevents slow relative 
-	//  drifting that can occur from repeated floor sampling as the base moves through the world.
 	SimBlackboard->Invalidate(CommonBlackboard::LastFoundDynamicMovementBase);
 
 	OutputSyncState.SetTransforms_WorldSpace(FinalLocation,
diff --git a/Engine/Plugins/Experimental/Mover/Source/Mover/Private/DefaultMovementSet/Modes/AsyncWalkingMode.cpp b/Engine/Plugins/Experimental/Mover/Source/Mover/Private/DefaultMovementSet/Modes/AsyncWalkingMode.cpp
--- a/Engine/Plugins/Experimental/Mover/Source/Mover/Private/DefaultMovementSet/Modes/AsyncWalkingMode.cpp
+++ b/Engine/Plugins/Experimental/Mover/Source/Mover/Private/DefaultMovementSet/Modes/AsyncWalkingMode.cpp
@@ -376,7 +376,7 @@ void UAsyncWalkingMode::CaptureFinalState(const FVector FinalLocation, const FRo
 
 	const bool bHasPriorBaseInfo = SimBlackboard->TryGet(CommonBlackboard::LastFoundDynamicMovementBase, PriorBaseInfo);
 
-	FRelativeBaseInfo CurrentBaseInfo = UpdateFloorAndBaseInfo(FloorResult);
+	FRelativeBaseInfo CurrentBaseInfo = UpdateFloorAndBaseInfo(FloorResult, FinalLocation);
 
 	// If we're on a dynamic base and we're not trying to move, keep using the same relative actor location. This prevents slow relative 
 	//  drifting that can occur from repeated floor sampling as the base moves through the world.
@@ -385,6 +385,7 @@ void UAsyncWalkingMode::CaptureFinalState(const FVector FinalLocation, const FRo
 		&& PriorBaseInfo.UsesSameBase(CurrentBaseInfo))
 	{
 		CurrentBaseInfo.ContactLocalPosition = PriorBaseInfo.ContactLocalPosition;
+		CurrentBaseInfo.WorldspaceOffsetFromContactPos = PriorBaseInfo.WorldspaceOffsetFromContactPos;
 	}
 	
 	if (CurrentBaseInfo.HasRelativeInfo())
@@ -410,7 +411,7 @@ void UAsyncWalkingMode::CaptureFinalState(const FVector FinalLocation, const FRo
 }
 
 
-FRelativeBaseInfo UAsyncWalkingMode::UpdateFloorAndBaseInfo(const FFloorCheckResult& FloorResult) const
+FRelativeBaseInfo UAsyncWalkingMode::UpdateFloorAndBaseInfo(const FFloorCheckResult& FloorResult, FVector CurrentLocation) const
 {
 	FRelativeBaseInfo ReturnBaseInfo;
 
@@ -421,7 +422,7 @@ FRelativeBaseInfo UAsyncWalkingMode::UpdateFloorAndBaseInfo(const FFloorCheckRes
 
 	if (FloorResult.IsWalkableFloor() && UBasedMovementUtils::IsADynamicBase(FloorResult.HitResult.GetComponent()))
 	{
-		ReturnBaseInfo.SetFromFloorResult(FloorResult);
+		ReturnBaseInfo.SetFromFloorResult(FloorResult, CurrentLocation, MoverComp->GetUpDirection());
 	}
 
 	return ReturnBaseInfo;
diff --git a/Engine/Plugins/Experimental/Mover/Source/Mover/Private/DefaultMovementSet/Modes/FallingMode.cpp b/Engine/Plugins/Experimental/Mover/Source/Mover/Private/DefaultMovementSet/Modes/FallingMode.cpp
--- a/Engine/Plugins/Experimental/Mover/Source/Mover/Private/DefaultMovementSet/Modes/FallingMode.cpp
+++ b/Engine/Plugins/Experimental/Mover/Source/Mover/Private/DefaultMovementSet/Modes/FallingMode.cpp
@@ -280,18 +280,20 @@ void UFallingMode::OnUnregistered(const FMoverSimContext& SimContext)
 	Super::OnUnregistered(SimContext);
 }
 
-void UFallingMode::ProcessLanded(const FFloorCheckResult& FloorResult, FVector& Velocity, FRelativeBaseInfo& BaseInfo, FMoverTickEndData& TickEndData) const
+void UFallingMode::ProcessLanded(const FFloorCheckResult& FloorResult, FVector Location, FVector& Velocity, FRelativeBaseInfo& BaseInfo, FMoverTickEndData& TickEndData) const
 {
 	const UMoverComponent* MoverComp = GetMoverComponent();
 	UMoverBlackboard* SimBlackboard = MoverComp->GetSimBlackboard_Mutable();
 
+	const FVector UpDir = MoverComp->GetUpDirection();
+
 	FName NextMovementMode = NAME_None; 
 	// if we can walk on the floor we landed on
 	if (FloorResult.IsWalkableFloor())
 	{
 		if (bCancelVerticalSpeedOnLanding)
 		{
-			const FPlane MovementPlane(FVector::ZeroVector, MoverComp->GetUpDirection());
+			const FPlane MovementPlane(FVector::ZeroVector, UpDir);
 			Velocity = UMovementUtils::ConstrainToPlane(Velocity, MovementPlane, false);
 		}
 		else
@@ -306,7 +308,7 @@ void UFallingMode::ProcessLanded(const FFloorCheckResult& FloorResult, FVector&
 
 		if (UBasedMovementUtils::IsADynamicBase(FloorResult.HitResult.GetComponent()))
 		{
-			BaseInfo.SetFromFloorResult(FloorResult);
+			BaseInfo.SetFromFloorResult(FloorResult, Location, UpDir);
 		}
 	}
 	// we could check for other surfaces here (i.e. when swimming is implemented we can check the floor hit here and see if we need to go into swimming)
@@ -347,7 +349,7 @@ void UFallingMode::CaptureFinalState(USceneComponent* UpdatedComponent, const FM
 	// TODO: Update Main/large movement record with substeps from our local record
 
 	FRelativeBaseInfo MovementBaseInfo;
-	ProcessLanded(FloorResult, EffectiveVelocity, MovementBaseInfo, TickEndData);
+	ProcessLanded(FloorResult, FinalLocation, EffectiveVelocity, MovementBaseInfo, TickEndData);
 
 	if (MovementBaseInfo.HasRelativeInfo())
 	{
diff --git a/Engine/Plugins/Experimental/Mover/Source/Mover/Private/DefaultMovementSet/Modes/WalkingMode.cpp b/Engine/Plugins/Experimental/Mover/Source/Mover/Private/DefaultMovementSet/Modes/WalkingMode.cpp
--- a/Engine/Plugins/Experimental/Mover/Source/Mover/Private/DefaultMovementSet/Modes/WalkingMode.cpp
+++ b/Engine/Plugins/Experimental/Mover/Source/Mover/Private/DefaultMovementSet/Modes/WalkingMode.cpp
@@ -127,6 +127,18 @@ void UWalkingMode::SimulationTick_Implementation(const FSimulationTickParams& Pa
 
 	FMoverDefaultSyncState& OutputSyncState = OutputState.SyncState.SyncStateCollection.FindOrAddMutableDataByType<FMoverDefaultSyncState>();
 
+	UMoverBlackboard* SimBlackboard = MoverComp->GetSimBlackboard_Mutable();
+
+	FRelativeBaseInfo CurrentBaseInfo;
+
+	if (StartingSyncState->GetMovementBase() && SimBlackboard->TryGet(CommonBlackboard::LastFoundDynamicMovementBase, CurrentBaseInfo))
+	{
+		if (CurrentBaseInfo.HasRelativeInfo())
+		{
+			CurrentBaseInfo = UBasedMovementUtils::ApplyMovementBaseChange(MoverComp, CurrentBaseInfo, CommonLegacySettings->bIgnoreBaseRotation);
+			SimBlackboard->Set(CommonBlackboard::LastFoundDynamicMovementBase, CurrentBaseInfo);
+		}
+	}
 
 	const float DeltaSeconds = Params.TimeStep.StepMs * 0.001f;
 
@@ -136,7 +148,6 @@ void UWalkingMode::SimulationTick_Implementation(const FSimulationTickParams& Pa
 	MoveRecord.SetDeltaSeconds(DeltaSeconds);
 
 	FFloorCheckResult CurrentFloor;
-	UMoverBlackboard* SimBlackboard = MoverComp->GetSimBlackboard_Mutable();
 
 	const FVector UpDirection = MoverComp->GetUpDirection();
 	FMovingComponentSet MovingComponents(MoverComp);
@@ -249,7 +260,7 @@ void UWalkingMode::SimulationTick_Implementation(const FSimulationTickParams& Pa
 			OutputState.MovementEndState.NextModeName = CommonLegacySettings->AirMovementModeName;
 			OutputState.MovementEndState.RemainingMs = Params.TimeStep.StepMs - (Params.TimeStep.StepMs * PercentTimeAppliedSoFar);
 			MoveRecord.SetDeltaSeconds((Params.TimeStep.StepMs - OutputState.MovementEndState.RemainingMs) * 0.001f);
-			CaptureFinalState(UpdatedComponent, bDidAttemptMovement, CurrentFloor, MoveRecord, ProposedMove.AngularVelocityDegrees, OutputSyncState);
+			CaptureFinalState(UpdatedComponent, bDidAttemptMovement, CurrentFloor, MoveRecord, Params.TimeStep, ProposedMove.AngularVelocityDegrees, OutputSyncState);
 			return;
 		}
 	}
@@ -281,13 +292,13 @@ void UWalkingMode::SimulationTick_Implementation(const FSimulationTickParams& Pa
 				OutputState.MovementEndState.NextModeName = CommonLegacySettings->AirMovementModeName;
 				OutputState.MovementEndState.RemainingMs = Params.TimeStep.StepMs;
 				MoveRecord.SetDeltaSeconds((Params.TimeStep.StepMs - OutputState.MovementEndState.RemainingMs) * 0.001f);
-				CaptureFinalState(UpdatedComponent, bDidAttemptMovement, CurrentFloor, MoveRecord, ProposedMove.AngularVelocityDegrees, OutputSyncState);
+				CaptureFinalState(UpdatedComponent, bDidAttemptMovement, CurrentFloor, MoveRecord, Params.TimeStep, ProposedMove.AngularVelocityDegrees, OutputSyncState);
 				return;
 			}
 		}
 	}
 
-	CaptureFinalState(UpdatedComponent, bDidAttemptMovement, CurrentFloor, MoveRecord, ProposedMove.AngularVelocityDegrees, OutputSyncState);
+	CaptureFinalState(UpdatedComponent, bDidAttemptMovement, CurrentFloor, MoveRecord, Params.TimeStep, ProposedMove.AngularVelocityDegrees, OutputSyncState);
 
 }
 
@@ -324,24 +335,24 @@ void UWalkingMode::OnUnregistered(const FMoverSimContext& SimContext)
 	Super::OnUnregistered(SimContext);
 }
 
-void UWalkingMode::CaptureFinalState(USceneComponent* UpdatedComponent, bool bDidAttemptMovement, const FFloorCheckResult& FloorResult, const FMovementRecord& Record, const FVector& AngularVelocityDegrees, FMoverDefaultSyncState& OutputSyncState) const
+void UWalkingMode::CaptureFinalState(USceneComponent* UpdatedComponent, bool bDidAttemptMovement, const FFloorCheckResult& FloorResult, const FMovementRecord& Record, const FMoverTimeStep& TimeStep, const FVector& AngularVelocityDegrees, FMoverDefaultSyncState& OutputSyncState) const
 {
-	FRelativeBaseInfo PriorBaseInfo;
 
 	const UMoverComponent* MoverComp = GetMoverComponent();
 	UMoverBlackboard* SimBlackboard = MoverComp->GetSimBlackboard_Mutable();
 
+	FRelativeBaseInfo PriorBaseInfo;
 	const bool bHasPriorBaseInfo = SimBlackboard->TryGet(CommonBlackboard::LastFoundDynamicMovementBase, PriorBaseInfo);
 
-	FRelativeBaseInfo CurrentBaseInfo = UpdateFloorAndBaseInfo(FloorResult);
+	FRelativeBaseInfo CurrentBaseInfo = UpdateFloorAndBaseInfo(FloorResult, UpdatedComponent->GetComponentLocation());
 
 	// If we're on a dynamic base and we're not trying to move, keep using the same relative actor location. This prevents slow relative 
 	//  drifting that can occur from repeated floor sampling as the base moves through the world.
-	if (CurrentBaseInfo.HasRelativeInfo() 
-		&& bHasPriorBaseInfo && !bDidAttemptMovement 
-		&& PriorBaseInfo.UsesSameBase(CurrentBaseInfo))
+	if (!bDidAttemptMovement && CurrentBaseInfo.HasRelativeInfo()
+		&& bHasPriorBaseInfo && PriorBaseInfo.UsesSameBase(CurrentBaseInfo))
 	{
-		CurrentBaseInfo.ContactLocalPosition = PriorBaseInfo.ContactLocalPosition;
+ 		CurrentBaseInfo.ContactLocalPosition           = PriorBaseInfo.ContactLocalPosition;
+ 		CurrentBaseInfo.WorldspaceOffsetFromContactPos = PriorBaseInfo.WorldspaceOffsetFromContactPos;
 	}
 
 	// TODO: Update Main/large movement record with substeps from our local record
@@ -372,7 +383,7 @@ void UWalkingMode::CaptureFinalState(USceneComponent* UpdatedComponent, bool bDi
 }
 
 
-FRelativeBaseInfo UWalkingMode::UpdateFloorAndBaseInfo(const FFloorCheckResult& FloorResult) const
+FRelativeBaseInfo UWalkingMode::UpdateFloorAndBaseInfo(const FFloorCheckResult& FloorResult, FVector CurrentLocation) const
 {
 	FRelativeBaseInfo ReturnBaseInfo;
 
@@ -383,7 +394,7 @@ FRelativeBaseInfo UWalkingMode::UpdateFloorAndBaseInfo(const FFloorCheckResult&
 
 	if (FloorResult.IsWalkableFloor() && UBasedMovementUtils::IsADynamicBase(FloorResult.HitResult.GetComponent()))
 	{
-		ReturnBaseInfo.SetFromFloorResult(FloorResult);
+		ReturnBaseInfo.SetFromFloorResult(FloorResult, CurrentLocation, MoverComp->GetUpDirection());
 	}
 
 	return ReturnBaseInfo;
diff --git a/Engine/Plugins/Experimental/Mover/Source/Mover/Private/MoveLibrary/BasedMovementUtils.cpp b/Engine/Plugins/Experimental/Mover/Source/Mover/Private/MoveLibrary/BasedMovementUtils.cpp
--- a/Engine/Plugins/Experimental/Mover/Source/Mover/Private/MoveLibrary/BasedMovementUtils.cpp
+++ b/Engine/Plugins/Experimental/Mover/Source/Mover/Private/MoveLibrary/BasedMovementUtils.cpp
@@ -18,6 +18,7 @@ void FRelativeBaseInfo::Clear()
 	Location = FVector::ZeroVector;
 	Rotation = FQuat::Identity;
 	ContactLocalPosition = FVector::ZeroVector;
+	WorldspaceOffsetFromContactPos = FVector::ZeroVector;
 }
 
 bool FRelativeBaseInfo::HasRelativeInfo() const
@@ -37,7 +38,7 @@ bool FRelativeBaseInfo::UsesSameBase(const UPrimitiveComponent* OtherComp, FName
 		&& (BoneName == OtherBoneName);
 }
 
-void FRelativeBaseInfo::SetFromFloorResult(const FFloorCheckResult& FloorTestResult)
+void FRelativeBaseInfo::SetFromFloorResult(const FFloorCheckResult& FloorTestResult, const FVector ObjectWorldLocation, const FVector UpDir)
 {
 	bool bDidSucceed = false;
 
@@ -49,10 +50,18 @@ void FRelativeBaseInfo::SetFromFloorResult(const FFloorCheckResult& FloorTestRes
 		{
 			BoneName = FloorTestResult.HitResult.BoneName;
 
-			if (UBasedMovementUtils::GetMovementBaseTransform(MovementBase.Get(), BoneName, OUT Location, OUT Rotation) &&
-				UBasedMovementUtils::TransformWorldLocationToBased(MovementBase.Get(), BoneName, FloorTestResult.HitResult.ImpactPoint, OUT ContactLocalPosition))
+			if (UBasedMovementUtils::GetMovementBaseTransform(MovementBase.Get(), BoneName, OUT Location, OUT Rotation))
 			{
-				bDidSucceed = true;
+				// Projecting onto the trace line, to account for a flat bottom base that does not have a reliable impact point directly under the object
+				FVector ContactPositionInWorldSpace = FMath::ClosestPointOnInfiniteLine(FloorTestResult.HitResult.TraceStart, 
+																						FloorTestResult.HitResult.TraceEnd,
+																						FloorTestResult.HitResult.ImpactPoint);
+
+				if (UBasedMovementUtils::TransformWorldLocationToBased(MovementBase.Get(), BoneName, ContactPositionInWorldSpace, OUT ContactLocalPosition))
+				{
+					WorldspaceOffsetFromContactPos = (ObjectWorldLocation - ContactPositionInWorldSpace).ProjectOnToNormal(UpDir);
+					bDidSucceed = true;
+				}
 			}
 		}
 	}
@@ -63,19 +72,9 @@ void FRelativeBaseInfo::SetFromFloorResult(const FFloorCheckResult& FloorTestRes
 	}
 }
 
-void FRelativeBaseInfo::SetFromComponent(UPrimitiveComponent* InRelativeComp, FName InBoneName)
+void FRelativeBaseInfo::UpdateToLatestBaseTransform()
 {
-	bool bDidSucceed = false;
-
-	MovementBase = InRelativeComp;
-
-	if (MovementBase.IsValid())
-	{
-		BoneName = InBoneName;
-		bDidSucceed = UBasedMovementUtils::GetMovementBaseTransform(MovementBase.Get(), BoneName, /*out*/Location, /*out*/Rotation);
-	}
-
-	if (!bDidSucceed)
+	if (!UBasedMovementUtils::GetMovementBaseTransform(MovementBase.Get(), BoneName, OUT Location, OUT Rotation))
 	{
 		Clear();
 	}
@@ -86,11 +85,12 @@ FString FRelativeBaseInfo::ToString() const
 {
 	if (MovementBase.IsValid())
 	{
-		return FString::Printf(TEXT("Base: %s, Loc: %s, Rot: %s, LocalContact: %s"),
+		return FString::Printf(TEXT("Base: %s, Loc: %s, Rot: %s, LocalContact: %s, WSContactOffset: %s"),
 			*GetNameSafe(MovementBase->GetOwner()),
 			*Location.ToCompactString(),
 			*Rotation.Rotator().ToCompactString(),
-			*ContactLocalPosition.ToCompactString());
+			*ContactLocalPosition.ToCompactString(),
+			*WorldspaceOffsetFromContactPos.ToCompactString());
 	}
 
 	return FString(TEXT("Base: NULL"));
@@ -262,6 +262,89 @@ void UBasedMovementUtils::TransformRotatorToLocal(FQuat BaseQuat, FRotator World
 	OutLocalRotator = (BaseQuat.Inverse() * WorldQuat).Rotator();
 }
 
+FRelativeBaseInfo UBasedMovementUtils::ApplyMovementBaseChange(UMoverComponent* TargetMoverComp, const FRelativeBaseInfo& PriorBaseInfo, bool bIgnoreBaseRotation)
+{
+	// Early out with no changes if this isn't a valid base
+	if (!PriorBaseInfo.HasRelativeInfo())
+	{
+		return PriorBaseInfo;
+	}
+
+	FRelativeBaseInfo CurrentBaseInfo = PriorBaseInfo;
+	CurrentBaseInfo.UpdateToLatestBaseTransform();
+
+	if (!CurrentBaseInfo.HasRelativeInfo())
+	{
+		return CurrentBaseInfo;
+	}
+
+	// Calculate new transform matrix of base actor (ignoring scale).
+	const FQuatRotationTranslationMatrix CurrentBaseLocalToWorld(CurrentBaseInfo.Rotation, CurrentBaseInfo.Location);
+
+	// Find where the actor *should* be
+	// NOTE that we are using a combination of the relative floor contact location, plus a worldspace offset. 
+	// This is so the location we're touching the base is unchanging, and the actor location always stays above the 
+	// base, even if the base tilts and heaves like a rocking boat.
+	const FVector CurrentWorldBaseContactPos = CurrentBaseLocalToWorld.TransformPosition(CurrentBaseInfo.ContactLocalPosition);
+	const FVector TargetWorldPos = CurrentWorldBaseContactPos + CurrentBaseInfo.WorldspaceOffsetFromContactPos;
+
+	USceneComponent* UpdatedComponent = TargetMoverComp->GetUpdatedComponent();
+	const FVector OldWorldPos = UpdatedComponent->GetComponentLocation();
+	const FVector WorldBaseMovementDelta = TargetWorldPos - OldWorldPos;
+
+	FQuat WorldTargetQuat = UpdatedComponent->GetComponentQuat();
+
+	// Find any change in local rotation if a rotating base should cause a similar rotation in the pawn
+	if (!bIgnoreBaseRotation && !CurrentBaseInfo.Rotation.Equals(PriorBaseInfo.Rotation, UE_SMALL_NUMBER))
+	{
+		// Attempting to match the platform's rotation as best as possible, but respect the up direction so the character
+		// stays oriented upright even if the platform pitches / heaves / tilts like a rocking boat.
+
+		FQuat DeltaQuat = CurrentBaseInfo.Rotation * PriorBaseInfo.Rotation.Inverse();
+		WorldTargetQuat = DeltaQuat * WorldTargetQuat;
+
+		FVector TargetForwVector = WorldTargetQuat.GetForwardVector();
+		TargetForwVector = FVector::VectorPlaneProject(TargetForwVector, -TargetMoverComp->GetUpDirection());
+		TargetForwVector.Normalize();
+
+		FVector TargetRightVector = WorldTargetQuat.GetRightVector();
+		TargetRightVector = FVector::VectorPlaneProject(TargetRightVector, -TargetMoverComp->GetUpDirection());
+		TargetRightVector.Normalize();
+
+		WorldTargetQuat = UKismetMathLibrary::MakeRotFromXY(TargetForwVector, TargetRightVector).Quaternion();
+	}
+
+
+
+	if (!WorldBaseMovementDelta.IsZero() || WorldTargetQuat != UpdatedComponent->GetComponentQuat())
+	{
+		EMoveComponentFlags MoveComponentFlags = MOVECOMP_IgnoreBases | MOVECOMP_DisableBlockingOverlapDispatch;
+		const bool bSweep = true;	// we need this to be a sweeping movement in case the base movement pushes us into an obstacle... we need to get scraped off
+		FHitResult MoveHitResult;
+
+		bool bDidMove = UMovementUtils::TryMoveUpdatedComponent_Internal(FMovingComponentSet(TargetMoverComp), WorldBaseMovementDelta, WorldTargetQuat, bSweep, MoveComponentFlags, &MoveHitResult, ETeleportType::None);
+
+		const FVector NewWorldPos = UpdatedComponent->GetComponentLocation();
+
+		// If we did not move the full amount, that means we are blocked and the base is moving along without us.
+		// We're getting scraped along our base, so adjust our contact point backwards to account for it.
+		const FVector RemainingWorldDelta = (OldWorldPos + WorldBaseMovementDelta) - NewWorldPos;
+
+		if (!RemainingWorldDelta.IsNearlyZero())
+		{
+			// Find the remaining delta that wasn't achieved
+			// Convert the remaining delta to current base space
+			FVector RemainingLocalDelta;
+			UBasedMovementUtils::TransformDirectionToLocal(CurrentBaseInfo.Rotation, RemainingWorldDelta, OUT RemainingLocalDelta);
+
+			// Subtract the remaining delta to reflect the change in the contact position
+			CurrentBaseInfo.ContactLocalPosition -= RemainingLocalDelta;
+		}
+	}
+
+	return CurrentBaseInfo;
+}
+
 void UBasedMovementUtils::AddTickDependency(FTickFunction& BasedObjectTick, UPrimitiveComponent* NewBase)
 {
 	if (NewBase && IsADynamicBase(NewBase))
@@ -336,178 +419,105 @@ void UBasedMovementUtils::UpdateSimpleBasedMovement(UMoverComponent* TargetMover
 		bIgnoreBaseRotation = CommonSettings->bIgnoreBaseRotation;
 	}
 
-	bool bDidGetUpToDate = false;
-
-	FRelativeBaseInfo LastFoundBaseInfo;	// Last-found is the most recent capture during movement, likely set this sim frame
-	FRelativeBaseInfo LastAppliedBaseInfo;	// Last-applied is the one that our based movement is up to date with, likely set in the last sim frame
-	FRelativeBaseInfo CurrentBaseInfo;		// Current info is the current snapshot of the current base, with up-to-date transform that may be different than last-found.
-
-	const bool bHasLastFoundInfo = SimBlackboard->TryGet(CommonBlackboard::LastFoundDynamicMovementBase, LastFoundBaseInfo);
-	const bool bHasLastAppliedInfo = SimBlackboard->TryGet(CommonBlackboard::LastAppliedDynamicMovementBase, LastAppliedBaseInfo);
-	if (bHasLastFoundInfo)
+	bool bDidPerformBasedMovement = false;
+	ON_SCOPE_EXIT
 	{
-		if (!bHasLastAppliedInfo || !LastFoundBaseInfo.UsesSameBase(LastAppliedBaseInfo))
+		if (!bDidPerformBasedMovement)
 		{
-			LastAppliedBaseInfo = LastFoundBaseInfo;	// This is the first time we've checked this base, so start with the last-found capture
-		}
-
-		if (!ensureMsgf(LastFoundBaseInfo.HasRelativeInfo() && LastFoundBaseInfo.UsesSameBase(LastAppliedBaseInfo),
-				TEXT("Attempting to update based movement with a missing or mismatched base. This may indicate a logic problem with detecting bases.")))
-		{ 
 			SimBlackboard->Invalidate(CommonBlackboard::LastFoundDynamicMovementBase);
-			SimBlackboard->Invalidate(CommonBlackboard::LastAppliedDynamicMovementBase);
-			return;
 		}
+	};
+
+
+
+	FRelativeBaseInfo PriorBaseInfo;		// This is the snapshot of the base where we last used it
+
+	const bool bHasPriorBaseInfo = SimBlackboard->TryGet(CommonBlackboard::LastFoundDynamicMovementBase, PriorBaseInfo);
 
-		CurrentBaseInfo.SetFromComponent(LastFoundBaseInfo.MovementBase.Get(), LastFoundBaseInfo.BoneName);
-		CurrentBaseInfo.ContactLocalPosition = LastFoundBaseInfo.ContactLocalPosition;
+	
+	if (!ensureMsgf(bHasPriorBaseInfo, TEXT("No information about the prior movement base info. Nothing we can do.")))
+	{
+		return;
+	}
 
-		FVector CurrentBaseLocation;
-		FQuat CurrentBaseQuat;
-			
 
-		if (UBasedMovementUtils::GetMovementBaseTransform(CurrentBaseInfo.MovementBase.Get(), CurrentBaseInfo.BoneName, OUT CurrentBaseLocation, OUT CurrentBaseQuat))
+
+	FRelativeBaseInfo CurrentBaseInfo = UBasedMovementUtils::ApplyMovementBaseChange(TargetMoverComp, PriorBaseInfo, bIgnoreBaseRotation);
+
+
+	// Propagate the movement changes to the backend's state, if supported
+
+	// Note that this is occurring out-of-band with the movement simulation, in order to support based movement regardless of update order or
+	// whether the movement base is also simulated through Mover.
+
+	FMoverSyncState PendingSimSyncState;
+	if (TargetMoverComp->BackendLiaisonComp->ReadPendingSyncState(OUT PendingSimSyncState))
+	{
+		// Modify the PENDING sync state that has not yet been committed to simulation history nor replicated
+		if (FMoverDefaultSyncState* PendingMoverState = PendingSimSyncState.SyncStateCollection.FindMutableDataByType<FMoverDefaultSyncState>())
 		{
-			const bool bDidBaseRotationChange = !LastAppliedBaseInfo.Rotation.Equals(CurrentBaseQuat, UE_SMALL_NUMBER);
-			const bool bDidBaseLocationChange = (LastAppliedBaseInfo.Location != CurrentBaseLocation);
+			FTransform OldSyncTransformWs = PendingMoverState->GetTransform_WorldSpace();
+			FTransform NewSyncTransformWs = UpdatedComponent->GetComponentTransform();
 
-			FQuat DeltaQuat = FQuat::Identity;
-			FVector WorldDeltaLocation = FVector::ZeroVector;
-			FQuat WorldTargetQuat = UpdatedComponent->GetComponentQuat();
+			PendingMoverState->SetTransforms_WorldSpace(
+				NewSyncTransformWs.GetLocation(),
+				NewSyncTransformWs.GetRotation().Rotator(),
+				PendingMoverState->GetVelocity_WorldSpace(),	// keep same velocity and base
+				PendingMoverState->GetAngularVelocityDegrees_WorldSpace(),
+				PendingMoverState->GetMovementBase(), PendingMoverState->GetMovementBaseBoneName());
 
-			// Find change in rotation
+			TargetMoverComp->BackendLiaisonComp->WritePendingSyncState(PendingSimSyncState);	// writes pending Simulation state
 
-			if (bDidBaseRotationChange && !bIgnoreBaseRotation)
+			// If smoothing, modify presentation-related states as well so that the visual offset location stays anchored to the movement base
+			if (TargetMoverComp->SmoothingMode != EMoverSmoothingMode::None)
 			{
-				DeltaQuat = CurrentBaseQuat * LastAppliedBaseInfo.Rotation.Inverse();
-				WorldTargetQuat = DeltaQuat * WorldTargetQuat;
-
-				FVector TargetForwVector = WorldTargetQuat.GetForwardVector();
-				TargetForwVector = FVector::VectorPlaneProject(TargetForwVector, -TargetMoverComp->GetUpDirection());
-				TargetForwVector.Normalize();
-
-				FVector TargetRightVector = WorldTargetQuat.GetRightVector();
-				TargetRightVector = FVector::VectorPlaneProject(TargetRightVector, -TargetMoverComp->GetUpDirection());
-				TargetRightVector.Normalize();
-					
-				WorldTargetQuat = UKismetMathLibrary::MakeRotFromXY(TargetForwVector, TargetRightVector).Quaternion();
-			}
+				const FTransform OldToNewTransform = NewSyncTransformWs.GetRelativeTransform(OldSyncTransformWs);
 
-			if (bDidBaseLocationChange || bDidBaseRotationChange)
-			{
-				// Calculate new transform matrix of base actor (ignoring scale).
-				const FQuatRotationTranslationMatrix OldLocalToWorld(LastAppliedBaseInfo.Rotation, LastAppliedBaseInfo.Location);
-				const FQuatRotationTranslationMatrix NewLocalToWorld(CurrentBaseQuat, CurrentBaseLocation);
-
-				// Find change in location
-				// NOTE that we are using the floor hit location, not the actor's root position which may be floating above the base
-				const FVector NewWorldBaseContactPos = NewLocalToWorld.TransformPosition(CurrentBaseInfo.ContactLocalPosition);
-				const FVector OldWorldBaseContactPos = OldLocalToWorld.TransformPosition(CurrentBaseInfo.ContactLocalPosition);
-				WorldDeltaLocation = NewWorldBaseContactPos - OldWorldBaseContactPos;
-
-				const FVector OldWorldLocation = UpdatedComponent->GetComponentLocation();
-				EMoveComponentFlags MoveComponentFlags = MOVECOMP_IgnoreBases;
-				const bool bSweep = true;
-				FHitResult MoveHitResult;
-
-				bool bDidMove = UMovementUtils::TryMoveUpdatedComponent_Internal(FMovingComponentSet(TargetMoverComp), WorldDeltaLocation, WorldTargetQuat, bSweep, MoveComponentFlags, &MoveHitResult, ETeleportType::None);
-					
-				const FVector NewWorldLocation = UpdatedComponent->GetComponentLocation();
-
-				if ((NewWorldLocation - (OldWorldLocation + WorldDeltaLocation)).IsNearlyZero() == false)
+				// Modify the PRESENTATION sync state that we're smoothing TO
+				FMoverSyncState PresentationSyncState;
+				if (TargetMoverComp->BackendLiaisonComp->ReadPresentationSyncState(OUT PresentationSyncState))
 				{
-					// Find the remaining delta that wasn't achieved
-					const FVector UnachievedWorldDelta = (OldWorldLocation + WorldDeltaLocation) - NewWorldLocation;
-
-					// Convert the remaining delta to current base space
-					FVector UnachievedLocalDelta;
-					UBasedMovementUtils::TransformLocationToLocal(CurrentBaseLocation, CurrentBaseQuat, UnachievedWorldDelta, OUT UnachievedLocalDelta);
-						
-					// Subtract the remaining delta to reflect the change in the contact position
-					CurrentBaseInfo.ContactLocalPosition -= UnachievedLocalDelta;
-				}
+					if (FMoverDefaultSyncState* PresentationMoverState = PresentationSyncState.SyncStateCollection.FindMutableDataByType<FMoverDefaultSyncState>())
+					{
+						OldSyncTransformWs = PresentationMoverState->GetTransform_WorldSpace();
+						NewSyncTransformWs = OldToNewTransform * OldSyncTransformWs;
 
-				// Propagate the movement changes to the backend's state, if supported
+						PresentationMoverState->SetTransforms_WorldSpace(
+							NewSyncTransformWs.GetLocation(),
+							NewSyncTransformWs.GetRotation().Rotator(),
+							PresentationMoverState->GetVelocity_WorldSpace(),	// keep same velocity and base
+							PresentationMoverState->GetAngularVelocityDegrees_WorldSpace(),
+							PresentationMoverState->GetMovementBase(), PresentationMoverState->GetMovementBaseBoneName());
 
-				// Note that this is occurring out-of-band with the movement simulation, in order to support based movement regardless of update order or
-				// whether the movement base is also simulated through Mover.
+						TargetMoverComp->BackendLiaisonComp->WritePresentationSyncState(PresentationSyncState);
+					}
+				}
 
-				FMoverSyncState PendingSimSyncState;
-				if (TargetMoverComp->BackendLiaisonComp->ReadPendingSyncState(OUT PendingSimSyncState))
+				// Modify the PREV PRESENTATION sync state that we're smoothing FROM
+				FMoverSyncState PrevPresentationSyncState;
+				if (TargetMoverComp->BackendLiaisonComp->ReadPrevPresentationSyncState(OUT PrevPresentationSyncState))
 				{
-					// Modify the PENDING sync state that has not yet been committed to simulation history nor replicated
-					if (FMoverDefaultSyncState* PendingMoverState = PendingSimSyncState.SyncStateCollection.FindMutableDataByType<FMoverDefaultSyncState>())
+					if (FMoverDefaultSyncState* PrevPresentationMoverState = PrevPresentationSyncState.SyncStateCollection.FindMutableDataByType<FMoverDefaultSyncState>())
 					{
-						FTransform OldSyncTransformWs = PendingMoverState->GetTransform_WorldSpace();
-						FTransform NewSyncTransformWs = UpdatedComponent->GetComponentTransform();
+						OldSyncTransformWs = PrevPresentationMoverState->GetTransform_WorldSpace();
+						NewSyncTransformWs = OldToNewTransform * OldSyncTransformWs;
 
-						PendingMoverState->SetTransforms_WorldSpace(
+						PrevPresentationMoverState->SetTransforms_WorldSpace(
 							NewSyncTransformWs.GetLocation(),
 							NewSyncTransformWs.GetRotation().Rotator(),
-							PendingMoverState->GetVelocity_WorldSpace(),	// keep same velocity and base
-							PendingMoverState->GetAngularVelocityDegrees_WorldSpace(),
-							PendingMoverState->GetMovementBase(), PendingMoverState->GetMovementBaseBoneName());
-
-						TargetMoverComp->BackendLiaisonComp->WritePendingSyncState(PendingSimSyncState);	// writes pending Simulation state
-
-						// If smoothing, modify presentation-related states as well so that the visual offset location stays anchored to the movement base
-						if (TargetMoverComp->SmoothingMode != EMoverSmoothingMode::None)
-						{
-							const FTransform OldToNewTransform = NewSyncTransformWs.GetRelativeTransform(OldSyncTransformWs);
-
-							// Modify the PRESENTATION sync state that we're smoothing TO
-							FMoverSyncState PresentationSyncState;
-							if (TargetMoverComp->BackendLiaisonComp->ReadPresentationSyncState(OUT PresentationSyncState))
-							{
-								if (FMoverDefaultSyncState* PresentationMoverState = PresentationSyncState.SyncStateCollection.FindMutableDataByType<FMoverDefaultSyncState>())
-								{
-									OldSyncTransformWs = PresentationMoverState->GetTransform_WorldSpace();
-									NewSyncTransformWs = OldToNewTransform * OldSyncTransformWs;
-
-									PresentationMoverState->SetTransforms_WorldSpace(
-										NewSyncTransformWs.GetLocation(),
-										NewSyncTransformWs.GetRotation().Rotator(),
-										PresentationMoverState->GetVelocity_WorldSpace(),	// keep same velocity and base
-										PresentationMoverState->GetAngularVelocityDegrees_WorldSpace(),
-										PresentationMoverState->GetMovementBase(), PresentationMoverState->GetMovementBaseBoneName());
-
-									TargetMoverComp->BackendLiaisonComp->WritePresentationSyncState(PresentationSyncState);
-								}
-							}
-
-							// Modify the PREV PRESENTATION sync state that we're smoothing FROM
-							FMoverSyncState PrevPresentationSyncState;
-							if (TargetMoverComp->BackendLiaisonComp->ReadPrevPresentationSyncState(OUT PrevPresentationSyncState))
-							{
-								if (FMoverDefaultSyncState* PrevPresentationMoverState = PrevPresentationSyncState.SyncStateCollection.FindMutableDataByType<FMoverDefaultSyncState>())
-								{
-									OldSyncTransformWs = PrevPresentationMoverState->GetTransform_WorldSpace();
-									NewSyncTransformWs = OldToNewTransform * OldSyncTransformWs;
-
-									PrevPresentationMoverState->SetTransforms_WorldSpace(
-										NewSyncTransformWs.GetLocation(),
-										NewSyncTransformWs.GetRotation().Rotator(),
-										PrevPresentationMoverState->GetVelocity_WorldSpace(),	// keep same velocity and base
-										PrevPresentationMoverState->GetAngularVelocityDegrees_WorldSpace(),
-										PrevPresentationMoverState->GetMovementBase(), PrevPresentationMoverState->GetMovementBaseBoneName());
-
-									TargetMoverComp->BackendLiaisonComp->WritePrevPresentationSyncState(PrevPresentationSyncState);	
-								}
-							}
-						}
+							PrevPresentationMoverState->GetVelocity_WorldSpace(),	// keep same velocity and base
+							PrevPresentationMoverState->GetAngularVelocityDegrees_WorldSpace(),
+							PrevPresentationMoverState->GetMovementBase(), PrevPresentationMoverState->GetMovementBaseBoneName());
+
+						TargetMoverComp->BackendLiaisonComp->WritePrevPresentationSyncState(PrevPresentationSyncState);
 					}
 				}
 			}
-
-			SimBlackboard->Set(CommonBlackboard::LastAppliedDynamicMovementBase, CurrentBaseInfo);
-			bDidGetUpToDate = true;
 		}
 	}
 
-	if (!bDidGetUpToDate)
-	{
-		SimBlackboard->Invalidate(CommonBlackboard::LastAppliedDynamicMovementBase);
-	}
+	SimBlackboard->Set(CommonBlackboard::LastFoundDynamicMovementBase, CurrentBaseInfo);
+	bDidPerformBasedMovement = true;
 }
 
 
diff --git a/Engine/Plugins/Experimental/Mover/Source/Mover/Private/MoverComponent.cpp b/Engine/Plugins/Experimental/Mover/Source/Mover/Private/MoverComponent.cpp
--- a/Engine/Plugins/Experimental/Mover/Source/Mover/Private/MoverComponent.cpp
+++ b/Engine/Plugins/Experimental/Mover/Source/Mover/Private/MoverComponent.cpp
@@ -1200,7 +1200,6 @@ void UMoverComponent::UpdateBasedMovementScheduling(const FMoverTickEndData& Sim
 		MovementBaseDependency = nullptr;
 
 		SimBlackboard->Invalidate(CommonBlackboard::LastFoundDynamicMovementBase);
-		SimBlackboard->Invalidate(CommonBlackboard::LastAppliedDynamicMovementBase);
 	}
 }
 
diff --git a/Engine/Plugins/Experimental/Mover/Source/Mover/Private/MoverSimulationTypes.cpp b/Engine/Plugins/Experimental/Mover/Source/Mover/Private/MoverSimulationTypes.cpp
--- a/Engine/Plugins/Experimental/Mover/Source/Mover/Private/MoverSimulationTypes.cpp
+++ b/Engine/Plugins/Experimental/Mover/Source/Mover/Private/MoverSimulationTypes.cpp
@@ -22,7 +22,6 @@ namespace CommonBlackboard
 	const FName LastFloorResult = TEXT("LastFloor");
 	const FName LastWaterResult = TEXT("LastWater");
 	const FName LastFoundDynamicMovementBase = TEXT("LastFoundDynamicMovementBase");
-	const FName LastAppliedDynamicMovementBase = TEXT("LastAppliedDynamicMovementBase");
 	const FName TimeSinceSupported = TEXT("TimeSinceSupported");
 
 	const FName LastModeChangeRecord = TEXT("LastModeChangeRecord");
diff --git a/Engine/Plugins/Experimental/Mover/Source/Mover/Public/DefaultMovementSet/Modes/AsyncFallingMode.h b/Engine/Plugins/Experimental/Mover/Source/Mover/Public/DefaultMovementSet/Modes/AsyncFallingMode.h
--- a/Engine/Plugins/Experimental/Mover/Source/Mover/Public/DefaultMovementSet/Modes/AsyncFallingMode.h
+++ b/Engine/Plugins/Experimental/Mover/Source/Mover/Public/DefaultMovementSet/Modes/AsyncFallingMode.h
@@ -89,7 +89,7 @@ protected:
 	 * (i.e. landing on a walkable surface would switch to the walking movement mode) 
 	 */
 	UFUNCTION(BlueprintCallable, Category=Mover)
-	UE_API virtual void ProcessLanded(const FFloorCheckResult& FloorResult, FVector& Velocity, FRelativeBaseInfo& BaseInfo, FMoverTickEndData& TickEndData) const;
+	UE_API virtual void ProcessLanded(const FFloorCheckResult& FloorResult, FVector Location, FVector& Velocity, FRelativeBaseInfo& BaseInfo, FMoverTickEndData& TickEndData) const;
 
 	UE_API void CaptureFinalState(const FMoverDefaultSyncState* StartSyncState, const FVector FinalLocation, const FRotator FinalRotation, const FFloorCheckResult& FloorResult, float DeltaSeconds, float DeltaSecondsUsed, const FVector& AngularVelocityDegrees, FMoverDefaultSyncState& OutputSyncState, FMoverTickEndData& TickEndData, FMovementRecord& Record) const;
 
diff --git a/Engine/Plugins/Experimental/Mover/Source/Mover/Public/DefaultMovementSet/Modes/AsyncWalkingMode.h b/Engine/Plugins/Experimental/Mover/Source/Mover/Public/DefaultMovementSet/Modes/AsyncWalkingMode.h
--- a/Engine/Plugins/Experimental/Mover/Source/Mover/Public/DefaultMovementSet/Modes/AsyncWalkingMode.h
+++ b/Engine/Plugins/Experimental/Mover/Source/Mover/Public/DefaultMovementSet/Modes/AsyncWalkingMode.h
@@ -54,7 +54,7 @@ protected:
 
 	UE_API void CaptureFinalState(const FVector FinalLocation, const FRotator FinalRotation, bool bDidAttemptMovement, const FFloorCheckResult& FloorResult, const FMovementRecord& Record, const FVector& AngularVelocityDegrees, FMoverDefaultSyncState& OutputSyncState) const;
 
-	UE_API FRelativeBaseInfo UpdateFloorAndBaseInfo(const FFloorCheckResult& FloorResult) const;
+	UE_API FRelativeBaseInfo UpdateFloorAndBaseInfo(const FFloorCheckResult& FloorResult, FVector CurrentLocation) const;
 
 	TWeakObjectPtr<const UCommonLegacyMovementSettings> CommonLegacySettings;
 };
diff --git a/Engine/Plugins/Experimental/Mover/Source/Mover/Public/DefaultMovementSet/Modes/FallingMode.h b/Engine/Plugins/Experimental/Mover/Source/Mover/Public/DefaultMovementSet/Modes/FallingMode.h
--- a/Engine/Plugins/Experimental/Mover/Source/Mover/Public/DefaultMovementSet/Modes/FallingMode.h
+++ b/Engine/Plugins/Experimental/Mover/Source/Mover/Public/DefaultMovementSet/Modes/FallingMode.h
@@ -97,7 +97,7 @@ protected:
 	 * (i.e. landing on a walkable surface would switch to the walking movement mode) 
 	 */
 	UFUNCTION(BlueprintCallable, Category=Mover)
-	UE_API virtual void ProcessLanded(const FFloorCheckResult& FloorResult, FVector& Velocity, FRelativeBaseInfo& BaseInfo, FMoverTickEndData& TickEndData) const;
+	UE_API virtual void ProcessLanded(const FFloorCheckResult& FloorResult, FVector Location, FVector& Velocity, FRelativeBaseInfo& BaseInfo, FMoverTickEndData& TickEndData) const;
 	
 	UE_API virtual void OnRegistered(const FName ModeName, const FMoverSimContext& SimContext) override;
 	UE_API virtual void OnUnregistered(const FMoverSimContext& SimContext) override;
diff --git a/Engine/Plugins/Experimental/Mover/Source/Mover/Public/DefaultMovementSet/Modes/WalkingMode.h b/Engine/Plugins/Experimental/Mover/Source/Mover/Public/DefaultMovementSet/Modes/WalkingMode.h
--- a/Engine/Plugins/Experimental/Mover/Source/Mover/Public/DefaultMovementSet/Modes/WalkingMode.h
+++ b/Engine/Plugins/Experimental/Mover/Source/Mover/Public/DefaultMovementSet/Modes/WalkingMode.h
@@ -62,9 +62,9 @@ protected:
 	UE_API virtual void OnRegistered(const FName ModeName, const FMoverSimContext& SimContext) override;
 	UE_API virtual void OnUnregistered(const FMoverSimContext& SimContext) override;
 
-	UE_API void CaptureFinalState(USceneComponent* UpdatedComponent, bool bDidAttemptMovement, const FFloorCheckResult& FloorResult, const FMovementRecord& Record, const FVector& AngularVelocityDegrees, FMoverDefaultSyncState& OutputSyncState) const;
+	UE_API void CaptureFinalState(USceneComponent* UpdatedComponent, bool bDidAttemptMovement, const FFloorCheckResult& FloorResult, const FMovementRecord& Record, const FMoverTimeStep& TimeStep, const FVector& AngularVelocityDegrees, FMoverDefaultSyncState& OutputSyncState) const;
 
-	UE_API FRelativeBaseInfo UpdateFloorAndBaseInfo(const FFloorCheckResult& FloorResult) const;
+	UE_API FRelativeBaseInfo UpdateFloorAndBaseInfo(const FFloorCheckResult& FloorResult, FVector CurrentLocation) const;
 
 	TObjectPtr<const UCommonLegacyMovementSettings> CommonLegacySettings;
 };
diff --git a/Engine/Plugins/Experimental/Mover/Source/Mover/Public/MoveLibrary/BasedMovementUtils.h b/Engine/Plugins/Experimental/Mover/Source/Mover/Public/MoveLibrary/BasedMovementUtils.h
--- a/Engine/Plugins/Experimental/Mover/Source/Mover/Public/MoveLibrary/BasedMovementUtils.h
+++ b/Engine/Plugins/Experimental/Mover/Source/Mover/Public/MoveLibrary/BasedMovementUtils.h
@@ -38,6 +38,11 @@ struct FRelativeBaseInfo
 	UPROPERTY(VisibleInstanceOnly, BlueprintReadWrite, Category = "Mover")
 	FVector ContactLocalPosition = FVector::ZeroVector;
 
+	/** Last captured offset of the Mover actor from contact location, in world space. */
+	UPROPERTY(VisibleInstanceOnly, BlueprintReadWrite, Category = "Mover")
+	FVector WorldspaceOffsetFromContactPos = FVector::ZeroVector;
+
+
 public:
 
 	UE_API void Clear();
@@ -46,8 +51,9 @@ public:
 	UE_API bool UsesSameBase(const FRelativeBaseInfo& Other) const;
 	UE_API bool UsesSameBase(const UPrimitiveComponent* OtherComp, FName OtherBoneName=NAME_None) const;
 
-	UE_API void SetFromFloorResult(const FFloorCheckResult& FloorTestResult);
-	UE_API void SetFromComponent(UPrimitiveComponent* InRelativeComp, FName InBoneName=NAME_None);
+	UE_API void SetFromFloorResult(const FFloorCheckResult& FloorTestResult, const FVector ObjectWorldLocation, const FVector UpDir);
+
+	UE_API void UpdateToLatestBaseTransform();
 
 	UE_API FString ToString() const;
 };
@@ -121,6 +127,8 @@ public:
 	UFUNCTION(BlueprintCallable, Category = "Mover/MovementBases")
 	static UE_API void TransformRotatorToLocal(FQuat BaseQuat, FRotator WorldSpaceRotator, FRotator& OutLocalRotator);
 
+	/** Apply a sweeping movement to attempt to follow a movement base that we're moving relative to. This is typically called during movement simulation. */
+	static UE_API FRelativeBaseInfo ApplyMovementBaseChange(UMoverComponent* TargetMoverComp, const FRelativeBaseInfo& RelativeBaseInfo, bool bIgnoreBaseRotation);
 
 	/** Makes it so BasedObjectTick ticks after NewBase's actor ticking */
 	static UE_API void AddTickDependency(FTickFunction& BasedObjectTick, UPrimitiveComponent* NewBase);
diff --git a/Engine/Plugins/Experimental/Mover/Source/Mover/Public/MoverSimulationTypes.h b/Engine/Plugins/Experimental/Mover/Source/Mover/Public/MoverSimulationTypes.h
--- a/Engine/Plugins/Experimental/Mover/Source/Mover/Public/MoverSimulationTypes.h
+++ b/Engine/Plugins/Experimental/Mover/Source/Mover/Public/MoverSimulationTypes.h
@@ -33,7 +33,6 @@ namespace CommonBlackboard
 	extern MOVER_API const FName LastFloorResult;
 	extern MOVER_API const FName LastWaterResult;
 	extern MOVER_API const FName LastFoundDynamicMovementBase;
-	extern MOVER_API const FName LastAppliedDynamicMovementBase;
 	extern MOVER_API const FName TimeSinceSupported;
 
 	extern MOVER_API const FName LastModeChangeRecord;
