Mover: Adding util functions to LayeredMove to make it easier to write code applying to all moves in a move group

#Synced-CL 52693474

diff --git a/Engine/Plugins/Experimental/Mover/Source/Mover/Public/LayeredMove.h b/Engine/Plugins/Experimental/Mover/Source/Mover/Public/LayeredMove.h
--- a/Engine/Plugins/Experimental/Mover/Source/Mover/Public/LayeredMove.h
+++ b/Engine/Plugins/Experimental/Mover/Source/Mover/Public/LayeredMove.h
@@ -232,6 +232,32 @@ struct FLayeredMoveGroup
 	}
 	UE_API const FLayeredMoveBase* FindActiveMove(const UScriptStruct* LayeredMoveStructType) const;
 
+	/** Calls Func for every active move whose struct type is or derives from MoveType. */
+	template <typename MoveType UE_REQUIRES(std::is_base_of_v<FLayeredMoveBase, MoveType>)>
+	void ForEachActiveMoveOfType(TFunctionRef<void(MoveType&)> Func)
+	{
+		for (const TSharedPtr<FLayeredMoveBase>& Move : ActiveLayeredMoves)
+		{
+			if (Move.IsValid() && Move->GetScriptStruct()->IsChildOf(MoveType::StaticStruct()))
+			{
+				Func(*static_cast<MoveType*>(Move.Get()));
+			}
+		}
+	}
+
+	/** Const overload: calls Func for every active move whose struct type is or derives from MoveType. */
+	template <typename MoveType UE_REQUIRES(std::is_base_of_v<FLayeredMoveBase, MoveType>)>
+	void ForEachActiveMoveOfType(TFunctionRef<void(const MoveType&)> Func) const
+	{
+		for (const TSharedPtr<FLayeredMoveBase>& Move : ActiveLayeredMoves)
+		{
+			if (Move.IsValid() && Move->GetScriptStruct()->IsChildOf(MoveType::StaticStruct()))
+			{
+				Func(*static_cast<const MoveType*>(Move.Get()));
+			}
+		}
+	}
+
 	/** Returns the first queued layered move of the specified type, if one exists */
 	template <typename MoveType UE_REQUIRES(std::is_base_of_v<FLayeredMoveBase, MoveType>)>
 	const MoveType* FindQueuedMove() const
