Mover: fixes crashes that can occur when compiling a MoverComponent BP while PIE is running with an instance of it.  We will now perform registration both in InitializeComponent and OnRegistered, since things occur in different orders when recompiling during play.
Also removes no-longer-needed avoidance, and downgrades an unnecessary blackboard warning.

#Synced-CL 52746816

diff --git a/Engine/Plugins/Experimental/Mover/Source/Mover/Private/DefaultMovementSet/Modes/FallingMode.cpp b/Engine/Plugins/Experimental/Mover/Source/Mover/Private/DefaultMovementSet/Modes/FallingMode.cpp
--- a/Engine/Plugins/Experimental/Mover/Source/Mover/Private/DefaultMovementSet/Modes/FallingMode.cpp
+++ b/Engine/Plugins/Experimental/Mover/Source/Mover/Private/DefaultMovementSet/Modes/FallingMode.cpp
@@ -42,11 +42,6 @@ void UFallingMode::GenerateMove_Implementation(const FMoverSimContext& SimContex
 
 	const float DeltaSeconds = TimeStep.StepMs * 0.001f;
 
-	if (!ensure(MoverComp))
-	{
-		return;
-	}
-
 	FVector UpDirection = MoverComp->GetUpDirection();
 	
 	// We don't want velocity limits to take the falling velocity component into account, since it is handled 
diff --git a/Engine/Plugins/Experimental/Mover/Source/Mover/Private/MovementModeStateMachine.cpp b/Engine/Plugins/Experimental/Mover/Source/Mover/Private/MovementModeStateMachine.cpp
--- a/Engine/Plugins/Experimental/Mover/Source/Mover/Private/MovementModeStateMachine.cpp
+++ b/Engine/Plugins/Experimental/Mover/Source/Mover/Private/MovementModeStateMachine.cpp
@@ -61,9 +61,9 @@ void UMovementModeStateMachine::RegisterMovementMode(FName ModeName, TSubclassOf
 
 void UMovementModeStateMachine::UnregisterMovementMode(FName ModeName)
 {
-	TObjectPtr<UBaseMovementMode> ModeToUnregister = Modes.FindAndRemoveChecked(ModeName);
+	TObjectPtr<UBaseMovementMode> ModeToUnregister;
 
-	if (ModeToUnregister)
+	if (Modes.RemoveAndCopyValue(ModeName, ModeToUnregister) && ModeToUnregister)
 	{
 		ModeToUnregister->OnUnregistered(GetSimContext());
 	}
diff --git a/Engine/Plugins/Experimental/Mover/Source/Mover/Private/MoverComponent.cpp b/Engine/Plugins/Experimental/Mover/Source/Mover/Private/MoverComponent.cpp
--- a/Engine/Plugins/Experimental/Mover/Source/Mover/Private/MoverComponent.cpp
+++ b/Engine/Plugins/Experimental/Mover/Source/Mover/Private/MoverComponent.cpp
@@ -161,40 +161,21 @@ void UMoverComponent::InitializeComponent()
 			ModeFSM->SetSimulation(*Simulation);
 		}
 
+		FindAndRegisterDefaults();
 
-		bool bHasMatchingStartingState = false;
-
-		for (const TPair<FName, TObjectPtr<UBaseMovementMode>>& Element : MovementModes)
+		if (StartingMovementMode != NAME_None)
 		{
-			if (Element.Value.Get() == nullptr)
+			if (ModeFSM->FindMovementMode(StartingMovementMode))
 			{
-				UE_LOGF(LogMover, Warning, "Invalid Movement Mode type '%ls' detected on %ls. Mover actor will not function correctly.",
-					*Element.Key.ToString(), *GetNameSafe(GetOwner()));
-				continue;
+				ModeFSM->SetDefaultMode(StartingMovementMode);
+				ModeFSM->QueueNextMode(StartingMovementMode);
+			}
+			else
+			{
+				UE_LOGF(LogMover, Warning, "Invalid StartingMovementMode '%ls' specified on %ls. Mover actor will not function until something queues a valid mode.",
+					*StartingMovementMode.ToString(), *GetNameSafe(GetOwner()));
 			}
-
-			const bool bIsStartingMode = (StartingMovementMode == Element.Key);
-
-			ModeFSM->RegisterMovementMode(Element.Key, Element.Value, bIsStartingMode);
-
-			bHasMatchingStartingState |= bIsStartingMode;
-		}
-
-		for (TObjectPtr<UBaseMovementModeTransition>& Transition : Transitions)
-		{
-			ModeFSM->RegisterGlobalTransition(Transition);
-		}
-
-		UE_CLOGF(!bHasMatchingStartingState, LogMover, Warning, "Invalid StartingMovementMode '%ls' specified on %ls. Mover actor will not function.",
-			*StartingMovementMode.ToString(), *GetNameSafe(GetOwner()));
-
-		if (bHasMatchingStartingState && StartingMovementMode != NAME_None)
-		{
-			ModeFSM->SetDefaultMode(StartingMovementMode);
-			ModeFSM->QueueNextMode(StartingMovementMode);
 		}
-
-
 	}
 
 	// Gather initial state to fulfill queries
@@ -236,7 +217,7 @@ void UMoverComponent::OnRegister()
 
 	Super::OnRegister();
 
-	FindDefaultUpdatedComponent();
+	FindAndRegisterDefaults();
 }
 
 
@@ -1223,6 +1204,38 @@ void UMoverComponent::FindDefaultUpdatedComponent()
 }
 
 
+void UMoverComponent::FindAndRegisterDefaults()
+{
+	FindDefaultUpdatedComponent();
+
+	if (ModeFSM)
+	{
+		// Movement Modes
+		ModeFSM->ClearAllMovementModes();
+		
+		for (const TPair<FName, TObjectPtr<UBaseMovementMode>>& Element : MovementModes)
+		{
+			if (Element.Value.Get() == nullptr)
+			{
+				UE_LOGF(LogMover, Warning, "Invalid Movement Mode type '%ls' detected on %ls. Mover actor will not function correctly.",
+					*Element.Key.ToString(), *GetNameSafe(GetOwner()));
+				continue;
+			}
+
+			const bool bIsStartingMode = (StartingMovementMode == Element.Key);
+			ModeFSM->RegisterMovementMode(Element.Key, Element.Value, bIsStartingMode);
+		}
+
+		// Global Transitions
+		ModeFSM->ClearAllGlobalTransitions();
+
+		for (TObjectPtr<UBaseMovementModeTransition>& Transition : Transitions)
+		{
+			ModeFSM->RegisterGlobalTransition(Transition);
+		}
+	}
+}
+
 void UMoverComponent::UpdateTickRegistration()
 {
 	const bool bHasUpdatedComponent = (UpdatedComponent != NULL);
diff --git a/Engine/Plugins/Experimental/Mover/Source/Mover/Public/MoveLibrary/RollbackBlackboard.h b/Engine/Plugins/Experimental/Mover/Source/Mover/Public/MoveLibrary/RollbackBlackboard.h
--- a/Engine/Plugins/Experimental/Mover/Source/Mover/Public/MoveLibrary/RollbackBlackboard.h
+++ b/Engine/Plugins/Experimental/Mover/Source/Mover/Public/MoveLibrary/RollbackBlackboard.h
@@ -348,7 +348,7 @@ public:
 	{
 		if (EntryMap.Contains(EntryName))
 		{
-			UE_LOGF(LogMover, Warning, "Skipping attempt to create a new blackboard entry named %ls since it already exists. Policy settings from the first entry will be retained.", *EntryName.ToString());
+			UE_LOGF(LogMover, Verbose, "Skipping attempt to create a new blackboard entry named %ls since it already exists. Policy settings from the first entry will be retained.", *EntryName.ToString());
 			return false;
 		}
 
diff --git a/Engine/Plugins/Experimental/Mover/Source/Mover/Public/MoverComponent.h b/Engine/Plugins/Experimental/Mover/Source/Mover/Public/MoverComponent.h
--- a/Engine/Plugins/Experimental/Mover/Source/Mover/Public/MoverComponent.h
+++ b/Engine/Plugins/Experimental/Mover/Source/Mover/Public/MoverComponent.h
@@ -756,6 +756,7 @@ public:
 
 protected:
 	MOVER_API void FindDefaultUpdatedComponent();
+	MOVER_API void FindAndRegisterDefaults();
 	MOVER_API void UpdateTickRegistration();
 
 	MOVER_API virtual void DoQueueNextMode(FName DesiredModeName, bool bShouldReenter=false);
